[RenderCore]
* Moved dllexport from types to methods/staticvars

diff --git a/Engine/Plugins/FX/Niagara/Source/Niagara/Classes/NiagaraDataInterface.h b/Engine/Plugins/FX/Niagara/Source/Niagara/Classes/NiagaraDataInterface.h
--- a/Engine/Plugins/FX/Niagara/Source/Niagara/Classes/NiagaraDataInterface.h
+++ b/Engine/Plugins/FX/Niagara/Source/Niagara/Classes/NiagaraDataInterface.h
@@ -8,7 +8,7 @@
 #include "VectorVM.h"
 #include "NiagaraDataInterface.generated.h"
 
-class RENDERCORE_API FRDGBuilder;
+class FRDGBuilder;
 class FRDGExternalAccessQueue;
 
 class INiagaraCompiler;
diff --git a/Engine/Plugins/FX/Niagara/Source/NiagaraCore/Public/NiagaraDataInterfaceBase.h b/Engine/Plugins/FX/Niagara/Source/NiagaraCore/Public/NiagaraDataInterfaceBase.h
--- a/Engine/Plugins/FX/Niagara/Source/NiagaraCore/Public/NiagaraDataInterfaceBase.h
+++ b/Engine/Plugins/FX/Niagara/Source/NiagaraCore/Public/NiagaraDataInterfaceBase.h
@@ -11,7 +11,7 @@ class INiagaraCompiler;
 class UCurveVector;
 class UCurveLinearColor;
 class UCurveFloat;
-class RENDERCORE_API FShaderParametersMetadataBuilder;
+class FShaderParametersMetadataBuilder;
 
 class UNiagaraEmitter;
 class FNiagaraSystemInstance;
diff --git a/Engine/Source/Runtime/RenderCore/Private/RenderGraphResourcePool.h b/Engine/Source/Runtime/RenderCore/Private/RenderGraphResourcePool.h
--- a/Engine/Source/Runtime/RenderCore/Private/RenderGraphResourcePool.h
+++ b/Engine/Source/Runtime/RenderCore/Private/RenderGraphResourcePool.h
@@ -10,20 +10,20 @@
 #include "RendererInterface.h"
 #include "RenderGraphResources.h"
 
-class RENDERCORE_API FRDGBufferPool : public FRenderResource
+class FRDGBufferPool : public FRenderResource
 {
 public:
 	FRDGBufferPool() = default;
 
 	/** Call once per frame to trim elements from the pool. */
-	void TickPoolElements();
+	RENDERCORE_API void TickPoolElements();
 
-	TRefCountPtr<FRDGPooledBuffer> FindFreeBuffer(const FRDGBufferDesc& Desc, const TCHAR* InDebugName, ERDGPooledBufferAlignment Alignment = ERDGPooledBufferAlignment::Page);
+	RENDERCORE_API TRefCountPtr<FRDGPooledBuffer> FindFreeBuffer(const FRDGBufferDesc& Desc, const TCHAR* InDebugName, ERDGPooledBufferAlignment Alignment = ERDGPooledBufferAlignment::Page);
 
-	void DumpMemoryUsage(FOutputDevice& OutputDevice);
+	RENDERCORE_API void DumpMemoryUsage(FOutputDevice& OutputDevice);
 
 private:
-	void ReleaseRHI() override;
+	RENDERCORE_API void ReleaseRHI() override;
 
 	/** Elements can be 0, we compact the buffer later. */
 	TArray<TRefCountPtr<FRDGPooledBuffer>> AllocatedBuffers;
diff --git a/Engine/Source/Runtime/RenderCore/Public/BuiltInRayTracingShaders.h b/Engine/Source/Runtime/RenderCore/Public/BuiltInRayTracingShaders.h
--- a/Engine/Source/Runtime/RenderCore/Public/BuiltInRayTracingShaders.h
+++ b/Engine/Source/Runtime/RenderCore/Public/BuiltInRayTracingShaders.h
@@ -54,7 +54,7 @@ public:
 	{}
 };
 
-class RENDERCORE_API FRayTracingDispatchDescCS : public FBuiltInRayTracingShader
+class FRayTracingDispatchDescCS : public FBuiltInRayTracingShader
 {
 	DECLARE_GLOBAL_SHADER(FRayTracingDispatchDescCS);
 
@@ -77,7 +77,7 @@ public:
 		FGlobalShader::ModifyCompilationEnvironment(Parameters, OutEnvironment);
 	}
 
-	static void Dispatch(FRHICommandList& RHICmdList,
+	static RENDERCORE_API void Dispatch(FRHICommandList& RHICmdList,
 		const void* DispatchDescInput, uint32 DispatchDescSize, uint32 DispatchDescDimensionsOffset,
 		FRHIShaderResourceView* DispatchDimensionsSRV, uint32 DimensionsBufferOffset,
 		FRHIUnorderedAccessView* DispatchDescOutputUAV);
diff --git a/Engine/Source/Runtime/RenderCore/Public/CommonRenderResources.h b/Engine/Source/Runtime/RenderCore/Public/CommonRenderResources.h
--- a/Engine/Source/Runtime/RenderCore/Public/CommonRenderResources.h
+++ b/Engine/Source/Runtime/RenderCore/Public/CommonRenderResources.h
@@ -117,7 +117,7 @@ class FScreenVertexShaderVS : public FGlobalShader
 
 /** Vertex shader to draw an instanced quad covering all the viewports (SV_ViewportArrayIndex is output for each SV_InstanceID). Does not have any shader parameters.
  * The pixel shader should just use SV_Position. */
-class RENDERCORE_API FInstancedScreenVertexShaderVS : public FScreenVertexShaderVS
+class FInstancedScreenVertexShaderVS : public FScreenVertexShaderVS
 {
 	DECLARE_GLOBAL_SHADER(FInstancedScreenVertexShaderVS);
 	SHADER_USE_PARAMETER_STRUCT(FInstancedScreenVertexShaderVS, FScreenVertexShaderVS);
diff --git a/Engine/Source/Runtime/RenderCore/Public/CustomResourcePool.h b/Engine/Source/Runtime/RenderCore/Public/CustomResourcePool.h
--- a/Engine/Source/Runtime/RenderCore/Public/CustomResourcePool.h
+++ b/Engine/Source/Runtime/RenderCore/Public/CustomResourcePool.h
@@ -5,12 +5,12 @@
 
 class FRHICommandListImmediate;
 
-class RENDERCORE_API ICustomResourcePool
+class ICustomResourcePool
 {
 public:
 	virtual ~ICustomResourcePool() {}
 	virtual void Tick(FRHICommandListImmediate& RHICmdList) = 0;
 
-	static void TickPoolElements(FRHICommandListImmediate& RHICmdList);
+	static RENDERCORE_API void TickPoolElements(FRHICommandListImmediate& RHICmdList);
 };
 extern RENDERCORE_API ICustomResourcePool* GCustomResourcePool;
diff --git a/Engine/Source/Runtime/RenderCore/Public/DumpGPU.h b/Engine/Source/Runtime/RenderCore/Public/DumpGPU.h
--- a/Engine/Source/Runtime/RenderCore/Public/DumpGPU.h
+++ b/Engine/Source/Runtime/RenderCore/Public/DumpGPU.h
@@ -6,11 +6,11 @@
 #include "UObject/NameTypes.h"
 
 
-class RENDERCORE_API IDumpGPUUploadServiceProvider
+class IDumpGPUUploadServiceProvider
 {
 public:
 
-	struct RENDERCORE_API FDumpParameters
+	struct FDumpParameters
 	{
 		static constexpr const TCHAR* kServiceFileName = TEXT("Base/DumpService.json");
 
@@ -21,12 +21,12 @@ public:
 		FName CompressionName;
 		FWildcardString CompressionFiles;
 
-		FString DumpServiceParametersFileContent() const;
-		bool DumpServiceParametersFile() const;
+		RENDERCORE_API FString DumpServiceParametersFileContent() const;
+		RENDERCORE_API bool DumpServiceParametersFile() const;
 	};
 
 	virtual void UploadDump(const FDumpParameters& Parameters) = 0;
 	virtual ~IDumpGPUUploadServiceProvider() = default;
 
-	static IDumpGPUUploadServiceProvider* GProvider;
+	static RENDERCORE_API IDumpGPUUploadServiceProvider* GProvider;
 };
diff --git a/Engine/Source/Runtime/RenderCore/Public/DynamicBufferAllocator.h b/Engine/Source/Runtime/RenderCore/Public/DynamicBufferAllocator.h
--- a/Engine/Source/Runtime/RenderCore/Public/DynamicBufferAllocator.h
+++ b/Engine/Source/Runtime/RenderCore/Public/DynamicBufferAllocator.h
@@ -45,7 +45,7 @@ struct FDynamicAllocReadBuffer : public FDynamicReadBuffer
   the renderer RHI will have already been destroyed and we can execute code on invalid data. By making ourself a render resource, we
   clean up immediately before the renderer dies.
 */
-class RENDERCORE_API FGlobalDynamicReadBuffer : public FRenderResource
+class FGlobalDynamicReadBuffer : public FRenderResource
 {
 public:
 	/**
@@ -74,30 +74,30 @@ public:
 		}
 	};
 
-	FGlobalDynamicReadBuffer();
-	~FGlobalDynamicReadBuffer();
+	RENDERCORE_API FGlobalDynamicReadBuffer();
+	RENDERCORE_API ~FGlobalDynamicReadBuffer();
 	
-	FAllocation AllocateHalf(uint32 Num);
-	FAllocation AllocateFloat(uint32 Num);
-	FAllocation AllocateInt32(uint32 Num);
-	FAllocation AllocateUInt32(uint32 Num);
+	RENDERCORE_API FAllocation AllocateHalf(uint32 Num);
+	RENDERCORE_API FAllocation AllocateFloat(uint32 Num);
+	RENDERCORE_API FAllocation AllocateInt32(uint32 Num);
+	RENDERCORE_API FAllocation AllocateUInt32(uint32 Num);
 
 	/**
 	* Commits allocated memory to the GPU.
 	*		WARNING: Once this buffer has been committed to the GPU, allocations
 	*		remain valid only until the next call to Allocate!
 	*/
-	void Commit();
+	RENDERCORE_API void Commit();
 
 
 	/** Returns true if log statements should be made because we exceeded GMaxVertexBytesAllocatedPerFrame */
-	bool IsRenderAlarmLoggingEnabled() const;
+	RENDERCORE_API bool IsRenderAlarmLoggingEnabled() const;
 
 protected:
-	virtual void InitRHI() override;
-	virtual void ReleaseRHI() override;
-	void Cleanup();
-	void IncrementTotalAllocations(uint32 Num);
+	RENDERCORE_API virtual void InitRHI() override;
+	RENDERCORE_API virtual void ReleaseRHI() override;
+	RENDERCORE_API void Cleanup();
+	RENDERCORE_API void IncrementTotalAllocations(uint32 Num);
 
 	/** The pools of read buffers from which allocations are made. */
 	FDynamicReadBufferPool* HalfBufferPool;
diff --git a/Engine/Source/Runtime/RenderCore/Public/DynamicRenderScaling.h b/Engine/Source/Runtime/RenderCore/Public/DynamicRenderScaling.h
--- a/Engine/Source/Runtime/RenderCore/Public/DynamicRenderScaling.h
+++ b/Engine/Source/Runtime/RenderCore/Public/DynamicRenderScaling.h
@@ -30,7 +30,7 @@ enum class EHeuristicModel
 };
 
 /* Represent an independent budget to dynamically scale by its own. */
-struct RENDERCORE_API FHeuristicSettings final
+struct FHeuristicSettings final
 {
 	static constexpr float kDefaultMinResolutionFraction = 0.5f;
 	static constexpr float kDefaultMaxResolutionFraction = 1.0f;
@@ -56,38 +56,38 @@ struct RENDERCORE_API FHeuristicSettings final
 	int32 UpperBoundQuantization     = kDefaultUpperBoundQuantization;
 
 	/** Returns whether the heuristic is enabled or not. */
-	bool IsEnabled() const;
+	RENDERCORE_API bool IsEnabled() const;
 
 	/** Returns the desired GPU cost to be targeted to have head room left to not go over budget. */
-	float GetTargetedMs(float BudgetMs) const;
+	RENDERCORE_API float GetTargetedMs(float BudgetMs) const;
 
 	/** Returns how much the GPU cost scales for a given ResolutionFraction. */
-	float EstimateCostScale(float ResolutionFraction) const;
+	RENDERCORE_API float EstimateCostScale(float ResolutionFraction) const;
 
 	/** Returns how much the ResolutionFraction should scale for a GPU timing to fit to target. */
-	float EstimateResolutionFactor(float TargetMs, float TimingMs) const;
+	RENDERCORE_API float EstimateResolutionFactor(float TargetMs, float TimingMs) const;
 
 	/** Returns how much the GPU time should scale between two different resolution fraction. */
-	float EstimateTimeFactor(float CurrentResolutionFraction, float NewResolutionFraction) const;
+	RENDERCORE_API float EstimateTimeFactor(float CurrentResolutionFraction, float NewResolutionFraction) const;
 
 	/** Corrects new resolution fraction to ensure it's within bounds, honor AmortizationFactor and FractionQuantization. */
-	float CorrectNewResolutionFraction(float CurrentResolutionFraction, float NewResolutionFraction, float ResolutionFractionScale) const;
+	RENDERCORE_API float CorrectNewResolutionFraction(float CurrentResolutionFraction, float NewResolutionFraction, float ResolutionFractionScale) const;
 
 	/** Returns whether the ResolutionFraction is changing meaningfully enough. */
-	bool DoesResolutionChangeEnough(float CurrentResolutionFraction, float NewResolutionFraction, bool bCanChangeResolution) const;
+	RENDERCORE_API bool DoesResolutionChangeEnough(float CurrentResolutionFraction, float NewResolutionFraction, bool bCanChangeResolution) const;
 };
 
 
 /* Represent an independent budget to dynamically scale by its own. */
-class RENDERCORE_API FBudget final
+class FBudget final
 {
 public:
-	FBudget(const TCHAR* Name, FHeuristicSettings (*HeuristicSettingsGetter)(void));
-	~FBudget();
+	RENDERCORE_API FBudget(const TCHAR* Name, FHeuristicSettings (*HeuristicSettingsGetter)(void));
+	RENDERCORE_API ~FBudget();
 
 	/* Returns the global list of all dynamic resolution budgets. */
-	static TLinkedList<FBudget*>*& GetGlobalList();
-	static int32 GetGlobalListSize();
+	static RENDERCORE_API TLinkedList<FBudget*>*& GetGlobalList();
+	static RENDERCORE_API int32 GetGlobalListSize();
 
 	inline const TCHAR* GetName() const
 	{
diff --git a/Engine/Source/Runtime/RenderCore/Public/GPUMessaging.h b/Engine/Source/Runtime/RenderCore/Public/GPUMessaging.h
--- a/Engine/Source/Runtime/RenderCore/Public/GPUMessaging.h
+++ b/Engine/Source/Runtime/RenderCore/Public/GPUMessaging.h
@@ -24,11 +24,11 @@ class FSocket;
 class FSystem;
 
 /** RAII object for controlling a new GPU message scope. */
-class RENDERCORE_API FScope
+class FScope
 {
 public:
-	FScope(FRDGBuilder& GraphBuilder);
-	~FScope();
+	RENDERCORE_API FScope(FRDGBuilder& GraphBuilder);
+	RENDERCORE_API ~FScope();
 
 private:
 	FRDGBuilder& GraphBuilder;
@@ -39,7 +39,7 @@ private:
 using FMessageId = TRDGHandle<FSocket, uint32>;
 
 /** A socket object mapped to a persistently registered handler. Supports only move construction. Will unregister the handler on destruction. */
-class RENDERCORE_API FSocket
+class FSocket
 {
 public:
 	FSocket() = default;
@@ -64,7 +64,7 @@ public:
 
 	const FMessageId GetMessageId() const { return MessageId; }
 
-	void Reset();
+	RENDERCORE_API void Reset();
 
 private:
 	FSocket(FMessageId InMessageId)
@@ -180,4 +180,4 @@ extern RENDERCORE_API FParameters GetShaderParameters(FRDGBuilder& GraphBuilder)
 
 ///////////////////////////////////////////////////////////////////////////////
 
-} // GPUMessage
\ No newline at end of file
+} // GPUMessage
diff --git a/Engine/Source/Runtime/RenderCore/Public/GenerateMips.h b/Engine/Source/Runtime/RenderCore/Public/GenerateMips.h
--- a/Engine/Source/Runtime/RenderCore/Public/GenerateMips.h
+++ b/Engine/Source/Runtime/RenderCore/Public/GenerateMips.h
@@ -29,13 +29,13 @@ enum class EGenerateMipsPass
 	Raster
 };
 
-class RENDERCORE_API FGenerateMips
+class FGenerateMips
 {
 public:
-	static bool WillFormatSupportCompute(EPixelFormat InPixelFormat);
+	static RENDERCORE_API bool WillFormatSupportCompute(EPixelFormat InPixelFormat);
 
 	/** (ES3.1+) Generates mips for the requested RHI texture using the feature-level appropriate means (Compute, Raster, or Fixed-Function). */
-	static void Execute(
+	static RENDERCORE_API void Execute(
 		FRDGBuilder& GraphBuilder,
 		ERHIFeatureLevel::Type FeatureLevel,
 		FRDGTextureRef Texture,
@@ -43,14 +43,14 @@ public:
 		EGenerateMipsPass Pass = EGenerateMipsPass::AutoDetect);
 
 	/** (SM5+) Generates mips for the requested RDG texture using the requested compute / raster pass. */
-	static void Execute(
+	static RENDERCORE_API void Execute(
 		FRDGBuilder& GraphBuilder,
 		ERHIFeatureLevel::Type FeatureLevel,
 		FRDGTextureRef Texture,
 		FRHISamplerState* Sampler,
 		EGenerateMipsPass Pass = EGenerateMipsPass::AutoDetect);
 
-	static void ExecuteCompute(
+	static RENDERCORE_API void ExecuteCompute(
 		FRDGBuilder& GraphBuilder, ERHIFeatureLevel::Type FeatureLevel,
 		FRDGTextureRef Texture,
 		FRHISamplerState* Sampler);
@@ -59,14 +59,14 @@ public:
 		if( uint(ConditionBuffer[Offset]) > 0)
 			Execute(...)
 	*/
-	static void ExecuteCompute(
+	static RENDERCORE_API void ExecuteCompute(
 		FRDGBuilder& GraphBuilder,
 		ERHIFeatureLevel::Type FeatureLevel,
 		FRDGTextureRef Texture,
 		FRHISamplerState* Sampler,
 		FRDGBufferRef ConditionBuffer, uint32 Offset = 0);
 
-	static void ExecuteRaster(
+	static RENDERCORE_API void ExecuteRaster(
 		FRDGBuilder& GraphBuilder,
 		ERHIFeatureLevel::Type FeatureLevel,
 		FRDGTextureRef Texture,
@@ -74,18 +74,18 @@ public:
 
 	//////////////////////////////////////////////////////////////////////////
 	UE_DEPRECATED(5.1, "This function now requires a ERHIFeatureLevel argument. You can obtain the correct Feature Level from a Scene or View.")
-	static void Execute(FRDGBuilder& GraphBuilder, FRDGTextureRef Texture, FGenerateMipsParams Params = {}, EGenerateMipsPass Pass = EGenerateMipsPass::AutoDetect);
+	static RENDERCORE_API void Execute(FRDGBuilder& GraphBuilder, FRDGTextureRef Texture, FGenerateMipsParams Params = {}, EGenerateMipsPass Pass = EGenerateMipsPass::AutoDetect);
 
 	UE_DEPRECATED(5.1, "This function now requires a ERHIFeatureLevel argument. You can obtain the correct Feature Level from a Scene or View.")
-	static void Execute(FRDGBuilder& GraphBuilder, FRDGTextureRef Texture, FRHISamplerState* Sampler, EGenerateMipsPass Pass = EGenerateMipsPass::AutoDetect);
+	static RENDERCORE_API void Execute(FRDGBuilder& GraphBuilder, FRDGTextureRef Texture, FRHISamplerState* Sampler, EGenerateMipsPass Pass = EGenerateMipsPass::AutoDetect);
 
 	UE_DEPRECATED(5.1, "This function now requires a ERHIFeatureLevel argument. You can obtain the correct Feature Level from a Scene or View.")
-	static void ExecuteCompute(FRDGBuilder& GraphBuilder, FRDGTextureRef Texture, FRHISamplerState* Sampler);
+	static RENDERCORE_API void ExecuteCompute(FRDGBuilder& GraphBuilder, FRDGTextureRef Texture, FRHISamplerState* Sampler);
 
 	UE_DEPRECATED(5.1, "This function now requires a ERHIFeatureLevel argument. You can obtain the correct Feature Level from a Scene or View.")
-	static void ExecuteCompute(FRDGBuilder& GraphBuilder, FRDGTextureRef Texture, FRHISamplerState* Sampler, FRDGBufferRef ConditionBuffer, uint32 Offset = 0);
+	static RENDERCORE_API void ExecuteCompute(FRDGBuilder& GraphBuilder, FRDGTextureRef Texture, FRHISamplerState* Sampler, FRDGBufferRef ConditionBuffer, uint32 Offset = 0);
 
 	UE_DEPRECATED(5.1, "This function now requires a ERHIFeatureLevel argument. You can obtain the correct Feature Level from a Scene or View.")
-	static void ExecuteRaster(FRDGBuilder& GraphBuilder, FRDGTextureRef Texture, FRHISamplerState* Sampler);
+	static RENDERCORE_API void ExecuteRaster(FRDGBuilder& GraphBuilder, FRDGTextureRef Texture, FRHISamplerState* Sampler);
 	//////////////////////////////////////////////////////////////////////////
-};
\ No newline at end of file
+};
diff --git a/Engine/Source/Runtime/RenderCore/Public/GlobalRenderResources.h b/Engine/Source/Runtime/RenderCore/Public/GlobalRenderResources.h
--- a/Engine/Source/Runtime/RenderCore/Public/GlobalRenderResources.h
+++ b/Engine/Source/Runtime/RenderCore/Public/GlobalRenderResources.h
@@ -74,14 +74,14 @@ extern RENDERCORE_API const uint32 GDiffuseConvolveMipLevel;
 * A vertex buffer with a single color component.  This is used on meshes that don't have a color component
 * to keep from needing a separate vertex factory to handle this case.
 */
-class RENDERCORE_API FNullColorVertexBuffer : public FVertexBuffer
+class FNullColorVertexBuffer : public FVertexBuffer
 {
 public:
-	FNullColorVertexBuffer();
-	~FNullColorVertexBuffer();
+	RENDERCORE_API FNullColorVertexBuffer();
+	RENDERCORE_API ~FNullColorVertexBuffer();
 
-	virtual void InitRHI() override;
-	virtual void ReleaseRHI() override;
+	RENDERCORE_API virtual void InitRHI() override;
+	RENDERCORE_API virtual void ReleaseRHI() override;
 
 	FShaderResourceViewRHIRef VertexBufferSRV;
 };
@@ -92,14 +92,14 @@ extern RENDERCORE_API TGlobalResource<FNullColorVertexBuffer> GNullColorVertexBu
 /**
 * A vertex buffer with a single zero float3 component.
 */
-class RENDERCORE_API FNullVertexBuffer : public FVertexBuffer
+class FNullVertexBuffer : public FVertexBuffer
 {
 public:
-	FNullVertexBuffer();
-	~FNullVertexBuffer();
+	RENDERCORE_API FNullVertexBuffer();
+	RENDERCORE_API ~FNullVertexBuffer();
 
-	virtual void InitRHI() override;
-	virtual void ReleaseRHI() override;
+	RENDERCORE_API virtual void InitRHI() override;
+	RENDERCORE_API virtual void ReleaseRHI() override;
 
 	FShaderResourceViewRHIRef VertexBufferSRV;
 };
@@ -107,25 +107,25 @@ public:
 /** The global null vertex buffer, which is set with a stride of 0 on meshes */
 extern RENDERCORE_API TGlobalResource<FNullVertexBuffer> GNullVertexBuffer;
 
-class RENDERCORE_API FScreenSpaceVertexBuffer : public FVertexBuffer
+class FScreenSpaceVertexBuffer : public FVertexBuffer
 {
 public:
 	/**
 	* Initialize the RHI for this rendering resource
 	*/
-	virtual void InitRHI() override;
+	RENDERCORE_API virtual void InitRHI() override;
 };
 
 extern RENDERCORE_API TGlobalResource<FScreenSpaceVertexBuffer> GScreenSpaceVertexBuffer;
 
-class RENDERCORE_API FTileVertexDeclaration : public FRenderResource
+class FTileVertexDeclaration : public FRenderResource
 {
 public:
-	FTileVertexDeclaration();
-	virtual ~FTileVertexDeclaration();
+	RENDERCORE_API FTileVertexDeclaration();
+	RENDERCORE_API virtual ~FTileVertexDeclaration();
 
-	virtual void InitRHI() override;
-	virtual void ReleaseRHI() override;
+	RENDERCORE_API virtual void InitRHI() override;
+	RENDERCORE_API virtual void ReleaseRHI() override;
 
 	FVertexDeclarationRHIRef VertexDeclarationRHI;
 };
@@ -177,7 +177,7 @@ struct FGlobalDynamicVertexBufferAllocation
 /**
  * A system for dynamically allocating GPU memory for vertices.
  */
-class RENDERCORE_API FGlobalDynamicVertexBuffer
+class FGlobalDynamicVertexBuffer
 {
 public:
 	using FAllocation = FGlobalDynamicVertexBufferAllocation;
@@ -187,21 +187,21 @@ public:
 	 * @param SizeInBytes - The amount of memory to allocate in bytes.
 	 * @returns An FAllocation with information regarding the allocated memory.
 	 */
-	FAllocation Allocate(FRHICommandList& RHICmdList, uint32 SizeInBytes);
-	FAllocation Allocate(uint32 SizeInBytes);
+	RENDERCORE_API FAllocation Allocate(FRHICommandList& RHICmdList, uint32 SizeInBytes);
+	RENDERCORE_API FAllocation Allocate(uint32 SizeInBytes);
 
 	/**
 	 * Commits allocated memory to the GPU.
 	 *		WARNING: Once this buffer has been committed to the GPU, allocations
 	 *		remain valid only until the next call to Allocate!
 	 */
-	void Commit(FRHICommandList& RHICmdList);
-	void Commit();
+	RENDERCORE_API void Commit(FRHICommandList& RHICmdList);
+	RENDERCORE_API void Commit();
 
-	static void GarbageCollect();
+	static RENDERCORE_API void GarbageCollect();
 
 	/** Returns true if log statements should be made because we exceeded GMaxVertexBytesAllocatedPerFrame */
-	bool IsRenderAlarmLoggingEnabled() const;
+	RENDERCORE_API bool IsRenderAlarmLoggingEnabled() const;
 
 private:
 	TArray<class FDynamicVertexBuffer*> VertexBuffers;
@@ -244,17 +244,17 @@ struct FGlobalDynamicIndexBufferAllocationEx : public FGlobalDynamicIndexBufferA
 /**
  * A system for dynamically allocating GPU memory for indices.
  */
-class RENDERCORE_API FGlobalDynamicIndexBuffer
+class FGlobalDynamicIndexBuffer
 {
 public:
 	using FAllocation = FGlobalDynamicIndexBufferAllocation;
 	using FAllocationEx = FGlobalDynamicIndexBufferAllocationEx;
 
 	/** Default constructor. */
-	FGlobalDynamicIndexBuffer();
+	RENDERCORE_API FGlobalDynamicIndexBuffer();
 
 	/** Destructor. */
-	~FGlobalDynamicIndexBuffer();
+	RENDERCORE_API ~FGlobalDynamicIndexBuffer();
 
 	/**
 	 * Allocates space in the global index buffer.
@@ -262,7 +262,7 @@ public:
 	 * @param IndexStride - The size of an index (2 or 4 bytes).
 	 * @returns An FAllocation with information regarding the allocated memory.
 	 */
-	FAllocation Allocate(uint32 NumIndices, uint32 IndexStride);
+	RENDERCORE_API FAllocation Allocate(uint32 NumIndices, uint32 IndexStride);
 
 	/**
 	 * Helper function to allocate.
@@ -280,7 +280,7 @@ public:
 	 *		WARNING: Once this buffer has been committed to the GPU, allocations
 	 *		remain valid only until the next call to Allocate!
 	 */
-	void Commit();
+	RENDERCORE_API void Commit();
 
 private:
 	/** The pool of vertex buffers from which allocations are made. */
diff --git a/Engine/Source/Runtime/RenderCore/Public/GlobalShader.h b/Engine/Source/Runtime/RenderCore/Public/GlobalShader.h
--- a/Engine/Source/Runtime/RenderCore/Public/GlobalShader.h
+++ b/Engine/Source/Runtime/RenderCore/Public/GlobalShader.h
@@ -48,7 +48,7 @@ public:
 	/** Append to a string that will be used as a DDC key. */
 	RENDERCORE_API void AppendKeyString(FString& KeyString, const TArray<FShaderTypeDependency>& Dependencies) const;
 
-	RENDERCORE_API const TMap<FString, TArray<FShaderTypeDependency>>& GetShaderFilenameToDependeciesMap() const { return ShaderFilenameToDependenciesMap; }
+	const TMap<FString, TArray<FShaderTypeDependency>>& GetShaderFilenameToDependeciesMap() const { return ShaderFilenameToDependenciesMap; }
 
 	RENDERCORE_API bool WithEditorOnly() const;
 
diff --git a/Engine/Source/Runtime/RenderCore/Public/MediaShaders.h b/Engine/Source/Runtime/RenderCore/Public/MediaShaders.h
--- a/Engine/Source/Runtime/RenderCore/Public/MediaShaders.h
+++ b/Engine/Source/Runtime/RenderCore/Public/MediaShaders.h
@@ -140,14 +140,14 @@ inline FBufferRHIRef CreateTempMediaVertexBuffer(float ULeft = 0.0f, float URigh
 /**
  * The simple element vertex declaration resource type.
  */
-class RENDERCORE_API FMediaVertexDeclaration : public FRenderResource
+class FMediaVertexDeclaration : public FRenderResource
 {
 public:
-	FMediaVertexDeclaration();
-	virtual ~FMediaVertexDeclaration();
+	RENDERCORE_API FMediaVertexDeclaration();
+	RENDERCORE_API virtual ~FMediaVertexDeclaration();
 
-	virtual void InitRHI() override;
-	virtual void ReleaseRHI() override;
+	RENDERCORE_API virtual void InitRHI() override;
+	RENDERCORE_API virtual void ReleaseRHI() override;
 
 	FVertexDeclarationRHIRef VertexDeclarationRHI;
 };
diff --git a/Engine/Source/Runtime/RenderCore/Public/PixelShaderUtils.h b/Engine/Source/Runtime/RenderCore/Public/PixelShaderUtils.h
--- a/Engine/Source/Runtime/RenderCore/Public/PixelShaderUtils.h
+++ b/Engine/Source/Runtime/RenderCore/Public/PixelShaderUtils.h
@@ -37,7 +37,7 @@ class FRDGBuilder;
 
 
 /** All utils for pixel shaders. */
-struct RENDERCORE_API FPixelShaderUtils
+struct FPixelShaderUtils
 {
 	/** Utility vertex shader for rect array based operations. For example for clearing specified parts of an atlas. */
 	class FRasterizeToRectsVS : public FGlobalShader
@@ -61,20 +61,20 @@ struct RENDERCORE_API FPixelShaderUtils
 	};
 
 	/** Draw a single triangle on the entire viewport. */
-	static void DrawFullscreenTriangle(FRHICommandList& RHICmdList, uint32 InstanceCount = 1);
+	static RENDERCORE_API void DrawFullscreenTriangle(FRHICommandList& RHICmdList, uint32 InstanceCount = 1);
 
 	/** Draw a two triangle on the entire viewport. */
-	static void DrawFullscreenQuad(FRHICommandList& RHICmdList, uint32 InstanceCount = 1);
+	static RENDERCORE_API void DrawFullscreenQuad(FRHICommandList& RHICmdList, uint32 InstanceCount = 1);
 
 	/** Initialize a pipeline state object initializer with almost all the basics required to do a full viewport pass. */
-	static void InitFullscreenPipelineState(
+	static RENDERCORE_API void InitFullscreenPipelineState(
 		FRHICommandList& RHICmdList,
 		const FGlobalShaderMap* GlobalShaderMap,
 		const TShaderRef<FShader>& PixelShader,
 		FGraphicsPipelineStateInitializer& GraphicsPSOInit);
 
 	/** Initialize a pipeline state object initializer with almost all the basics required to do a full multi-viewport pass. */
-	static void InitFullscreenMultiviewportPipelineState(
+	static RENDERCORE_API void InitFullscreenMultiviewportPipelineState(
 		FRHICommandList& RHICmdList,
 		const FGlobalShaderMap* GlobalShaderMap,
 		const TShaderRef<FShader>& PixelShader,
@@ -285,4 +285,4 @@ struct RENDERCORE_API FPixelShaderUtils
 			}
 		});
 	}
-};
\ No newline at end of file
+};
diff --git a/Engine/Source/Runtime/RenderCore/Public/ProfilingDebugging/RealtimeGPUProfiler.h b/Engine/Source/Runtime/RenderCore/Public/ProfilingDebugging/RealtimeGPUProfiler.h
--- a/Engine/Source/Runtime/RenderCore/Public/ProfilingDebugging/RealtimeGPUProfiler.h
+++ b/Engine/Source/Runtime/RenderCore/Public/ProfilingDebugging/RealtimeGPUProfiler.h
@@ -40,7 +40,7 @@ class FRenderQueryPool;
 class FScopedGPUStatEvent;
 
 #if RHI_WANT_BREADCRUMB_EVENTS
-struct RENDERCORE_API FBreadcrumbEvent
+struct FBreadcrumbEvent
 {
 	FRHIComputeCommandList* RHICmdList{};
 
@@ -86,7 +86,7 @@ struct RENDERCORE_API FBreadcrumbEvent
 	 * Class that logs draw events based upon class scope. Draw events can be seen
 	 * in PIX
 	 */
-	struct RENDERCORE_API FDrawEvent
+	struct FDrawEvent
 	{
 		/** Cmdlist to push onto. */
 		FRHIComputeCommandList* RHICmdList;
@@ -118,15 +118,15 @@ struct RENDERCORE_API FBreadcrumbEvent
 		 * - Stop can be called on any thread but must be called on the same thread Start was called on and will either use the same 
 		 *  command list (rendering / RHI thread) or enqueue a command on the current immediate command list (game thread)
 		 */
-		void CDECL Start(FRHIComputeCommandList* RHICmdList, FColor Color, const TCHAR* Fmt, ...);
-		void Stop();
+		RENDERCORE_API void CDECL Start(FRHIComputeCommandList* RHICmdList, FColor Color, const TCHAR* Fmt, ...);
+		RENDERCORE_API void Stop();
 	};
 
 	/** Legacy support for template class version. */
 	template <typename TRHICmdList>
 	struct TDrawEvent : FDrawEvent {};
 
-	struct RENDERCORE_API FDrawEventRHIExecute
+	struct FDrawEventRHIExecute
 	{
 		/** Context to execute on*/
 		class IRHIComputeContext* RHICommandContext;
@@ -150,8 +150,8 @@ struct RENDERCORE_API FBreadcrumbEvent
 		/**
 		* Function for logging a PIX event with var args
 		*/
-		void CDECL Start(IRHIComputeContext& InRHICommandContext, FColor Color, const TCHAR* Fmt, ...);
-		void Stop();
+		RENDERCORE_API void CDECL Start(IRHIComputeContext& InRHICommandContext, FColor Color, const TCHAR* Fmt, ...);
+		RENDERCORE_API void Stop();
 	};
 
 	// Macros to allow for scoping of draw events outside of RHI function implementations
@@ -205,7 +205,7 @@ struct RENDERCORE_API FBreadcrumbEvent
 
 #else
 
-	struct RENDERCORE_API FDrawEvent
+	struct FDrawEvent
 	{
 	};
 
@@ -366,7 +366,7 @@ struct FRealtimeGPUProfilerHistoryByDescription
 	mutable FRWLock Mutex;
 };
 
-struct RENDERCORE_API FRealtimeGPUProfilerDescriptionResult
+struct FRealtimeGPUProfilerDescriptionResult
 {
 	// Times are in microseconds
 	FString Description;
diff --git a/Engine/Source/Runtime/RenderCore/Public/RayGenShaderUtils.h b/Engine/Source/Runtime/RenderCore/Public/RayGenShaderUtils.h
--- a/Engine/Source/Runtime/RenderCore/Public/RayGenShaderUtils.h
+++ b/Engine/Source/Runtime/RenderCore/Public/RayGenShaderUtils.h
@@ -11,7 +11,7 @@
 #include "PipelineStateCache.h"
 
 /** All utils for ray generation shaders. */
-struct RENDERCORE_API FRayGenShaderUtils
+struct FRayGenShaderUtils
 {
 	/** Dispatch a ray generation shader to render graph builder with its parameters. */
 	template<typename TShaderClass>
diff --git a/Engine/Source/Runtime/RenderCore/Public/RayTracingGeometry.h b/Engine/Source/Runtime/RenderCore/Public/RayTracingGeometry.h
--- a/Engine/Source/Runtime/RenderCore/Public/RayTracingGeometry.h
+++ b/Engine/Source/Runtime/RenderCore/Public/RayTracingGeometry.h
@@ -19,13 +19,13 @@ enum class ERTAccelerationStructureBuildPriority
 };
 
 /** A ray tracing geometry resource */
-class RENDERCORE_API FRayTracingGeometry : public FRenderResource
+class FRayTracingGeometry : public FRenderResource
 {
 public:
 	TResourceArray<uint8> RawData;
 
-	FRayTracingGeometry();
-	virtual ~FRayTracingGeometry();
+	RENDERCORE_API FRayTracingGeometry();
+	RENDERCORE_API virtual ~FRayTracingGeometry();
 
 #if RHI_RAYTRACING
 
@@ -68,7 +68,7 @@ public:
 		Initializer = InInitializer;
 	}
 
-	bool IsValid() const;
+	RENDERCORE_API bool IsValid() const;
 
 	void SetAsStreamedIn()
 	{
@@ -97,33 +97,33 @@ public:
 		return GeometryState;
 	}
 
-	void InitRHIForStreaming(FRHIRayTracingGeometry* IntermediateGeometry, FRHIResourceUpdateBatcher& Batcher);
-	void ReleaseRHIForStreaming(FRHIResourceUpdateBatcher& Batcher);
+	RENDERCORE_API void InitRHIForStreaming(FRHIRayTracingGeometry* IntermediateGeometry, FRHIResourceUpdateBatcher& Batcher);
+	RENDERCORE_API void ReleaseRHIForStreaming(FRHIResourceUpdateBatcher& Batcher);
 
-	void CreateRayTracingGeometryFromCPUData(TResourceArray<uint8>& OfflineData);
-	void RequestBuildIfNeeded(ERTAccelerationStructureBuildPriority InBuildPriority);
+	RENDERCORE_API void CreateRayTracingGeometryFromCPUData(TResourceArray<uint8>& OfflineData);
+	RENDERCORE_API void RequestBuildIfNeeded(ERTAccelerationStructureBuildPriority InBuildPriority);
 
 	// That function is only supposed to be used when dynamic ray tracing is enabled
-	void InitRHIForDynamicRayTracing();
+	RENDERCORE_API void InitRHIForDynamicRayTracing();
 
-	void CreateRayTracingGeometry(ERTAccelerationStructureBuildPriority InBuildPriority);
+	RENDERCORE_API void CreateRayTracingGeometry(ERTAccelerationStructureBuildPriority InBuildPriority);
 
 	bool HasPendingBuildRequest() const
 	{
 		return RayTracingBuildRequestIndex != INDEX_NONE;
 	}
-	void BoostBuildPriority(float InBoostValue = 0.01f) const;
+	RENDERCORE_API void BoostBuildPriority(float InBoostValue = 0.01f) const;
 
 	// FRenderResource interface
 
 	virtual FString GetFriendlyName() const override { return TEXT("FRayTracingGeometry"); }
 
-	virtual void InitRHI() override;
-	virtual void ReleaseRHI() override;
+	RENDERCORE_API virtual void InitRHI() override;
+	RENDERCORE_API virtual void ReleaseRHI() override;
 
-	virtual void ReleaseResource() override;
+	RENDERCORE_API virtual void ReleaseResource() override;
 protected:
-	void RemoveBuildRequest();
+	RENDERCORE_API void RemoveBuildRequest();
 
 	friend class FRayTracingGeometryManager;
 	int32 RayTracingBuildRequestIndex = INDEX_NONE;
diff --git a/Engine/Source/Runtime/RenderCore/Public/RayTracingGeometryManager.h b/Engine/Source/Runtime/RenderCore/Public/RayTracingGeometryManager.h
--- a/Engine/Source/Runtime/RenderCore/Public/RayTracingGeometryManager.h
+++ b/Engine/Source/Runtime/RenderCore/Public/RayTracingGeometryManager.h
@@ -15,7 +15,7 @@ class FRHIComputeCommandList;
 enum class EAccelerationStructureBuildMode;
 enum class ERTAccelerationStructureBuildPriority;
 
-class RENDERCORE_API FRayTracingGeometryManager
+class FRayTracingGeometryManager
 {
 public:
 
@@ -29,17 +29,17 @@ public:
 	{
 		return RequestBuildAccelerationStructure(InGeometry, InPriority, EAccelerationStructureBuildMode::Build);
 	}	
-	BuildRequestIndex RequestBuildAccelerationStructure(FRayTracingGeometry* InGeometry, ERTAccelerationStructureBuildPriority InPriority, EAccelerationStructureBuildMode InBuildMode);
+	RENDERCORE_API BuildRequestIndex RequestBuildAccelerationStructure(FRayTracingGeometry* InGeometry, ERTAccelerationStructureBuildPriority InPriority, EAccelerationStructureBuildMode InBuildMode);
 
-	void RemoveBuildRequest(BuildRequestIndex InRequestIndex);
-	void BoostPriority(BuildRequestIndex InRequestIndex, float InBoostValue);
-	void ForceBuildIfPending(FRHIComputeCommandList& InCmdList, const TArrayView<const FRayTracingGeometry*> InGeometries);
-	void ProcessBuildRequests(FRHIComputeCommandList& InCmdList, bool bInBuildAll = false);
+	RENDERCORE_API void RemoveBuildRequest(BuildRequestIndex InRequestIndex);
+	RENDERCORE_API void BoostPriority(BuildRequestIndex InRequestIndex, float InBoostValue);
+	RENDERCORE_API void ForceBuildIfPending(FRHIComputeCommandList& InCmdList, const TArrayView<const FRayTracingGeometry*> InGeometries);
+	RENDERCORE_API void ProcessBuildRequests(FRHIComputeCommandList& InCmdList, bool bInBuildAll = false);
 
-	void Tick(bool bHasRayTracingEnableChanged);
+	RENDERCORE_API void Tick(bool bHasRayTracingEnableChanged);
 
-	RayTracingGeometryHandle RegisterRayTracingGeometry(FRayTracingGeometry* InGeometry);
-	void ReleaseRayTracingGeometryHandle(RayTracingGeometryHandle Handle);
+	RENDERCORE_API RayTracingGeometryHandle RegisterRayTracingGeometry(FRayTracingGeometry* InGeometry);
+	RENDERCORE_API void ReleaseRayTracingGeometryHandle(RayTracingGeometryHandle Handle);
 private:
 
 	struct BuildRequest
@@ -51,7 +51,7 @@ private:
 		EAccelerationStructureBuildMode BuildMode;
 	};
 
-	void SetupBuildParams(const BuildRequest& InBuildRequest, TArray<FRayTracingGeometryBuildParams>& InBuildParams);
+	RENDERCORE_API void SetupBuildParams(const BuildRequest& InBuildRequest, TArray<FRayTracingGeometryBuildParams>& InBuildParams);
 
 	FCriticalSection RequestCS;
 
diff --git a/Engine/Source/Runtime/RenderCore/Public/RayTracingValidationShaders.h b/Engine/Source/Runtime/RenderCore/Public/RayTracingValidationShaders.h
--- a/Engine/Source/Runtime/RenderCore/Public/RayTracingValidationShaders.h
+++ b/Engine/Source/Runtime/RenderCore/Public/RayTracingValidationShaders.h
@@ -8,7 +8,7 @@
 
 #if RHI_RAYTRACING
 
-class RENDERCORE_API FRayTracingValidateGeometryBuildParamsCS : public FBuiltInRayTracingShader
+class FRayTracingValidateGeometryBuildParamsCS : public FBuiltInRayTracingShader
 {
 	DECLARE_GLOBAL_SHADER(FRayTracingValidateGeometryBuildParamsCS);
 
@@ -42,7 +42,7 @@ public:
 		FBuiltInRayTracingShader::ModifyCompilationEnvironment(Parameters, OutEnvironment);
 	}
 
-	static void Dispatch(FRHICommandList& RHICmdList, const FRayTracingGeometryBuildParams& Params);
+	static RENDERCORE_API void Dispatch(FRHICommandList& RHICmdList, const FRayTracingGeometryBuildParams& Params);
 
 	LAYOUT_FIELD(FShaderResourceParameter, IndexBufferParam);
 	LAYOUT_FIELD(FShaderResourceParameter, VertexBufferParam);
@@ -55,7 +55,7 @@ public:
 
 };
 
-class RENDERCORE_API FRayTracingValidateSceneBuildParamsCS : public FBuiltInRayTracingShader
+class FRayTracingValidateSceneBuildParamsCS : public FBuiltInRayTracingShader
 {
 	DECLARE_GLOBAL_SHADER(FRayTracingValidateSceneBuildParamsCS);
 
@@ -86,7 +86,7 @@ public:
 		FBuiltInRayTracingShader::ModifyCompilationEnvironment(Parameters, OutEnvironment);
 	}
 
-	static void Dispatch(FRHICommandList& RHICmdList, 
+	static RENDERCORE_API void Dispatch(FRHICommandList& RHICmdList, 
 		uint32 NumHitGroups, uint32 NumInstances, 
 		FRHIBuffer* InstanceBuffer, uint32 InstanceBufferOffset, uint32 InstanceBufferStride);
 
diff --git a/Engine/Source/Runtime/RenderCore/Public/RenderCommandFence.h b/Engine/Source/Runtime/RenderCore/Public/RenderCommandFence.h
--- a/Engine/Source/Runtime/RenderCore/Public/RenderCommandFence.h
+++ b/Engine/Source/Runtime/RenderCore/Public/RenderCommandFence.h
@@ -14,7 +14,7 @@ namespace ENamedThreads { enum Type : int32; }
  /**
  * Used to track pending rendering commands from the game thread.
  */
-class RENDERCORE_API FRenderCommandFence
+class FRenderCommandFence
 {
 public:
 
@@ -24,20 +24,20 @@ public:
 	 * Once the rendering thread has executed the fence command, it decrements the pending fence count.
 	 * @param bSyncToRHIAndGPU, true if we should wait for the RHI thread or GPU, otherwise we only wait for the render thread.
 	 */
-	void BeginFence(bool bSyncToRHIAndGPU = false);
+	RENDERCORE_API void BeginFence(bool bSyncToRHIAndGPU = false);
 
 	/**
 	 * Waits for pending fence commands to retire.
 	 * @param bProcessGameThreadTasks, if true we are on a short callstack where it is safe to process arbitrary game thread tasks while we wait
 	 */
-	void Wait(bool bProcessGameThreadTasks = false) const;
+	RENDERCORE_API void Wait(bool bProcessGameThreadTasks = false) const;
 
 	// return true if the fence is complete
-	bool IsFenceComplete() const;
+	RENDERCORE_API bool IsFenceComplete() const;
 
 	// Ctor/dtor
-	FRenderCommandFence();
-	~FRenderCommandFence();
+	RENDERCORE_API FRenderCommandFence();
+	RENDERCORE_API ~FRenderCommandFence();
 
 private:
 	/** Graph event that represents completion of this fence **/
diff --git a/Engine/Source/Runtime/RenderCore/Public/RenderDeferredCleanup.h b/Engine/Source/Runtime/RenderCore/Public/RenderDeferredCleanup.h
--- a/Engine/Source/Runtime/RenderCore/Public/RenderDeferredCleanup.h
+++ b/Engine/Source/Runtime/RenderCore/Public/RenderDeferredCleanup.h
@@ -7,7 +7,7 @@
 /**
  * The base class of objects that need to defer deletion until the render command queue has been flushed.
  */
-class RENDERCORE_API FDeferredCleanupInterface
+class FDeferredCleanupInterface
 {
 public:
 	virtual ~FDeferredCleanupInterface() {}
diff --git a/Engine/Source/Runtime/RenderCore/Public/RenderGraphAllocator.h b/Engine/Source/Runtime/RenderCore/Public/RenderGraphAllocator.h
--- a/Engine/Source/Runtime/RenderCore/Public/RenderGraphAllocator.h
+++ b/Engine/Source/Runtime/RenderCore/Public/RenderGraphAllocator.h
@@ -17,7 +17,7 @@
 #define RDG_USE_MALLOC USING_ADDRESS_SANITISER
 
 /** Private allocator used by RDG to track its internal memory. All memory is released after RDG builder execution. */
-class RENDERCORE_API FRDGAllocator final
+class FRDGAllocator final
 {
 public:
 	class FObject
@@ -39,8 +39,8 @@ public:
 		T Alloc;
 	};
 
-	static FRDGAllocator& Get();
-	~FRDGAllocator();
+	static RENDERCORE_API FRDGAllocator& Get();
+	RENDERCORE_API ~FRDGAllocator();
 
 	/** Allocates raw memory. */
 	FORCEINLINE void* Alloc(uint64 SizeInBytes, uint32 AlignInBytes)
@@ -101,7 +101,7 @@ private:
 	FRDGAllocator(FRDGAllocator&&) = default;
 	FRDGAllocator& operator = (FRDGAllocator&&) = default;
 
-	void ReleaseAll();
+	RENDERCORE_API void ReleaseAll();
 
 	struct FContext
 	{
@@ -132,14 +132,14 @@ private:
 #define RDG_FRIEND_ALLOCATOR_FRIEND(Type) friend class FRDGAllocator::TObject<Type>
 
 /** Base class for RDG builder which scopes the allocations and releases them in the destructor. */
-class RENDERCORE_API FRDGAllocatorScope
+class FRDGAllocatorScope
 {
 public:
 	FRDGAllocatorScope()
 		: Allocator(FRDGAllocator::Get())
 	{}
 
-	~FRDGAllocatorScope();
+	RENDERCORE_API ~FRDGAllocatorScope();
 
 protected:
 	FRDGAllocator& Allocator;
@@ -259,4 +259,4 @@ struct TAllocatorTraits<TRDGArrayAllocator<Alignment>> : TAllocatorTraitsBase<TR
 using FRDGArrayAllocator = TRDGArrayAllocator<>;
 using FRDGBitArrayAllocator = TInlineAllocator<4, FRDGArrayAllocator>;
 using FRDGSparseArrayAllocator = TSparseArrayAllocator<FRDGArrayAllocator, FRDGBitArrayAllocator>;
-using FRDGSetAllocator = TSetAllocator<FRDGSparseArrayAllocator, TInlineAllocator<1, FRDGBitArrayAllocator>>;
\ No newline at end of file
+using FRDGSetAllocator = TSetAllocator<FRDGSparseArrayAllocator, TInlineAllocator<1, FRDGBitArrayAllocator>>;
diff --git a/Engine/Source/Runtime/RenderCore/Public/RenderGraphBuilder.h b/Engine/Source/Runtime/RenderCore/Public/RenderGraphBuilder.h
--- a/Engine/Source/Runtime/RenderCore/Public/RenderGraphBuilder.h
+++ b/Engine/Source/Runtime/RenderCore/Public/RenderGraphBuilder.h
@@ -43,39 +43,39 @@ struct TStatId;
  *  graph is compiled, culled, and executed in Execute(). The builder should be created on the stack and executed prior to
  *  destruction.
  */
-class RENDERCORE_API FRDGBuilder
+class FRDGBuilder
 	: FRDGAllocatorScope
 {
 public:
-	FRDGBuilder(FRHICommandListImmediate& RHICmdList, FRDGEventName Name = {}, ERDGBuilderFlags Flags = ERDGBuilderFlags::None);
+	RENDERCORE_API FRDGBuilder(FRHICommandListImmediate& RHICmdList, FRDGEventName Name = {}, ERDGBuilderFlags Flags = ERDGBuilderFlags::None);
 	FRDGBuilder(const FRDGBuilder&) = delete;
-	~FRDGBuilder();
+	RENDERCORE_API ~FRDGBuilder();
 
 	/** Finds an RDG texture associated with the external texture, or returns null if none is found. */
-	FRDGTexture* FindExternalTexture(FRHITexture* Texture) const;
-	FRDGTexture* FindExternalTexture(IPooledRenderTarget* ExternalPooledTexture) const;
+	RENDERCORE_API FRDGTexture* FindExternalTexture(FRHITexture* Texture) const;
+	RENDERCORE_API FRDGTexture* FindExternalTexture(IPooledRenderTarget* ExternalPooledTexture) const;
 
 	/** Finds an RDG buffer associated with the external buffer, or returns null if none is found. */
-	FRDGBuffer* FindExternalBuffer(FRHIBuffer* Buffer) const;
-	FRDGBuffer* FindExternalBuffer(FRDGPooledBuffer* ExternalPooledBuffer) const;
+	RENDERCORE_API FRDGBuffer* FindExternalBuffer(FRHIBuffer* Buffer) const;
+	RENDERCORE_API FRDGBuffer* FindExternalBuffer(FRDGPooledBuffer* ExternalPooledBuffer) const;
 
 	/** Registers a external pooled render target texture to be tracked by the render graph. The name of the registered RDG texture is pulled from the pooled render target. */
-	FRDGTextureRef RegisterExternalTexture(
+	RENDERCORE_API FRDGTextureRef RegisterExternalTexture(
 		const TRefCountPtr<IPooledRenderTarget>& ExternalPooledTexture,
 		ERDGTextureFlags Flags = ERDGTextureFlags::None);
 
 	/** Register an external texture with a custom name. The name is only used if the texture has not already been registered. */
-	FRDGTextureRef RegisterExternalTexture(
+	RENDERCORE_API FRDGTextureRef RegisterExternalTexture(
 		const TRefCountPtr<IPooledRenderTarget>& ExternalPooledTexture,
 		const TCHAR* NameIfNotRegistered,
 		ERDGTextureFlags Flags = ERDGTextureFlags::None);
 
 	/** Register a external buffer to be tracked by the render graph. */
-	FRDGBufferRef RegisterExternalBuffer(const TRefCountPtr<FRDGPooledBuffer>& ExternalPooledBuffer, ERDGBufferFlags Flags = ERDGBufferFlags::None);
-	FRDGBufferRef RegisterExternalBuffer(const TRefCountPtr<FRDGPooledBuffer>& ExternalPooledBuffer, ERDGBufferFlags Flags, ERHIAccess AccessFinal);
+	RENDERCORE_API FRDGBufferRef RegisterExternalBuffer(const TRefCountPtr<FRDGPooledBuffer>& ExternalPooledBuffer, ERDGBufferFlags Flags = ERDGBufferFlags::None);
+	RENDERCORE_API FRDGBufferRef RegisterExternalBuffer(const TRefCountPtr<FRDGPooledBuffer>& ExternalPooledBuffer, ERDGBufferFlags Flags, ERHIAccess AccessFinal);
 
 	/** Register an external buffer with a custom name. The name is only used if the buffer has not already been registered. */
-	FRDGBufferRef RegisterExternalBuffer(
+	RENDERCORE_API FRDGBufferRef RegisterExternalBuffer(
 		const TRefCountPtr<FRDGPooledBuffer>& ExternalPooledBuffer,
 		const TCHAR* NameIfNotRegistered,
 		ERDGBufferFlags Flags = ERDGBufferFlags::None);
@@ -85,25 +85,25 @@ public:
 	 *  declare the texture in the pass parameter struct. The name is the name used for GPU debugging tools and the the
 	 *  VisualizeTexture/Vis command.
 	 */
-	FRDGTextureRef CreateTexture(const FRDGTextureDesc& Desc, const TCHAR* Name, ERDGTextureFlags Flags = ERDGTextureFlags::None);
+	RENDERCORE_API FRDGTextureRef CreateTexture(const FRDGTextureDesc& Desc, const TCHAR* Name, ERDGTextureFlags Flags = ERDGTextureFlags::None);
 
 	/** Create graph tracked buffer from a descriptor. The CPU memory is guaranteed to be valid through execution of
 	 *  the graph, at which point it is released. The underlying RHI buffer lifetime is only guaranteed for passes which
 	 *  declare the buffer in the pass parameter struct. The name is the name used for GPU debugging tools.
 	 */
-	FRDGBufferRef CreateBuffer(const FRDGBufferDesc& Desc, const TCHAR* Name, ERDGBufferFlags Flags = ERDGBufferFlags::None);
+	RENDERCORE_API FRDGBufferRef CreateBuffer(const FRDGBufferDesc& Desc, const TCHAR* Name, ERDGBufferFlags Flags = ERDGBufferFlags::None);
 
 	/** A variant of CreateBuffer where users supply NumElements through a callback. This allows creating buffers with
 	 *  sizes unknown at creation time. The callback is called before executing the most recent RDG pass that references
 	 *  the buffer so data must be ready before that.
 	 */
-	FRDGBufferRef CreateBuffer(const FRDGBufferDesc& Desc, const TCHAR* Name, FRDGBufferNumElementsCallback&& NumElementsCallback, ERDGBufferFlags Flags = ERDGBufferFlags::None);
+	RENDERCORE_API FRDGBufferRef CreateBuffer(const FRDGBufferDesc& Desc, const TCHAR* Name, FRDGBufferNumElementsCallback&& NumElementsCallback, ERDGBufferFlags Flags = ERDGBufferFlags::None);
 
 	/** Create graph tracked SRV for a texture from a descriptor. */
-	FRDGTextureSRVRef CreateSRV(const FRDGTextureSRVDesc& Desc);
+	RENDERCORE_API FRDGTextureSRVRef CreateSRV(const FRDGTextureSRVDesc& Desc);
 
 	/** Create graph tracked SRV for a buffer from a descriptor. */
-	FRDGBufferSRVRef CreateSRV(const FRDGBufferSRVDesc& Desc);
+	RENDERCORE_API FRDGBufferSRVRef CreateSRV(const FRDGBufferSRVDesc& Desc);
 
 	FORCEINLINE FRDGBufferSRVRef CreateSRV(FRDGBufferRef Buffer, EPixelFormat Format)
 	{
@@ -111,7 +111,7 @@ public:
 	}
 
 	/** Create graph tracked UAV for a texture from a descriptor. */
-	FRDGTextureUAVRef CreateUAV(const FRDGTextureUAVDesc& Desc, ERDGUnorderedAccessViewFlags Flags = ERDGUnorderedAccessViewFlags::None);
+	RENDERCORE_API FRDGTextureUAVRef CreateUAV(const FRDGTextureUAVDesc& Desc, ERDGUnorderedAccessViewFlags Flags = ERDGUnorderedAccessViewFlags::None);
 
 	FORCEINLINE FRDGTextureUAVRef CreateUAV(FRDGTextureRef Texture, ERDGUnorderedAccessViewFlags Flags = ERDGUnorderedAccessViewFlags::None, EPixelFormat Format = PF_Unknown)
 	{
@@ -119,7 +119,7 @@ public:
 	}
 
 	/** Create graph tracked UAV for a buffer from a descriptor. */
-	FRDGBufferUAVRef CreateUAV(const FRDGBufferUAVDesc& Desc, ERDGUnorderedAccessViewFlags Flags = ERDGUnorderedAccessViewFlags::None);
+	RENDERCORE_API FRDGBufferUAVRef CreateUAV(const FRDGBufferUAVDesc& Desc, ERDGUnorderedAccessViewFlags Flags = ERDGUnorderedAccessViewFlags::None);
 
 	FORCEINLINE FRDGBufferUAVRef CreateUAV(FRDGBufferRef Buffer, EPixelFormat Format, ERDGUnorderedAccessViewFlags Flags = ERDGUnorderedAccessViewFlags::None)
 	{
@@ -137,7 +137,7 @@ public:
 	// Allocation Methods
 
 	/** Allocates raw memory using an allocator tied to the lifetime of the graph. */
-	void* Alloc(uint64 SizeInBytes, uint32 AlignInBytes = 16);
+	RENDERCORE_API void* Alloc(uint64 SizeInBytes, uint32 AlignInBytes = 16);
 
 	/** Allocates POD memory using an allocator tied to the lifetime of the graph. Does not construct / destruct. */
 	template <typename PODType>
@@ -209,16 +209,16 @@ public:
 	 *  Recommended usage is to set a workload equal to the number of complex draw / dispatch calls (each with its own parameters, etc), and
 	 *  only as a performance tweak if a particular pass is very expensive relative to other passes.
 	 */
-	void SetPassWorkload(FRDGPass* Pass, uint32 Workload);
+	RENDERCORE_API void SetPassWorkload(FRDGPass* Pass, uint32 Workload);
 
 	/** Adds a user-defined dependency between two passes. This can be used to fine-tune async compute overlap by forcing a sync point. */
-	void AddPassDependency(FRDGPass* Producer, FRDGPass* Consumer);
+	RENDERCORE_API void AddPassDependency(FRDGPass* Producer, FRDGPass* Consumer);
 
 	/** Sets the current command list stat for all subsequent passes. */
-	void SetCommandListStat(TStatId StatId);
+	RENDERCORE_API void SetCommandListStat(TStatId StatId);
 
 	/** A hint to the builder to flush work to the RHI thread after the last queued pass on the execution timeline. */
-	void AddDispatchHint();
+	RENDERCORE_API void AddDispatchHint();
 
 	/** Launches a task that is synced prior to graph execution. If parallel execution is not enabled, the lambda is run immediately. */
 	template <typename TaskLambda>
@@ -256,10 +256,10 @@ public:
 	 *   Deferred:  RHI resource flushes are performed prior to execution.
 	 *   Immediate: RHI resource flushes are performed immediately.
 	 */
-	void SetFlushResourcesRHI();
+	RENDERCORE_API void SetFlushResourcesRHI();
 
 	/** Queues a buffer upload operation prior to execution. The resource lifetime is extended and the data is uploaded prior to executing passes. */
-	void QueueBufferUpload(FRDGBufferRef Buffer, const void* InitialData, uint64 InitialDataSize, ERDGInitialDataFlags InitialDataFlags = ERDGInitialDataFlags::None);
+	RENDERCORE_API void QueueBufferUpload(FRDGBufferRef Buffer, const void* InitialData, uint64 InitialDataSize, ERDGInitialDataFlags InitialDataFlags = ERDGInitialDataFlags::None);
 
 	template <typename ElementType>
 	inline void QueueBufferUpload(FRDGBufferRef Buffer, TArrayView<ElementType, int32> Container, ERDGInitialDataFlags InitialDataFlags = ERDGInitialDataFlags::None)
@@ -268,7 +268,7 @@ public:
 	}
 
 	/** Queues a buffer upload operation prior to execution. The resource lifetime is extended and the data is uploaded prior to executing passes. */
-	void QueueBufferUpload(FRDGBufferRef Buffer, const void* InitialData, uint64 InitialDataSize, FRDGBufferInitialDataFreeCallback&& InitialDataFreeCallback);
+	RENDERCORE_API void QueueBufferUpload(FRDGBufferRef Buffer, const void* InitialData, uint64 InitialDataSize, FRDGBufferInitialDataFreeCallback&& InitialDataFreeCallback);
 
 	template <typename ElementType>
 	inline void QueueBufferUpload(FRDGBufferRef Buffer, TArrayView<ElementType, int32> Container, FRDGBufferInitialDataFreeCallback&& InitialDataFreeCallback)
@@ -279,45 +279,45 @@ public:
 	/** A variant where InitialData and InitialDataSize are supplied through callbacks. This allows queuing an upload with information unknown at
 	 *  creation time. The callbacks are called before RDG pass execution so data must be ready before that.
 	 */
-	void QueueBufferUpload(FRDGBufferRef Buffer, FRDGBufferInitialDataCallback&& InitialDataCallback, FRDGBufferInitialDataSizeCallback&& InitialDataSizeCallback);
-	void QueueBufferUpload(FRDGBufferRef Buffer, FRDGBufferInitialDataCallback&& InitialDataCallback, FRDGBufferInitialDataSizeCallback&& InitialDataSizeCallback, FRDGBufferInitialDataFreeCallback&& InitialDataFreeCallback);
+	RENDERCORE_API void QueueBufferUpload(FRDGBufferRef Buffer, FRDGBufferInitialDataCallback&& InitialDataCallback, FRDGBufferInitialDataSizeCallback&& InitialDataSizeCallback);
+	RENDERCORE_API void QueueBufferUpload(FRDGBufferRef Buffer, FRDGBufferInitialDataCallback&& InitialDataCallback, FRDGBufferInitialDataSizeCallback&& InitialDataSizeCallback, FRDGBufferInitialDataFreeCallback&& InitialDataFreeCallback);
 
 	/** Queues a pooled render target extraction to happen at the end of graph execution. For graph-created textures, this extends
 	 *  the lifetime of the GPU resource until execution, at which point the pointer is filled. If specified, the texture is transitioned
 	 *  to the AccessFinal state, or kDefaultAccessFinal otherwise.
 	 */
-	void QueueTextureExtraction(FRDGTextureRef Texture, TRefCountPtr<IPooledRenderTarget>* OutPooledTexturePtr, ERDGResourceExtractionFlags Flags = ERDGResourceExtractionFlags::None);
-	void QueueTextureExtraction(FRDGTextureRef Texture, TRefCountPtr<IPooledRenderTarget>* OutPooledTexturePtr, ERHIAccess AccessFinal, ERDGResourceExtractionFlags Flags = ERDGResourceExtractionFlags::None);
+	RENDERCORE_API void QueueTextureExtraction(FRDGTextureRef Texture, TRefCountPtr<IPooledRenderTarget>* OutPooledTexturePtr, ERDGResourceExtractionFlags Flags = ERDGResourceExtractionFlags::None);
+	RENDERCORE_API void QueueTextureExtraction(FRDGTextureRef Texture, TRefCountPtr<IPooledRenderTarget>* OutPooledTexturePtr, ERHIAccess AccessFinal, ERDGResourceExtractionFlags Flags = ERDGResourceExtractionFlags::None);
 
 	/** Queues a pooled buffer extraction to happen at the end of graph execution. For graph-created buffers, this extends the lifetime
 	 *  of the GPU resource until execution, at which point the pointer is filled. If specified, the buffer is transitioned to the
 	 *  AccessFinal state, or kDefaultAccessFinal otherwise.
 	 */
-	void QueueBufferExtraction(FRDGBufferRef Buffer, TRefCountPtr<FRDGPooledBuffer>* OutPooledBufferPtr);
-	void QueueBufferExtraction(FRDGBufferRef Buffer, TRefCountPtr<FRDGPooledBuffer>* OutPooledBufferPtr, ERHIAccess AccessFinal);
+	RENDERCORE_API void QueueBufferExtraction(FRDGBufferRef Buffer, TRefCountPtr<FRDGPooledBuffer>* OutPooledBufferPtr);
+	RENDERCORE_API void QueueBufferExtraction(FRDGBufferRef Buffer, TRefCountPtr<FRDGPooledBuffer>* OutPooledBufferPtr, ERHIAccess AccessFinal);
 
 	/** For graph-created resources, this forces immediate allocation of the underlying pooled resource, effectively promoting it
 	 *  to an external resource. This will increase memory pressure, but allows for querying the pooled resource with GetPooled{Texture, Buffer}.
 	 *  This is primarily used as an aid for porting code incrementally to RDG.
 	 */
-	const TRefCountPtr<IPooledRenderTarget>& ConvertToExternalTexture(FRDGTextureRef Texture);
-	const TRefCountPtr<FRDGPooledBuffer>& ConvertToExternalBuffer(FRDGBufferRef Buffer);
+	RENDERCORE_API const TRefCountPtr<IPooledRenderTarget>& ConvertToExternalTexture(FRDGTextureRef Texture);
+	RENDERCORE_API const TRefCountPtr<FRDGPooledBuffer>& ConvertToExternalBuffer(FRDGBufferRef Buffer);
 	/** For a graph-created uniform buffer, this forces immediate allocation of the underlying resource, effectively promoting it
 	 *  to an external resource. This will increase memory pressure, but allows access to the RHI resource.
 	 *  Graph resources that are referenced in the buffer will be converted to external.
 	 *  This is primarily used as an aid for porting code incrementally to RDG.
 	 */
-	FRHIUniformBuffer* ConvertToExternalUniformBuffer(FRDGUniformBufferRef UniformBuffer);
+	RENDERCORE_API FRHIUniformBuffer* ConvertToExternalUniformBuffer(FRDGUniformBufferRef UniformBuffer);
 
 	/** Performs an immediate query for the underlying pooled resource. This is only allowed for external or extracted resources. */
-	const TRefCountPtr<IPooledRenderTarget>& GetPooledTexture(FRDGTextureRef Texture) const;
-	const TRefCountPtr<FRDGPooledBuffer>& GetPooledBuffer(FRDGBufferRef Buffer) const;
+	RENDERCORE_API const TRefCountPtr<IPooledRenderTarget>& GetPooledTexture(FRDGTextureRef Texture) const;
+	RENDERCORE_API const TRefCountPtr<FRDGPooledBuffer>& GetPooledBuffer(FRDGBufferRef Buffer) const;
 
 	/** (External | Extracted only) Sets the access to transition to after execution at the end of the graph. Overwrites any previously set final access. */
-	void SetTextureAccessFinal(FRDGTextureRef Texture, ERHIAccess Access);
+	RENDERCORE_API void SetTextureAccessFinal(FRDGTextureRef Texture, ERHIAccess Access);
 
 	/** (External | Extracted only) Sets the access to transition to after execution at the end of the graph. Overwrites any previously set final access. */
-	void SetBufferAccessFinal(FRDGBufferRef Buffer, ERHIAccess Access);
+	RENDERCORE_API void SetBufferAccessFinal(FRDGBufferRef Buffer, ERHIAccess Access);
 
 	/** Configures the resource for external access for all subsequent passes, or until UseInternalAccessMode is called.
 	 *  Only read-only access states are allowed. When in external access mode, it is safe to access the underlying RHI
@@ -325,7 +325,7 @@ public:
 	 *  The method effectively guarantees that RDG will transition the resource into the desired state for all subsequent
 	 *  passes so long as the resource remains externally accessible.
 	 */
-	void UseExternalAccessMode(FRDGViewableResource* Resource, ERHIAccess ReadOnlyAccess, ERHIPipeline Pipelines = ERHIPipeline::Graphics);
+	RENDERCORE_API void UseExternalAccessMode(FRDGViewableResource* Resource, ERHIAccess ReadOnlyAccess, ERHIPipeline Pipelines = ERHIPipeline::Graphics);
 
 	void UseExternalAccessMode(TArrayView<FRDGViewableResource* const> Resources, ERHIAccess ReadOnlyAccess, ERHIPipeline Pipelines = ERHIPipeline::Graphics)
 	{
@@ -339,7 +339,7 @@ public:
 	 *  even if external access mode was not enabled (it will simply no-op). It is not valid to access the underlying RHI
 	 *  resource in any pass added after calling this method.
 	 */
-	void UseInternalAccessMode(FRDGViewableResource* Resource);
+	RENDERCORE_API void UseInternalAccessMode(FRDGViewableResource* Resource);
 
 	inline void UseInternalAccessMode(TArrayView<FRDGViewableResource* const> Resources)
 	{
@@ -350,26 +350,26 @@ public:
 	}
 
 	/** Flag a resource that is produced by a pass but never used or extracted to not emit an 'unused' warning. */
-	void RemoveUnusedTextureWarning(FRDGTextureRef Texture);
-	void RemoveUnusedBufferWarning(FRDGBufferRef Buffer);
+	RENDERCORE_API void RemoveUnusedTextureWarning(FRDGTextureRef Texture);
+	RENDERCORE_API void RemoveUnusedBufferWarning(FRDGBufferRef Buffer);
 
 	/** Manually begins a new GPU event scope. */
-	void BeginEventScope(FRDGEventName&& Name);
+	RENDERCORE_API void BeginEventScope(FRDGEventName&& Name);
 
 	/** Manually ends the current GPU event scope. */
-	void EndEventScope();
+	RENDERCORE_API void EndEventScope();
 
 	/** Flushes all queued passes to an async task to perform setup work. */
-	void FlushSetupQueue();
+	RENDERCORE_API void FlushSetupQueue();
 
 	/** Executes the queued passes, managing setting of render targets (RHI RenderPasses), resource transitions and queued texture extraction. */
-	void Execute();
+	RENDERCORE_API void Execute();
 
 	/** Per-frame update of the render graph resource pool. */
-	static void TickPoolElements();
+	static RENDERCORE_API void TickPoolElements();
 
 	/** Whether RDG is running in immediate mode. */
-	static bool IsImmediateMode();
+	static RENDERCORE_API bool IsImmediateMode();
 
 	/** The RHI command list used for the render graph. */
 	FRHICommandListImmediate& RHICmdList;
@@ -378,17 +378,17 @@ public:
 	FRDGBlackboard Blackboard;
 
 #if RDG_DUMP_RESOURCES
-	static FString BeginResourceDump(const TCHAR* Cmd);
-	static void InitResourceDump();
-	static void EndResourceDump();
-	static bool IsDumpingFrame();
+	static RENDERCORE_API FString BeginResourceDump(const TCHAR* Cmd);
+	static RENDERCORE_API void InitResourceDump();
+	static RENDERCORE_API void EndResourceDump();
+	static RENDERCORE_API bool IsDumpingFrame();
 #else
 	static bool IsDumpingFrame() { return false; }
 #endif
 
 #if RDG_DUMP_RESOURCES_AT_EACH_DRAW
-	static void DumpDraw(const FRDGEventName& DrawEventName);
-	static bool IsDumpingDraws();
+	static RENDERCORE_API void DumpDraw(const FRDGEventName& DrawEventName);
+	static RENDERCORE_API bool IsDumpingDraws();
 #else
 	static inline bool IsDumpingDraws()
 	{
@@ -405,7 +405,7 @@ public:
 #endif
 
 private:
-	static const char* const kDefaultUnaccountedCSVStat;
+	static RENDERCORE_API const char* const kDefaultUnaccountedCSVStat;
 
 	const FRDGEventName BuilderName;
 
@@ -417,9 +417,9 @@ private:
 		ERDGPassFlags Flags,
 		ExecuteLambdaType&& ExecuteLambda);
 
-	static ERDGPassFlags OverridePassFlags(const TCHAR* PassName, ERDGPassFlags Flags);
+	static RENDERCORE_API ERDGPassFlags OverridePassFlags(const TCHAR* PassName, ERDGPassFlags Flags);
 
-	void AddProloguePass();
+	RENDERCORE_API void AddProloguePass();
 
 	FORCEINLINE FRDGPass* GetProloguePass() const
 	{
@@ -500,7 +500,7 @@ private:
 	}
 
 #if WITH_MGPU
-	void ForceCopyCrossGPU();
+	RENDERCORE_API void ForceCopyCrossGPU();
 #endif
 
 	/** Registry of graph objects. */
@@ -721,108 +721,108 @@ private:
 
 	FRHICommandListScopedExtendResourceLifetime ExtendResourceLifetimeScope;
 
-	void MarkResourcesAsProduced(FRDGPass* Pass);
+	RENDERCORE_API void MarkResourcesAsProduced(FRDGPass* Pass);
 
-	void Compile();
-	void Clear();
+	RENDERCORE_API void Compile();
+	RENDERCORE_API void Clear();
 
-	void SetRHI(FRDGTexture* Texture, IPooledRenderTarget* RenderTarget, FRDGPassHandle PassHandle);
-	void SetRHI(FRDGTexture* Texture, FRDGPooledTexture* PooledTexture, FRDGPassHandle PassHandle);
-	void SetRHI(FRDGTexture* Texture, FRHITransientTexture* TransientTexture, FRDGPassHandle PassHandle);
-	void SetRHI(FRDGBuffer* Buffer, FRDGPooledBuffer* PooledBuffer, FRDGPassHandle PassHandle);
-	void SetRHI(FRDGBuffer* Buffer, FRHITransientBuffer* TransientBuffer, FRDGPassHandle PassHandle);
+	RENDERCORE_API void SetRHI(FRDGTexture* Texture, IPooledRenderTarget* RenderTarget, FRDGPassHandle PassHandle);
+	RENDERCORE_API void SetRHI(FRDGTexture* Texture, FRDGPooledTexture* PooledTexture, FRDGPassHandle PassHandle);
+	RENDERCORE_API void SetRHI(FRDGTexture* Texture, FRHITransientTexture* TransientTexture, FRDGPassHandle PassHandle);
+	RENDERCORE_API void SetRHI(FRDGBuffer* Buffer, FRDGPooledBuffer* PooledBuffer, FRDGPassHandle PassHandle);
+	RENDERCORE_API void SetRHI(FRDGBuffer* Buffer, FRHITransientBuffer* TransientBuffer, FRDGPassHandle PassHandle);
 
-	void BeginResourcesRHI(FRDGPass* ResourcePass, FRDGPassHandle ExecutePassHandle);
-	void BeginResourceRHI(FRDGPassHandle, FRDGTexture* Texture);
-	void BeginResourceRHI(FRDGPassHandle, FRDGBuffer* Buffer);
+	RENDERCORE_API void BeginResourcesRHI(FRDGPass* ResourcePass, FRDGPassHandle ExecutePassHandle);
+	RENDERCORE_API void BeginResourceRHI(FRDGPassHandle, FRDGTexture* Texture);
+	RENDERCORE_API void BeginResourceRHI(FRDGPassHandle, FRDGBuffer* Buffer);
 
-	void EndResourcesRHI(FRDGPass* ResourcePass, FRDGPassHandle ExecutePassHandle);
-	void EndResourceRHI(FRDGPassHandle, FRDGTexture* Texture, uint32 ReferenceCount);
-	void EndResourceRHI(FRDGPassHandle, FRDGBuffer* Buffer, uint32 ReferenceCount);
+	RENDERCORE_API void EndResourcesRHI(FRDGPass* ResourcePass, FRDGPassHandle ExecutePassHandle);
+	RENDERCORE_API void EndResourceRHI(FRDGPassHandle, FRDGTexture* Texture, uint32 ReferenceCount);
+	RENDERCORE_API void EndResourceRHI(FRDGPassHandle, FRDGBuffer* Buffer, uint32 ReferenceCount);
 
-	void InitRHI(FRDGView* View);
-	void InitRHI(FRDGBufferSRV* SRV);
-	void InitRHI(FRDGBufferUAV* UAV);
-	void InitRHI(FRDGTextureSRV* SRV);
-	void InitRHI(FRDGTextureUAV* UAV);
+	RENDERCORE_API void InitRHI(FRDGView* View);
+	RENDERCORE_API void InitRHI(FRDGBufferSRV* SRV);
+	RENDERCORE_API void InitRHI(FRDGBufferUAV* UAV);
+	RENDERCORE_API void InitRHI(FRDGTextureSRV* SRV);
+	RENDERCORE_API void InitRHI(FRDGTextureUAV* UAV);
 
-	void SetupParallelExecute();
-	void DispatchParallelExecute();
+	RENDERCORE_API void SetupParallelExecute();
+	RENDERCORE_API void DispatchParallelExecute();
 
-	void PrepareBufferUploads();
-	UE::Tasks::FTask SubmitBufferUploads();
-	void BeginFlushResourcesRHI();
-	void EndFlushResourcesRHI();
+	RENDERCORE_API void PrepareBufferUploads();
+	RENDERCORE_API UE::Tasks::FTask SubmitBufferUploads();
+	RENDERCORE_API void BeginFlushResourcesRHI();
+	RENDERCORE_API void EndFlushResourcesRHI();
 
-	void FlushAccessModeQueue();
+	RENDERCORE_API void FlushAccessModeQueue();
 
-	FRDGPass* SetupEmptyPass(FRDGPass* Pass);
-	FRDGPass* SetupParameterPass(FRDGPass* Pass);
+	RENDERCORE_API FRDGPass* SetupEmptyPass(FRDGPass* Pass);
+	RENDERCORE_API FRDGPass* SetupParameterPass(FRDGPass* Pass);
 
-	void SetupPassInternals(FRDGPass* Pass);
-	void SetupPassResources(FRDGPass* Pass);
-	void SetupAuxiliaryPasses(FRDGPass* Pass);
-	void SetupPassDependencies(FRDGPass* Pass);
+	RENDERCORE_API void SetupPassInternals(FRDGPass* Pass);
+	RENDERCORE_API void SetupPassResources(FRDGPass* Pass);
+	RENDERCORE_API void SetupAuxiliaryPasses(FRDGPass* Pass);
+	RENDERCORE_API void SetupPassDependencies(FRDGPass* Pass);
 
-	void CompilePassOps(FRDGPass* Pass);
-	void ExecutePass(FRDGPass* Pass, FRHIComputeCommandList& RHICmdListPass);
+	RENDERCORE_API void CompilePassOps(FRDGPass* Pass);
+	RENDERCORE_API void ExecutePass(FRDGPass* Pass, FRHIComputeCommandList& RHICmdListPass);
 
-	void ExecutePassPrologue(FRHIComputeCommandList& RHICmdListPass, FRDGPass* Pass);
-	void ExecutePassEpilogue(FRHIComputeCommandList& RHICmdListPass, FRDGPass* Pass);
+	RENDERCORE_API void ExecutePassPrologue(FRHIComputeCommandList& RHICmdListPass, FRDGPass* Pass);
+	RENDERCORE_API void ExecutePassEpilogue(FRHIComputeCommandList& RHICmdListPass, FRDGPass* Pass);
 
-	void CompilePassBarriers();
-	void CollectPassBarriers(FRDGPass* Pass);
-	void CreatePassBarriers(TFunctionRef<void()> PreWork);
+	RENDERCORE_API void CompilePassBarriers();
+	RENDERCORE_API void CollectPassBarriers(FRDGPass* Pass);
+	RENDERCORE_API void CreatePassBarriers(TFunctionRef<void()> PreWork);
 
-	UE::Tasks::FTask CreateUniformBuffers();
+	RENDERCORE_API UE::Tasks::FTask CreateUniformBuffers();
 
-	void AddCullingDependency(FRDGProducerStatesByPipeline& LastProducers, const FRDGProducerState& NextState, ERHIPipeline NextPipeline);
+	RENDERCORE_API void AddCullingDependency(FRDGProducerStatesByPipeline& LastProducers, const FRDGProducerState& NextState, ERHIPipeline NextPipeline);
 
-	void AddEpilogueTransition(FRDGTextureRef Texture);
-	void AddEpilogueTransition(FRDGBufferRef Buffer);
+	RENDERCORE_API void AddEpilogueTransition(FRDGTextureRef Texture);
+	RENDERCORE_API void AddEpilogueTransition(FRDGBufferRef Buffer);
 
-	void AddTransition(
+	RENDERCORE_API void AddTransition(
 		FRDGPassHandle PassHandle,
 		FRDGTextureRef Texture,
 		FRDGTextureSubresourceStateIndirect& StateAfter);
 
-	void AddTransition(
+	RENDERCORE_API void AddTransition(
 		FRDGPassHandle PassHandle,
 		FRDGBufferRef Buffer,
 		FRDGSubresourceState StateAfter);
 
-	void AddTransition(
+	RENDERCORE_API void AddTransition(
 		FRDGViewableResource* Resource,
 		FRDGSubresourceState StateBefore,
 		FRDGSubresourceState StateAfter,
 		const FRHITransitionInfo& TransitionInfo);
 
-	void AddAliasingTransition(
+	RENDERCORE_API void AddAliasingTransition(
 		FRDGPassHandle BeginPassHandle,
 		FRDGPassHandle EndPassHandle,
 		FRDGViewableResource* Resource,
 		const FRHITransientAliasingInfo& Info);
 
-	bool IsTransient(FRDGTextureRef Texture) const;
-	bool IsTransient(FRDGBufferRef Buffer) const;
-	bool IsTransientInternal(FRDGViewableResource* Resource, bool bFastVRAM) const;
+	RENDERCORE_API bool IsTransient(FRDGTextureRef Texture) const;
+	RENDERCORE_API bool IsTransient(FRDGBufferRef Buffer) const;
+	RENDERCORE_API bool IsTransientInternal(FRDGViewableResource* Resource, bool bFastVRAM) const;
 
-	FRHIRenderPassInfo GetRenderPassInfo(const FRDGPass* Pass) const;
+	RENDERCORE_API FRHIRenderPassInfo GetRenderPassInfo(const FRDGPass* Pass) const;
 
-	FRDGSubresourceState* AllocSubresource(const FRDGSubresourceState& Other);
+	RENDERCORE_API FRDGSubresourceState* AllocSubresource(const FRDGSubresourceState& Other);
 
 #if RDG_DUMP_RESOURCES
-	void DumpResourcePassOutputs(const FRDGPass* Pass);
+	RENDERCORE_API void DumpResourcePassOutputs(const FRDGPass* Pass);
 
 #if RDG_DUMP_RESOURCES_AT_EACH_DRAW
-	void BeginPassDump(const FRDGPass* Pass);
-	void EndPassDump(const FRDGPass* Pass);
+	RENDERCORE_API void BeginPassDump(const FRDGPass* Pass);
+	RENDERCORE_API void EndPassDump(const FRDGPass* Pass);
 #endif
 #endif
 
 #if RDG_ENABLE_DEBUG
-	void VisualizePassOutputs(const FRDGPass* Pass);
-	void ClobberPassOutputs(const FRDGPass* Pass);
+	RENDERCORE_API void VisualizePassOutputs(const FRDGPass* Pass);
+	RENDERCORE_API void ClobberPassOutputs(const FRDGPass* Pass);
 #endif
 
 	friend FRDGTrace;
diff --git a/Engine/Source/Runtime/RenderCore/Public/RenderGraphEvent.h b/Engine/Source/Runtime/RenderCore/Public/RenderGraphEvent.h
--- a/Engine/Source/Runtime/RenderCore/Public/RenderGraphEvent.h
+++ b/Engine/Source/Runtime/RenderCore/Public/RenderGraphEvent.h
@@ -308,19 +308,19 @@ public:
 
 using FRDGTimingScopeOp = TRDGScopeOp<FRDGTimingScope>;
 
-class RENDERCORE_API FRDGTimingScopeOpArray final
+class FRDGTimingScopeOpArray final
 {
 public:
 	FRDGTimingScopeOpArray() = default;
 	FRDGTimingScopeOpArray(ERHIPipeline Pipeline, const TRDGScopeOpArray<FRDGTimingScopeOp>& Ops);
 
-	void Execute(FRHIComputeCommandList& RHICmdList);
+	RENDERCORE_API void Execute(FRHIComputeCommandList& RHICmdList);
 
 private:
 	FRHIRenderQuery* TimestampQuery = nullptr;
 };
 
-class RENDERCORE_API FRDGTimingScopeStack final
+class FRDGTimingScopeStack final
 {
 public:
 	FRDGTimingScopeStack(FRDGAllocator& Allocator)
@@ -343,7 +343,7 @@ public:
 		ScopeStack.ReserveOps();
 	}
 
-	FRDGTimingScopeOpArray CompilePassPrologue(const FRDGPass* Pass);
+	RENDERCORE_API FRDGTimingScopeOpArray CompilePassPrologue(const FRDGPass* Pass);
 	
 	inline void EndExecute(FRHIComputeCommandList& RHICmdList)
 	{
@@ -361,12 +361,12 @@ public:
 namespace DynamicRenderScaling
 {
 
-class RENDERCORE_API FRDGScope final
+class FRDGScope final
 {
 public:
-	FRDGScope(FRDGBuilder& InGraphBuilder, const FBudget& InBudget);
+	RENDERCORE_API FRDGScope(FRDGBuilder& InGraphBuilder, const FBudget& InBudget);
 	FRDGScope(const FRDGScope&) = delete;
-	~FRDGScope();
+	RENDERCORE_API ~FRDGScope();
 
 private:
 	FRDGBuilder& GraphBuilder;
@@ -386,19 +386,19 @@ private:
 /** Stores a GPU event name for the render graph. Draw events can be compiled out entirely from
  *  a release build for performance.
  */
-class RENDERCORE_API FRDGEventName final
+class FRDGEventName final
 {
 public:
 	FRDGEventName() = default;
 
-	explicit FRDGEventName(const TCHAR* EventFormat, ...);
+	RENDERCORE_API explicit FRDGEventName(const TCHAR* EventFormat, ...);
 
 	FRDGEventName(const FRDGEventName& Other);
 	FRDGEventName(FRDGEventName&& Other);
-	FRDGEventName& operator=(const FRDGEventName& Other);
-	FRDGEventName& operator=(FRDGEventName&& Other);
+	RENDERCORE_API FRDGEventName& operator=(const FRDGEventName& Other);
+	RENDERCORE_API FRDGEventName& operator=(FRDGEventName&& Other);
 
-	const TCHAR* GetTCHAR() const;
+	RENDERCORE_API const TCHAR* GetTCHAR() const;
 
 private:
 #if RDG_EVENTS == RDG_EVENTS_STRING_REF || RDG_EVENTS == RDG_EVENTS_STRING_COPY
@@ -462,7 +462,7 @@ struct FRDGBreadcrumbState : TRHIBreadcrumbState<FRDGArrayAllocator>
 };
 #endif
 
-class RENDERCORE_API FRDGEventScopeOpArray final
+class FRDGEventScopeOpArray final
 {
 public:
 	FRDGEventScopeOpArray(bool bInRDGEvents = true)
@@ -474,10 +474,10 @@ public:
 		, bRDGEvents(bInRDGEvents)
 	{}
 
-	void Execute(FRHIComputeCommandList& RHICmdList);
+	RENDERCORE_API void Execute(FRHIComputeCommandList& RHICmdList);
 
 #if RHI_WANT_BREADCRUMB_EVENTS
-	void Execute(FRDGBreadcrumbState& State);
+	RENDERCORE_API void Execute(FRDGBreadcrumbState& State);
 #endif
 
 	TRDGScopeOpArray<FRDGEventScopeOp> Ops;
@@ -487,7 +487,7 @@ public:
 /** Manages a stack of event scopes. Scopes are recorded ahead of time in a hierarchical fashion
  *  and later executed topologically during pass execution.
  */
-class RENDERCORE_API FRDGEventScopeStack final
+class FRDGEventScopeStack final
 {
 public:
 	FRDGEventScopeStack(FRDGAllocator& Allocator)
@@ -519,9 +519,9 @@ public:
 		}
 	}
 
-	FRDGEventScopeOpArray CompilePassPrologue(const FRDGPass* Pass);
+	RENDERCORE_API FRDGEventScopeOpArray CompilePassPrologue(const FRDGPass* Pass);
 	
-	FRDGEventScopeOpArray CompilePassEpilogue();
+	RENDERCORE_API FRDGEventScopeOpArray CompilePassEpilogue();
 
 	inline void EndExecute(FRHIComputeCommandList& RHICmdList, ERHIPipeline Pipeline)
 	{
@@ -560,12 +560,12 @@ private:
 
 RENDERCORE_API FString GetRDGEventPath(const FRDGEventScope* Scope, const FRDGEventName& Event);
 
-class RENDERCORE_API FRDGEventScopeGuard final
+class FRDGEventScopeGuard final
 {
 public:
-	FRDGEventScopeGuard(FRDGBuilder& InGraphBuilder, FRDGEventName&& ScopeName, bool bCondition = true, ERDGEventScopeFlags Flags = ERDGEventScopeFlags::None);
+	RENDERCORE_API FRDGEventScopeGuard(FRDGBuilder& InGraphBuilder, FRDGEventName&& ScopeName, bool bCondition = true, ERDGEventScopeFlags Flags = ERDGEventScopeFlags::None);
 	FRDGEventScopeGuard(const FRDGEventScopeGuard&) = delete;
-	~FRDGEventScopeGuard();
+	RENDERCORE_API ~FRDGEventScopeGuard();
 
 private:
 	FRDGBuilder& GraphBuilder;
@@ -618,7 +618,7 @@ public:
 #endif
 };
 
-class RENDERCORE_API FRDGGPUStatScopeOpArray final
+class FRDGGPUStatScopeOpArray final
 {
 public:
 	static const int32 kInvalidEventIndex = -1;
@@ -632,14 +632,14 @@ public:
 	FRDGGPUStatScopeOpArray() = default;
 	FRDGGPUStatScopeOpArray(TRDGScopeOpArray<FRDGGPUStatScopeOp> InOps, FRHIGPUMask GPUMask);
 
-	void Execute(FRHIComputeCommandList& RHICmdList);
+	RENDERCORE_API void Execute(FRHIComputeCommandList& RHICmdList);
 
 	TRDGScopeOpArray<FRDGGPUStatScopeOp> Ops;
 	int32 OverrideEventIndex = kInvalidEventIndex;
 	EType Type = EType::Epilogue;
 };
 
-class RENDERCORE_API FRDGGPUStatScopeStack final
+class FRDGGPUStatScopeStack final
 {
 public:
 	FRDGGPUStatScopeStack(FRDGAllocator& Allocator)
@@ -673,9 +673,9 @@ public:
 		}
 	}
 
-	FRDGGPUStatScopeOpArray CompilePassPrologue(const FRDGPass* Pass, FRHIGPUMask GPUMask);
+	RENDERCORE_API FRDGGPUStatScopeOpArray CompilePassPrologue(const FRDGPass* Pass, FRHIGPUMask GPUMask);
 
-	FRDGGPUStatScopeOpArray CompilePassEpilogue();
+	RENDERCORE_API FRDGGPUStatScopeOpArray CompilePassEpilogue();
 
 	inline void EndExecute(FRHIComputeCommandList& RHICmdList, ERHIPipeline Pipeline)
 	{
@@ -707,12 +707,12 @@ private:
 	bool bGPUStats = false;
 };
 
-class RENDERCORE_API FRDGGPUStatScopeGuard final
+class FRDGGPUStatScopeGuard final
 {
 public:
-	FRDGGPUStatScopeGuard(FRDGBuilder& InGraphBuilder, const FName& Name, const FName& StatName, const TCHAR* Description, FDrawCallCategoryName& InCategory);
+	RENDERCORE_API FRDGGPUStatScopeGuard(FRDGBuilder& InGraphBuilder, const FName& Name, const FName& StatName, const TCHAR* Description, FDrawCallCategoryName& InCategory);
 	FRDGGPUStatScopeGuard(const FRDGGPUStatScopeGuard&) = delete;
-	~FRDGGPUStatScopeGuard();
+	RENDERCORE_API ~FRDGGPUStatScopeGuard();
 
 private:
 	FRDGBuilder& GraphBuilder;
@@ -803,7 +803,7 @@ struct FRDGGPUScopeStacks
 	IF_RDG_GPU_DEBUG_SCOPES(FRDGGPUStatScopeStack Stat);
 };
 
-struct RENDERCORE_API FRDGGPUScopeStacksByPipeline
+struct FRDGGPUScopeStacksByPipeline
 {
 	FRDGGPUScopeStacksByPipeline(FRDGAllocator& Allocator)
 		: Graphics(Allocator)
@@ -864,13 +864,13 @@ struct RENDERCORE_API FRDGGPUScopeStacksByPipeline
 		AsyncCompute.ReserveOps(PassCount);
 	}
 
-	FRDGGPUScopeOpArrays CompilePassPrologue(const FRDGPass* Pass, FRHIGPUMask GPUMask);
+	RENDERCORE_API FRDGGPUScopeOpArrays CompilePassPrologue(const FRDGPass* Pass, FRHIGPUMask GPUMask);
 
-	FRDGGPUScopeOpArrays CompilePassEpilogue(const FRDGPass* Pass);
+	RENDERCORE_API FRDGGPUScopeOpArrays CompilePassEpilogue(const FRDGPass* Pass);
 
-	const FRDGGPUScopeStacks& GetScopeStacks(ERHIPipeline Pipeline) const;
+	RENDERCORE_API const FRDGGPUScopeStacks& GetScopeStacks(ERHIPipeline Pipeline) const;
 
-	FRDGGPUScopeStacks& GetScopeStacks(ERHIPipeline Pipeline);
+	RENDERCORE_API FRDGGPUScopeStacks& GetScopeStacks(ERHIPipeline Pipeline);
 
 	FRDGGPUScopes GetCurrentScopes(ERHIPipeline Pipeline) const
 	{
@@ -906,7 +906,7 @@ public:
 
 using FRDGCSVStatScopeOp = TRDGScopeOp<FRDGCSVStatScope>;
 
-class RENDERCORE_API FRDGCSVStatScopeOpArray final
+class FRDGCSVStatScopeOpArray final
 {
 public:
 	FRDGCSVStatScopeOpArray() = default;
@@ -914,12 +914,12 @@ public:
 		: Ops(InOps)
 	{}
 
-	void Execute();
+	RENDERCORE_API void Execute();
 
 	TRDGScopeOpArray<FRDGCSVStatScopeOp> Ops;
 };
 
-class RENDERCORE_API FRDGCSVStatScopeStack final
+class FRDGCSVStatScopeStack final
 {
 public:
 	FRDGCSVStatScopeStack(FRDGAllocator& Allocator)
@@ -950,7 +950,7 @@ public:
 		}
 	}
 
-	FRDGCSVStatScopeOpArray CompilePassPrologue(const FRDGPass* Pass);
+	RENDERCORE_API FRDGCSVStatScopeOpArray CompilePassPrologue(const FRDGPass* Pass);
 
 	void EndExecute()
 	{
@@ -979,21 +979,21 @@ private:
 
 #if CSV_PROFILER
 
-class RENDERCORE_API FRDGScopedCsvStatExclusive : public FScopedCsvStatExclusive
+class FRDGScopedCsvStatExclusive : public FScopedCsvStatExclusive
 {
 public:
-	FRDGScopedCsvStatExclusive(FRDGBuilder& InGraphBuilder, const char* InStatName);
-	~FRDGScopedCsvStatExclusive();
+	RENDERCORE_API FRDGScopedCsvStatExclusive(FRDGBuilder& InGraphBuilder, const char* InStatName);
+	RENDERCORE_API ~FRDGScopedCsvStatExclusive();
 
 private:
 	FRDGBuilder& GraphBuilder;
 };
 
-class RENDERCORE_API FRDGScopedCsvStatExclusiveConditional : public FScopedCsvStatExclusiveConditional
+class FRDGScopedCsvStatExclusiveConditional : public FScopedCsvStatExclusiveConditional
 {
 public:
-	FRDGScopedCsvStatExclusiveConditional(FRDGBuilder& InGraphBuilder, const char* InStatName, bool bInCondition);
-	~FRDGScopedCsvStatExclusiveConditional();
+	RENDERCORE_API FRDGScopedCsvStatExclusiveConditional(FRDGBuilder& InGraphBuilder, const char* InStatName, bool bInCondition);
+	RENDERCORE_API ~FRDGScopedCsvStatExclusiveConditional();
 
 private:
 	FRDGBuilder& GraphBuilder;
diff --git a/Engine/Source/Runtime/RenderCore/Public/RenderGraphParameter.h b/Engine/Source/Runtime/RenderCore/Public/RenderGraphParameter.h
--- a/Engine/Source/Runtime/RenderCore/Public/RenderGraphParameter.h
+++ b/Engine/Source/Runtime/RenderCore/Public/RenderGraphParameter.h
@@ -209,7 +209,7 @@ private:
 };
 
 /** Wraps a pass parameter struct payload and provides helpers for traversing members. */
-class RENDERCORE_API FRDGParameterStruct
+class FRDGParameterStruct
 {
 public:
 	template <typename ParameterStructType>
@@ -276,13 +276,13 @@ public:
 	void EnumerateUniformBuffers(FunctionType Function) const;
 
 	/** Returns a set of static uniform buffer bindings for the parameter struct. */
-	FUniformBufferStaticBindings GetStaticUniformBuffers() const;
+	RENDERCORE_API FUniformBufferStaticBindings GetStaticUniformBuffers() const;
 
 	/** Returns the render pass info generated from the render target binding slots. */
-	FRHIRenderPassInfo GetRenderPassInfo() const;
+	RENDERCORE_API FRHIRenderPassInfo GetRenderPassInfo() const;
 
 	/** Clears out all uniform buffer references in the parameter struct. */
-	static void ClearUniformBuffers(void* Contents, const FRHIUniformBufferLayout* Layout);
+	static RENDERCORE_API void ClearUniformBuffers(void* Contents, const FRHIUniformBufferLayout* Layout);
 
 private:
 	FRDGParameter GetParameterInternal(TArrayView<const FRHIUniformBufferResource> Parameters, uint32 ParameterIndex) const
@@ -343,4 +343,4 @@ template <typename TParameterStruct>
 FORCEINLINE static FUniformBufferStaticBindings GetStaticUniformBuffers(TParameterStruct* Parameters)
 {
 	return FRDGParameterStruct(Parameters, TParameterStruct::FTypeInfo::GetStructMetadata()).GetStaticUniformBuffers();
-}
\ No newline at end of file
+}
diff --git a/Engine/Source/Runtime/RenderCore/Public/RenderGraphPass.h b/Engine/Source/Runtime/RenderCore/Public/RenderGraphPass.h
--- a/Engine/Source/Runtime/RenderCore/Public/RenderGraphPass.h
+++ b/Engine/Source/Runtime/RenderCore/Public/RenderGraphPass.h
@@ -89,15 +89,15 @@ struct FRDGBarrierBatchEndId
 	ERDGBarrierLocation BarrierLocation = ERDGBarrierLocation::Epilogue;
 };
 
-class RENDERCORE_API FRDGBarrierBatchBegin
+class FRDGBarrierBatchBegin
 {
 public:
-	FRDGBarrierBatchBegin(ERHIPipeline PipelinesToBegin, ERHIPipeline PipelinesToEnd, const TCHAR* DebugName, FRDGPass* DebugPass);
-	FRDGBarrierBatchBegin(ERHIPipeline PipelinesToBegin, ERHIPipeline PipelinesToEnd, const TCHAR* DebugName, FRDGPassesByPipeline DebugPasses);
+	RENDERCORE_API FRDGBarrierBatchBegin(ERHIPipeline PipelinesToBegin, ERHIPipeline PipelinesToEnd, const TCHAR* DebugName, FRDGPass* DebugPass);
+	RENDERCORE_API FRDGBarrierBatchBegin(ERHIPipeline PipelinesToBegin, ERHIPipeline PipelinesToEnd, const TCHAR* DebugName, FRDGPassesByPipeline DebugPasses);
 
-	void AddTransition(FRDGViewableResource* Resource, const FRHITransitionInfo& Info);
+	RENDERCORE_API void AddTransition(FRDGViewableResource* Resource, const FRHITransitionInfo& Info);
 
-	void AddAlias(FRDGViewableResource* Resource, const FRHITransientAliasingInfo& Info);
+	RENDERCORE_API void AddAlias(FRDGViewableResource* Resource, const FRHITransientAliasingInfo& Info);
 
 	void SetUseCrossPipelineFence()
 	{
@@ -105,10 +105,10 @@ public:
 		bTransitionNeeded = true;
 	}
 
-	void CreateTransition();
+	RENDERCORE_API void CreateTransition();
 
-	void Submit(FRHIComputeCommandList& RHICmdList, ERHIPipeline Pipeline);
-	void Submit(FRHIComputeCommandList& RHICmdList, ERHIPipeline Pipeline, FRDGTransitionQueue& TransitionsToBegin);
+	RENDERCORE_API void Submit(FRHIComputeCommandList& RHICmdList, ERHIPipeline Pipeline);
+	RENDERCORE_API void Submit(FRHIComputeCommandList& RHICmdList, ERHIPipeline Pipeline, FRDGTransitionQueue& TransitionsToBegin);
 
 	void Reserve(uint32 TransitionCount)
 	{
@@ -145,7 +145,7 @@ private:
 
 using FRDGTransitionCreateQueue = TArray<FRDGBarrierBatchBegin*, FRDGArrayAllocator>;
 
-class RENDERCORE_API FRDGBarrierBatchEnd
+class FRDGBarrierBatchEnd
 {
 public:
 	FRDGBarrierBatchEnd(FRDGPass* InPass, ERDGBarrierLocation InBarrierLocation)
@@ -154,18 +154,18 @@ public:
 	{}
 
 	/** Inserts a dependency on a begin batch. A begin batch can be inserted into more than one end batch. */
-	void AddDependency(FRDGBarrierBatchBegin* BeginBatch);
+	RENDERCORE_API void AddDependency(FRDGBarrierBatchBegin* BeginBatch);
 
-	void Submit(FRHIComputeCommandList& RHICmdList, ERHIPipeline Pipeline);
+	RENDERCORE_API void Submit(FRHIComputeCommandList& RHICmdList, ERHIPipeline Pipeline);
 
 	void Reserve(uint32 TransitionBatchCount)
 	{
 		Dependencies.Reserve(TransitionBatchCount);
 	}
 
-	FRDGBarrierBatchEndId GetId() const;
+	RENDERCORE_API FRDGBarrierBatchEndId GetId() const;
 
-	bool IsPairedWith(const FRDGBarrierBatchBegin& BeginBatch) const;
+	RENDERCORE_API bool IsPairedWith(const FRDGBarrierBatchBegin& BeginBatch) const;
 
 private:
 	TArray<FRDGBarrierBatchBegin*, TInlineAllocator<4, FRDGArrayAllocator>> Dependencies;
@@ -177,15 +177,15 @@ private:
 };
 
 /** Base class of a render graph pass. */
-class RENDERCORE_API FRDGPass
+class FRDGPass
 {
 public:
-	FRDGPass(FRDGEventName&& InName, FRDGParameterStruct InParameterStruct, ERDGPassFlags InFlags);
+	RENDERCORE_API FRDGPass(FRDGEventName&& InName, FRDGParameterStruct InParameterStruct, ERDGPassFlags InFlags);
 	FRDGPass(const FRDGPass&) = delete;
 	virtual ~FRDGPass() = default;
 
 #if RDG_ENABLE_DEBUG
-	const TCHAR* GetName() const;
+	RENDERCORE_API const TCHAR* GetName() const;
 #else
 	FORCEINLINE const TCHAR* GetName() const
 	{
@@ -320,10 +320,10 @@ public:
 #endif
 
 protected:
-	FRDGBarrierBatchBegin& GetPrologueBarriersToBegin(FRDGAllocator& Allocator, FRDGTransitionCreateQueue& CreateQueue);
-	FRDGBarrierBatchBegin& GetEpilogueBarriersToBeginForGraphics(FRDGAllocator& Allocator, FRDGTransitionCreateQueue& CreateQueue);
-	FRDGBarrierBatchBegin& GetEpilogueBarriersToBeginForAsyncCompute(FRDGAllocator& Allocator, FRDGTransitionCreateQueue& CreateQueue);
-	FRDGBarrierBatchBegin& GetEpilogueBarriersToBeginForAll(FRDGAllocator& Allocator, FRDGTransitionCreateQueue& CreateQueue);
+	RENDERCORE_API FRDGBarrierBatchBegin& GetPrologueBarriersToBegin(FRDGAllocator& Allocator, FRDGTransitionCreateQueue& CreateQueue);
+	RENDERCORE_API FRDGBarrierBatchBegin& GetEpilogueBarriersToBeginForGraphics(FRDGAllocator& Allocator, FRDGTransitionCreateQueue& CreateQueue);
+	RENDERCORE_API FRDGBarrierBatchBegin& GetEpilogueBarriersToBeginForAsyncCompute(FRDGAllocator& Allocator, FRDGTransitionCreateQueue& CreateQueue);
+	RENDERCORE_API FRDGBarrierBatchBegin& GetEpilogueBarriersToBeginForAll(FRDGAllocator& Allocator, FRDGTransitionCreateQueue& CreateQueue);
 
 	FRDGBarrierBatchBegin& GetEpilogueBarriersToBeginFor(FRDGAllocator& Allocator, FRDGTransitionCreateQueue& CreateQueue, ERHIPipeline PipelineForEnd)
 	{
@@ -344,8 +344,8 @@ protected:
 		}
 	}
 
-	FRDGBarrierBatchEnd& GetPrologueBarriersToEnd(FRDGAllocator& Allocator);
-	FRDGBarrierBatchEnd& GetEpilogueBarriersToEnd(FRDGAllocator& Allocator);
+	RENDERCORE_API FRDGBarrierBatchEnd& GetPrologueBarriersToEnd(FRDGAllocator& Allocator);
+	RENDERCORE_API FRDGBarrierBatchEnd& GetEpilogueBarriersToEnd(FRDGAllocator& Allocator);
 
 	virtual void Execute(FRHIComputeCommandList& RHICmdList) {}
 
diff --git a/Engine/Source/Runtime/RenderCore/Public/RenderGraphResources.h b/Engine/Source/Runtime/RenderCore/Public/RenderGraphResources.h
--- a/Engine/Source/Runtime/RenderCore/Public/RenderGraphResources.h
+++ b/Engine/Source/Runtime/RenderCore/Public/RenderGraphResources.h
@@ -111,7 +111,7 @@ using FRDGTextureSubresourceState = TRDGTextureSubresourceArray<FRDGSubresourceS
 using FRDGTextureSubresourceStateIndirect = TRDGTextureSubresourceArray<FRDGSubresourceState*, FRDGArrayAllocator>;
 
 /** Generic graph resource. */
-class RENDERCORE_API FRDGResource
+class FRDGResource
 {
 public:
 	FRDGResource(const FRDGResource&) = delete;
@@ -125,7 +125,7 @@ public:
 
 	/** Marks this resource as actually used by a resource. This is to track what dependencies on pass was actually unnecessary. */
 #if RDG_ENABLE_DEBUG
-	virtual void MarkResourceAsUsed();
+	RENDERCORE_API virtual void MarkResourceAsUsed();
 #else
 	inline  void MarkResourceAsUsed() {}
 #endif
@@ -163,13 +163,13 @@ protected:
 	FRHIResource* ResourceRHI = nullptr;
 
 #if RDG_ENABLE_DEBUG
-	void ValidateRHIAccess() const;
+	RENDERCORE_API void ValidateRHIAccess() const;
 #endif
 
 private:
 #if RDG_ENABLE_DEBUG
 	struct FRDGResourceDebugData* DebugData = nullptr;
-	FRDGResourceDebugData& GetDebugData() const;
+	RENDERCORE_API FRDGResourceDebugData& GetDebugData() const;
 #endif
 
 #if RHI_ENABLE_RESOURCE_INFO
@@ -181,7 +181,7 @@ private:
 	friend FRDGBarrierValidation;
 };
 
-class RENDERCORE_API FRDGUniformBuffer
+class FRDGUniformBuffer
 	: public FRDGResource
 {
 public:
@@ -199,7 +199,7 @@ public:
 	}
 
 #if RDG_ENABLE_DEBUG
-	void MarkResourceAsUsed() override;
+	RENDERCORE_API void MarkResourceAsUsed() override;
 #else
 	inline void MarkResourceAsUsed() {}
 #endif
@@ -227,7 +227,7 @@ private:
 		return static_cast<FRHIUniformBuffer*>(FRDGResource::GetRHIUnchecked());
 	}
 
-	void InitRHI();
+	RENDERCORE_API void InitRHI();
 
 	const FRDGParameterStruct ParameterStruct;
 	TRefCountPtr<FRHIUniformBuffer> UniformBufferRHI;
@@ -280,7 +280,7 @@ private:
 };
 
 /** A render graph resource with an allocation lifetime tracked by the graph. May have child resources which reference it (e.g. views). */
-class RENDERCORE_API FRDGViewableResource
+class FRDGViewableResource
 	: public FRDGResource
 {
 public:
@@ -313,7 +313,7 @@ public:
 	}
 
 protected:
-	FRDGViewableResource(const TCHAR* InName, ERDGViewableResourceType InType, bool bSkipTracking, bool bImmediateFirstBarrier);
+	RENDERCORE_API FRDGViewableResource(const TCHAR* InName, ERDGViewableResourceType InType, bool bSkipTracking, bool bImmediateFirstBarrier);
 
 	static const ERHIAccess DefaultEpilogueAccess = ERHIAccess::SRVMask;
 
@@ -419,7 +419,7 @@ private:
 
 #if RDG_ENABLE_DEBUG
 	struct FRDGViewableResourceDebugData* ViewableDebugData = nullptr;
-	FRDGViewableResourceDebugData& GetViewableDebugData() const;
+	RENDERCORE_API FRDGViewableResourceDebugData& GetViewableDebugData() const;
 #endif
 
 	friend FRDGBuilder;
@@ -472,7 +472,7 @@ inline FRDGTextureDesc Translate(const FPooledRenderTargetDesc& InDesc);
 /** Translates from an RHI/RDG texture descriptor to a pooled render target descriptor. */
 inline FPooledRenderTargetDesc Translate(const FRHITextureDesc& InDesc);
 
-class RENDERCORE_API FRDGPooledTexture final
+class FRDGPooledTexture final
 	: public FRefCountBase
 {
 public:
@@ -496,7 +496,7 @@ private:
 };
 
 /** Render graph tracked Texture. */
-class RENDERCORE_API FRDGTexture final
+class FRDGTexture final
 	: public FRDGViewableResource
 {
 public:
@@ -541,7 +541,7 @@ public:
 		return Layout.GetSubresource(SubresourceIndex);
 	}
 
-	FRDGTextureSubresourceRange GetSubresourceRangeSRV() const;
+	RENDERCORE_API FRDGTextureSubresourceRange GetSubresourceRangeSRV() const;
 
 private:
 	FRDGTexture(const TCHAR* InName, const FRDGTextureDesc& InDesc, ERDGTextureFlags InFlags)
@@ -616,7 +616,7 @@ private:
 
 #if RDG_ENABLE_DEBUG
 	struct FRDGTextureDebugData* TextureDebugData = nullptr;
-	FRDGTextureDebugData& GetTextureDebugData() const;
+	RENDERCORE_API FRDGTextureDebugData& GetTextureDebugData() const;
 #endif
 
 	friend FRDGBuilder;
@@ -1116,7 +1116,7 @@ struct FRDGBufferUAVDesc final
 /** Translates from a RDG buffer descriptor to a RHI buffer creation info */
 inline FRHIBufferCreateInfo Translate(const FRDGBufferDesc& InDesc);
 
-class RENDERCORE_API FRDGPooledBuffer final
+class FRDGPooledBuffer final
 	: public FRefCountBase
 {
 public:
@@ -1192,7 +1192,7 @@ private:
 };
 
 /** A render graph tracked buffer. */
-class RENDERCORE_API FRDGBuffer final
+class FRDGBuffer final
 	: public FRDGViewableResource
 {
 public:
@@ -1303,7 +1303,7 @@ private:
 
 #if RDG_ENABLE_DEBUG
 	struct FRDGBufferDebugData* BufferDebugData = nullptr;
-	FRDGBufferDebugData& GetBufferDebugData() const;
+	RENDERCORE_API FRDGBufferDebugData& GetBufferDebugData() const;
 #endif
 
 	friend FRDGBuilder;
diff --git a/Engine/Source/Runtime/RenderCore/Public/RenderGraphTrace.h b/Engine/Source/Runtime/RenderCore/Public/RenderGraphTrace.h
--- a/Engine/Source/Runtime/RenderCore/Public/RenderGraphTrace.h
+++ b/Engine/Source/Runtime/RenderCore/Public/RenderGraphTrace.h
@@ -18,21 +18,21 @@ namespace UE { namespace Trace { class FChannel; } }
 
 UE_TRACE_CHANNEL_EXTERN(RDGChannel, RENDERCORE_API);
 
-class RENDERCORE_API FRDGTrace
+class FRDGTrace
 {
 public:
-	FRDGTrace();
+	RENDERCORE_API FRDGTrace();
 
-	void OutputGraphBegin();
-	void OutputGraphEnd(const FRDGBuilder& GraphBuilder);
+	RENDERCORE_API void OutputGraphBegin();
+	RENDERCORE_API void OutputGraphEnd(const FRDGBuilder& GraphBuilder);
 
-	void AddResource(FRDGViewableResource* Resource);
-	void AddTexturePassDependency(FRDGTexture* Texture, FRDGPass* Pass);
-	void AddBufferPassDependency(FRDGBuffer* Buffer, FRDGPass* Pass);
+	RENDERCORE_API void AddResource(FRDGViewableResource* Resource);
+	RENDERCORE_API void AddTexturePassDependency(FRDGTexture* Texture, FRDGPass* Pass);
+	RENDERCORE_API void AddBufferPassDependency(FRDGBuffer* Buffer, FRDGPass* Pass);
 
 	FRHITransientAllocationStats TransientAllocationStats;
 
-	bool IsEnabled() const;
+	RENDERCORE_API bool IsEnabled() const;
 
 private:
 	uint64 GraphStartCycles{};
@@ -40,4 +40,4 @@ private:
 	bool bEnabled;
 };
 
-#endif
\ No newline at end of file
+#endif
diff --git a/Engine/Source/Runtime/RenderCore/Public/RenderGraphUtils.h b/Engine/Source/Runtime/RenderCore/Public/RenderGraphUtils.h
--- a/Engine/Source/Runtime/RenderCore/Public/RenderGraphUtils.h
+++ b/Engine/Source/Runtime/RenderCore/Public/RenderGraphUtils.h
@@ -338,7 +338,7 @@ inline FRDGTextureMSAA CreateTextureMSAA(
 
 /** All utils for compute shaders.
  */
-struct RENDERCORE_API FComputeShaderUtils
+struct FComputeShaderUtils
 {
 	/** Ideal size of group size 8x8 to occupy at least an entire wave on GCN, two warp on Nvidia. */
 	static constexpr int32 kGolden2DGroupSize = 8;
@@ -641,8 +641,8 @@ struct RENDERCORE_API FComputeShaderUtils
 		return AddPass(GraphBuilder, Forward<FRDGEventName>(PassName), ERDGPassFlags::Compute, ComputeShader, Parameters, IndirectArgsBuffer, IndirectArgsOffset, MoveTemp(DispatchLateParamCallback));
 	}
 
-	static void ClearUAV(FRDGBuilder& GraphBuilder, FGlobalShaderMap* ShaderMap, FRDGBufferUAVRef UAV, uint32 ClearValue);
-	static void ClearUAV(FRDGBuilder& GraphBuilder, FGlobalShaderMap* ShaderMap, FRDGBufferUAVRef UAV, FVector4f ClearValue);
+	static RENDERCORE_API void ClearUAV(FRDGBuilder& GraphBuilder, FGlobalShaderMap* ShaderMap, FRDGBufferUAVRef UAV, uint32 ClearValue);
+	static RENDERCORE_API void ClearUAV(FRDGBuilder& GraphBuilder, FGlobalShaderMap* ShaderMap, FRDGBufferUAVRef UAV, FVector4f ClearValue);
 
 	static inline void ValidateGroupCount(const FIntVector& GroupCount)
 	{
@@ -676,7 +676,7 @@ struct RENDERCORE_API FComputeShaderUtils
 	 * 	Sets up a group count as (InputCountBuffer[InputCountOffset] * Multiplier + Divisor - 1U) / Divisor;
 	 *  Commonly use Divisor <=> number of threads per group.
 	 */
-	static FRDGBufferRef AddIndirectArgsSetupCsPass1D(FRDGBuilder& GraphBuilder, ERHIFeatureLevel::Type FeatureLevel, FRDGBufferRef& InputCountBuffer, const TCHAR* OutputBufferName, uint32 Divisor, uint32 InputCountOffset = 0U, uint32 Multiplier = 1U);
+	static RENDERCORE_API FRDGBufferRef AddIndirectArgsSetupCsPass1D(FRDGBuilder& GraphBuilder, ERHIFeatureLevel::Type FeatureLevel, FRDGBufferRef& InputCountBuffer, const TCHAR* OutputBufferName, uint32 Divisor, uint32 InputCountOffset = 0U, uint32 Multiplier = 1U);
 };
 
 /** Adds a render graph pass to copy a region from one texture to another. Uses RHICopyTexture under the hood.
@@ -1066,7 +1066,7 @@ void AddReadbackBufferPass(FRDGBuilder& GraphBuilder, FRDGEventName&& Name, FRDG
 }
 
 /** Batches up RDG external resource access mode requests and submits them all at once to RDG. */
-class RENDERCORE_API FRDGExternalAccessQueue
+class FRDGExternalAccessQueue
 {
 public:
 	FRDGExternalAccessQueue() = default;
@@ -1109,7 +1109,7 @@ public:
 		Resources.Emplace(Resource, Access, Pipelines);
 	}
 
-	void Submit(FRDGBuilder& GraphBuilder);
+	RENDERCORE_API void Submit(FRDGBuilder& GraphBuilder);
 
 	bool Contains(FRDGViewableResource* Resource)
 	{
diff --git a/Engine/Source/Runtime/RenderCore/Public/RenderGraphValidation.h b/Engine/Source/Runtime/RenderCore/Public/RenderGraphValidation.h
--- a/Engine/Source/Runtime/RenderCore/Public/RenderGraphValidation.h
+++ b/Engine/Source/Runtime/RenderCore/Public/RenderGraphValidation.h
@@ -37,93 +37,93 @@ template <typename ReferencedType> class TRefCountPtr;
  *  and modular location to extend the validation layer as well as clearly separate the graph implementation
  *  details from events in the graph.
  */
-class RENDERCORE_API FRDGUserValidation final
+class FRDGUserValidation final
 {
 public:
 	FRDGUserValidation(FRDGAllocator& Allocator, bool bParallelExecuteEnabled);
 	FRDGUserValidation(const FRDGUserValidation&) = delete;
-	~FRDGUserValidation();
+	RENDERCORE_API ~FRDGUserValidation();
 
 	/** Tracks and validates inputs into resource creation functions. */
-	void ValidateCreateTexture(const FRDGTextureDesc& Desc, const TCHAR* Name, ERDGTextureFlags Flags);
-	void ValidateCreateBuffer(const FRDGBufferDesc& Desc, const TCHAR* Name, ERDGBufferFlags Flags);
-	void ValidateCreateSRV(const FRDGTextureSRVDesc& Desc);
-	void ValidateCreateSRV(const FRDGBufferSRVDesc& Desc);
-	void ValidateCreateUAV(const FRDGTextureUAVDesc& Desc);
-	void ValidateCreateUAV(const FRDGBufferUAVDesc& Desc);
-	void ValidateCreateUniformBuffer(const void* ParameterStruct, const FShaderParametersMetadata* Metadata);
+	RENDERCORE_API void ValidateCreateTexture(const FRDGTextureDesc& Desc, const TCHAR* Name, ERDGTextureFlags Flags);
+	RENDERCORE_API void ValidateCreateBuffer(const FRDGBufferDesc& Desc, const TCHAR* Name, ERDGBufferFlags Flags);
+	RENDERCORE_API void ValidateCreateSRV(const FRDGTextureSRVDesc& Desc);
+	RENDERCORE_API void ValidateCreateSRV(const FRDGBufferSRVDesc& Desc);
+	RENDERCORE_API void ValidateCreateUAV(const FRDGTextureUAVDesc& Desc);
+	RENDERCORE_API void ValidateCreateUAV(const FRDGBufferUAVDesc& Desc);
+	RENDERCORE_API void ValidateCreateUniformBuffer(const void* ParameterStruct, const FShaderParametersMetadata* Metadata);
 
 	/** Tracks and validates the creation of a new externally created / registered resource instances. */
-	void ValidateCreateTexture(FRDGTextureRef Texture);
-	void ValidateCreateBuffer(FRDGBufferRef Buffer);
-	void ValidateCreateSRV(FRDGTextureSRVRef SRV);
-	void ValidateCreateSRV(FRDGBufferSRVRef SRV);
-	void ValidateCreateUAV(FRDGTextureUAVRef UAV);
-	void ValidateCreateUAV(FRDGBufferUAVRef UAV);
-	void ValidateCreateUniformBuffer(FRDGUniformBufferRef UniformBuffer);
-
-	void ValidateRegisterExternalTexture(
+	RENDERCORE_API void ValidateCreateTexture(FRDGTextureRef Texture);
+	RENDERCORE_API void ValidateCreateBuffer(FRDGBufferRef Buffer);
+	RENDERCORE_API void ValidateCreateSRV(FRDGTextureSRVRef SRV);
+	RENDERCORE_API void ValidateCreateSRV(FRDGBufferSRVRef SRV);
+	RENDERCORE_API void ValidateCreateUAV(FRDGTextureUAVRef UAV);
+	RENDERCORE_API void ValidateCreateUAV(FRDGBufferUAVRef UAV);
+	RENDERCORE_API void ValidateCreateUniformBuffer(FRDGUniformBufferRef UniformBuffer);
+
+	RENDERCORE_API void ValidateRegisterExternalTexture(
 		const TRefCountPtr<IPooledRenderTarget>& ExternalPooledTexture,
 		const TCHAR* Name,
 		ERDGTextureFlags Flags);
 
-	void ValidateRegisterExternalBuffer(
+	RENDERCORE_API void ValidateRegisterExternalBuffer(
 		const TRefCountPtr<FRDGPooledBuffer>& ExternalPooledBuffer,
 		const TCHAR* Name,
 		ERDGBufferFlags Flags);
 
-	void ValidateRegisterExternalTexture(FRDGTextureRef Texture);
-	void ValidateRegisterExternalBuffer(FRDGBufferRef Buffer);
+	RENDERCORE_API void ValidateRegisterExternalTexture(FRDGTextureRef Texture);
+	RENDERCORE_API void ValidateRegisterExternalBuffer(FRDGBufferRef Buffer);
 
-	void ValidateUploadBuffer(FRDGBufferRef Buffer, const void* InitialData, uint64 InitialDataSize);
-	void ValidateUploadBuffer(FRDGBufferRef Buffer, const void* InitialData, uint64 InitialDataSize, const FRDGBufferInitialDataFreeCallback& InitialDataFreeCallback);
-	void ValidateUploadBuffer(FRDGBufferRef Buffer, const FRDGBufferInitialDataCallback& InitialDataCallback, const FRDGBufferInitialDataSizeCallback& InitialDataSizeCallback);
-	void ValidateUploadBuffer(FRDGBufferRef Buffer, const FRDGBufferInitialDataCallback& InitialDataCallback, const FRDGBufferInitialDataSizeCallback& InitialDataSizeCallback, const FRDGBufferInitialDataFreeCallback& InitialDataFreeCallback);
+	RENDERCORE_API void ValidateUploadBuffer(FRDGBufferRef Buffer, const void* InitialData, uint64 InitialDataSize);
+	RENDERCORE_API void ValidateUploadBuffer(FRDGBufferRef Buffer, const void* InitialData, uint64 InitialDataSize, const FRDGBufferInitialDataFreeCallback& InitialDataFreeCallback);
+	RENDERCORE_API void ValidateUploadBuffer(FRDGBufferRef Buffer, const FRDGBufferInitialDataCallback& InitialDataCallback, const FRDGBufferInitialDataSizeCallback& InitialDataSizeCallback);
+	RENDERCORE_API void ValidateUploadBuffer(FRDGBufferRef Buffer, const FRDGBufferInitialDataCallback& InitialDataCallback, const FRDGBufferInitialDataSizeCallback& InitialDataSizeCallback, const FRDGBufferInitialDataFreeCallback& InitialDataFreeCallback);
 
 	/** Validates a resource extraction operation. */
-	void ValidateExtractTexture(FRDGTextureRef Texture, TRefCountPtr<IPooledRenderTarget>* OutTexturePtr);
-	void ValidateExtractBuffer(FRDGBufferRef Buffer, TRefCountPtr<FRDGPooledBuffer>* OutBufferPtr);
+	RENDERCORE_API void ValidateExtractTexture(FRDGTextureRef Texture, TRefCountPtr<IPooledRenderTarget>* OutTexturePtr);
+	RENDERCORE_API void ValidateExtractBuffer(FRDGBufferRef Buffer, TRefCountPtr<FRDGPooledBuffer>* OutBufferPtr);
 
-	void ValidateConvertToExternalResource(FRDGViewableResource* Resource);
-	void ValidateConvertToExternalUniformBuffer(FRDGUniformBuffer* UniformBuffer);
+	RENDERCORE_API void ValidateConvertToExternalResource(FRDGViewableResource* Resource);
+	RENDERCORE_API void ValidateConvertToExternalUniformBuffer(FRDGUniformBuffer* UniformBuffer);
 
 	/** Tracks and validates the addition of a new pass to the graph.
 	 *  @param bSkipPassAccessMarking Skips marking the pass as a producer or incrementing the pass access. Useful when
 	 *      the builder needs to inject a pass for debugging while preserving error messages and warnings for the original
 	 *      graph structure.
 	 */
-	void ValidateAddPass(const void* ParameterStruct, const FShaderParametersMetadata* Metadata, const FRDGEventName& Name, ERDGPassFlags Flags);
-	void ValidateAddPass(const FRDGEventName& Name, ERDGPassFlags Flags);
-	void ValidateAddPass(const FRDGPass* Pass, bool bSkipPassAccessMarking);
+	RENDERCORE_API void ValidateAddPass(const void* ParameterStruct, const FShaderParametersMetadata* Metadata, const FRDGEventName& Name, ERDGPassFlags Flags);
+	RENDERCORE_API void ValidateAddPass(const FRDGEventName& Name, ERDGPassFlags Flags);
+	RENDERCORE_API void ValidateAddPass(const FRDGPass* Pass, bool bSkipPassAccessMarking);
 
 	/** Validate pass state before and after execution. */
-	void ValidateExecutePassBegin(const FRDGPass* Pass);
-	void ValidateExecutePassEnd(const FRDGPass* Pass);
+	RENDERCORE_API void ValidateExecutePassBegin(const FRDGPass* Pass);
+	RENDERCORE_API void ValidateExecutePassEnd(const FRDGPass* Pass);
 
 	/** Validate graph state before and after execution. */
-	void ValidateExecuteBegin();
-	void ValidateExecuteEnd();
+	RENDERCORE_API void ValidateExecuteBegin();
+	RENDERCORE_API void ValidateExecuteEnd();
 
 	/** Removes the 'produced but not used' warning from the requested resource. */
-	void RemoveUnusedWarning(FRDGViewableResource* Resource);
+	RENDERCORE_API void RemoveUnusedWarning(FRDGViewableResource* Resource);
 
 	/** Attempts to mark a resource for clobbering. If already marked, returns false.  */
-	bool TryMarkForClobber(FRDGViewableResource* Resource) const;
+	RENDERCORE_API bool TryMarkForClobber(FRDGViewableResource* Resource) const;
 
-	void ValidateGetPooledTexture(FRDGTextureRef Texture) const;
-	void ValidateGetPooledBuffer(FRDGBufferRef Buffer) const;
+	RENDERCORE_API void ValidateGetPooledTexture(FRDGTextureRef Texture) const;
+	RENDERCORE_API void ValidateGetPooledBuffer(FRDGBufferRef Buffer) const;
 
-	void ValidateSetAccessFinal(FRDGViewableResource* Resource, ERHIAccess AccessFinal);
+	RENDERCORE_API void ValidateSetAccessFinal(FRDGViewableResource* Resource, ERHIAccess AccessFinal);
 
-	void ValidateAddSubresourceAccess(FRDGViewableResource* Resource, const FRDGSubresourceState& Subresource, ERHIAccess Access);
+	RENDERCORE_API void ValidateAddSubresourceAccess(FRDGViewableResource* Resource, const FRDGSubresourceState& Subresource, ERHIAccess Access);
 
-	void ValidateUseExternalAccessMode(FRDGViewableResource* Resource, ERHIAccess ReadOnlyAccess, ERHIPipeline Pipelines);
-	void ValidateUseInternalAccessMode(FRDGViewableResource* Resaource);
+	RENDERCORE_API void ValidateUseExternalAccessMode(FRDGViewableResource* Resource, ERHIAccess ReadOnlyAccess, ERHIPipeline Pipelines);
+	RENDERCORE_API void ValidateUseInternalAccessMode(FRDGViewableResource* Resaource);
 
-	void ValidateExternalAccess(FRDGViewableResource* Resource, ERHIAccess Access, const FRDGPass* Pass);
+	RENDERCORE_API void ValidateExternalAccess(FRDGViewableResource* Resource, ERHIAccess Access, const FRDGPass* Pass);
 
 	/** Traverses all resources in the pass and marks whether they are externally accessible by user pass implementations. */
-	static void SetAllowRHIAccess(const FRDGPass* Pass, bool bAllowAccess);
+	static RENDERCORE_API void SetAllowRHIAccess(const FRDGPass* Pass, bool bAllowAccess);
 
 private:
 	void ValidateCreateViewableResource(FRDGViewableResource* Resource);
@@ -177,4 +177,4 @@ private:
 	const TCHAR* GraphName = nullptr;
 };
 
-#endif
\ No newline at end of file
+#endif
diff --git a/Engine/Source/Runtime/RenderCore/Public/RenderResource.h b/Engine/Source/Runtime/RenderCore/Public/RenderResource.h
--- a/Engine/Source/Runtime/RenderCore/Public/RenderResource.h
+++ b/Engine/Source/Runtime/RenderCore/Public/RenderResource.h
@@ -60,34 +60,34 @@ enum class ERayTracingMode : uint8
  * A rendering resource which is owned by the rendering thread.
  * NOTE - Adding new virtual methods to this class may require stubs added to FViewport/FDummyViewport, otherwise certain modules may have link errors
  */
-class RENDERCORE_API FRenderResource
+class FRenderResource
 {
 public:
 	////////////////////////////////////////////////////////////////////////////////////
 	// The following methods may not be called while asynchronously initializing / releasing render resources.
 
 	/** Release all render resources that are currently initialized. */
-	static void ReleaseRHIForAllResources();
+	static RENDERCORE_API void ReleaseRHIForAllResources();
 
 	/** Initialize all resources initialized before the RHI was initialized. */
-	static void InitPreRHIResources();
+	static RENDERCORE_API void InitPreRHIResources();
 
 	/** Call periodically to coalesce the render resource list. */
-	static void CoalesceResourceList();
+	static RENDERCORE_API void CoalesceResourceList();
 
 	/** Reinitializes render resources at a new feature level. */
-	static void ChangeFeatureLevel(ERHIFeatureLevel::Type NewFeatureLevel);
+	static RENDERCORE_API void ChangeFeatureLevel(ERHIFeatureLevel::Type NewFeatureLevel);
 
 	////////////////////////////////////////////////////////////////////////////////////
 
 	/** Default constructor. */
-	FRenderResource();
+	RENDERCORE_API FRenderResource();
 
 	/** Constructor when we know what feature level this resource should support */
-	FRenderResource(ERHIFeatureLevel::Type InFeatureLevel);
+	RENDERCORE_API FRenderResource(ERHIFeatureLevel::Type InFeatureLevel);
 
 	/** Destructor used to catch unreleased resources. */
-	virtual ~FRenderResource();
+	RENDERCORE_API virtual ~FRenderResource();
 
 	/**
 	 * Initializes the dynamic RHI resource and/or RHI render target used by this resource.
@@ -126,21 +126,21 @@ public:
 	 * Initializes the resource.
 	 * This is only called by the rendering thread.
 	 */
-	virtual void InitResource();
-	virtual void InitResource(FRHICommandList& RHICmdList);
+	RENDERCORE_API virtual void InitResource();
+	RENDERCORE_API virtual void InitResource(FRHICommandList& RHICmdList);
 
 	/**
 	 * Prepares the resource for deletion.
 	 * This is only called by the rendering thread.
 	 */
-	virtual void ReleaseResource();
+	RENDERCORE_API virtual void ReleaseResource();
 
 	/**
 	 * If the resource's RHI resources have been initialized, then release and reinitialize it.  Otherwise, do nothing.
 	 * This is only called by the rendering thread.
 	 */
-	void UpdateRHI();
-	void UpdateRHI(FRHICommandList& RHICmdList);
+	RENDERCORE_API void UpdateRHI();
+	RENDERCORE_API void UpdateRHI(FRHICommandList& RHICmdList);
 
 	/** @return The resource's friendly name.  Typically a UObject name. */
 	virtual FString GetFriendlyName() const { return TEXT("undefined"); }
@@ -151,8 +151,8 @@ public:
 	int32 GetListIndex() const { return ListIndex; }
 
 	/** SetOwnerName should be called before BeginInitResource for the owner name to be successfully tracked. */
-	void SetOwnerName(const FName& InOwnerName);
-	FName GetOwnerName() const;
+	RENDERCORE_API void SetOwnerName(const FName& InOwnerName);
+	RENDERCORE_API FName GetOwnerName() const;
 
 protected:
 	// This is used during mobile editor preview refactor, this will eventually be replaced with a parameter to InitRHI() etc..
@@ -186,8 +186,8 @@ protected:
 	}
 
 private:
-	static bool ShouldFreeResourceObject(void* ResourceObject, FResourceArrayInterface* ResourceArray);
-	static FBufferRHIRef CreateRHIBufferInternal(
+	static RENDERCORE_API bool ShouldFreeResourceObject(void* ResourceObject, FResourceArrayInterface* ResourceArray);
+	static RENDERCORE_API FBufferRHIRef CreateRHIBufferInternal(
 		const TCHAR* InDebugName,
 		const FName& InOwnerName,
 		uint32 ResourceCount,
@@ -338,7 +338,7 @@ struct FMipBiasFade
 };
 
 /** A textures resource. */
-class RENDERCORE_API FTexture : public FRenderResource
+class FTexture : public FRenderResource
 {
 public:
 
@@ -375,43 +375,43 @@ public:
 	 **/
 	bool				bSRGB = false;
 
-	FTexture();
-	virtual ~FTexture();
+	RENDERCORE_API FTexture();
+	RENDERCORE_API virtual ~FTexture();
 
 	const FTextureRHIRef& GetTextureRHI() { return TextureRHI; }
 
 	/** Returns the width of the texture in pixels. */
-	virtual uint32 GetSizeX() const;
+	RENDERCORE_API virtual uint32 GetSizeX() const;
 
 	/** Returns the height of the texture in pixels. */
-	virtual uint32 GetSizeY() const;
+	RENDERCORE_API virtual uint32 GetSizeY() const;
 
 	/** Returns the depth of the texture in pixels. */
-	virtual uint32 GetSizeZ() const;
+	RENDERCORE_API virtual uint32 GetSizeZ() const;
 
 	// FRenderResource interface.
-	virtual void ReleaseRHI() override;
-	virtual FString GetFriendlyName() const override;
+	RENDERCORE_API virtual void ReleaseRHI() override;
+	RENDERCORE_API virtual FString GetFriendlyName() const override;
 
 protected:
-	static FRHISamplerState* GetOrCreateSamplerState(const FSamplerStateInitializerRHI& Initializer);
+	static RENDERCORE_API FRHISamplerState* GetOrCreateSamplerState(const FSamplerStateInitializerRHI& Initializer);
 };
 
 /** A textures resource that includes an SRV. */
-class RENDERCORE_API FTextureWithSRV : public FTexture
+class FTextureWithSRV : public FTexture
 {
 public:
-	FTextureWithSRV();
-	virtual ~FTextureWithSRV();
+	RENDERCORE_API FTextureWithSRV();
+	RENDERCORE_API virtual ~FTextureWithSRV();
 
-	virtual void ReleaseRHI() override;
+	RENDERCORE_API virtual void ReleaseRHI() override;
 
 	/** SRV that views the entire texture */
 	FShaderResourceViewRHIRef ShaderResourceViewRHI;
 };
 
 /** A texture reference resource. */
-class RENDERCORE_API FTextureReference : public FRenderResource
+class FTextureReference : public FRenderResource
 {
 public:
 	/** The texture reference's RHI resource. */
@@ -423,53 +423,53 @@ private:
 
 public:
 	/** Default constructor. */
-	FTextureReference();
+	RENDERCORE_API FTextureReference();
 
 	// Destructor
-	virtual ~FTextureReference();
+	RENDERCORE_API virtual ~FTextureReference();
 
 	/** Returns the last time the texture has been rendered via this reference. */
-	double GetLastRenderTime() const;
+	RENDERCORE_API double GetLastRenderTime() const;
 
 	/** Invalidates the last render time. */
-	void InvalidateLastRenderTime();
+	RENDERCORE_API void InvalidateLastRenderTime();
 
 	/** Returns true if the texture reference has been initialized from the game thread. */
 	bool IsInitialized_GameThread() const { return bInitialized_GameThread; }
 
 	/** Kicks off the initialization process on the game thread. */
-	void BeginInit_GameThread();
+	RENDERCORE_API void BeginInit_GameThread();
 
 	/** Kicks off the release process on the game thread. */
-	void BeginRelease_GameThread();
+	RENDERCORE_API void BeginRelease_GameThread();
 
 	// FRenderResource interface.
-	virtual void InitRHI();
-	virtual void ReleaseRHI();
-	virtual FString GetFriendlyName() const;
+	RENDERCORE_API virtual void InitRHI();
+	RENDERCORE_API virtual void ReleaseRHI();
+	RENDERCORE_API virtual FString GetFriendlyName() const;
 };
 
 /** A vertex buffer resource */
-class RENDERCORE_API FVertexBuffer : public FRenderResource
+class FVertexBuffer : public FRenderResource
 {
 public:
-	FVertexBuffer();
-	virtual ~FVertexBuffer();
+	RENDERCORE_API FVertexBuffer();
+	RENDERCORE_API virtual ~FVertexBuffer();
 
 	// FRenderResource interface.
-	virtual void ReleaseRHI() override;
-	virtual FString GetFriendlyName() const override;
+	RENDERCORE_API virtual void ReleaseRHI() override;
+	RENDERCORE_API virtual FString GetFriendlyName() const override;
 
 	FBufferRHIRef VertexBufferRHI;
 };
 
-class RENDERCORE_API FVertexBufferWithSRV : public FVertexBuffer
+class FVertexBufferWithSRV : public FVertexBuffer
 {
 public:
-	FVertexBufferWithSRV();
-	~FVertexBufferWithSRV();
+	RENDERCORE_API FVertexBufferWithSRV();
+	RENDERCORE_API ~FVertexBufferWithSRV();
 
-	virtual void ReleaseRHI() override;
+	RENDERCORE_API virtual void ReleaseRHI() override;
 
 	/** SRV that views the entire texture */
 	FShaderResourceViewRHIRef ShaderResourceViewRHI;
@@ -479,28 +479,28 @@ public:
 };
 
 /** An index buffer resource. */
-class RENDERCORE_API FIndexBuffer : public FRenderResource
+class FIndexBuffer : public FRenderResource
 {
 public:
-	FIndexBuffer();
-	virtual ~FIndexBuffer();
+	RENDERCORE_API FIndexBuffer();
+	RENDERCORE_API virtual ~FIndexBuffer();
 
 	// FRenderResource interface.
-	virtual void ReleaseRHI() override;
-	virtual FString GetFriendlyName() const override;
+	RENDERCORE_API virtual void ReleaseRHI() override;
+	RENDERCORE_API virtual FString GetFriendlyName() const override;
 
 	FBufferRHIRef IndexBufferRHI;
 };
 
-class RENDERCORE_API FBufferWithRDG : public FRenderResource
+class FBufferWithRDG : public FRenderResource
 {
 public:
-	FBufferWithRDG();
-	FBufferWithRDG(const FBufferWithRDG& Other);
-	FBufferWithRDG& operator=(const FBufferWithRDG& Other);
-	~FBufferWithRDG() override;
+	RENDERCORE_API FBufferWithRDG();
+	RENDERCORE_API FBufferWithRDG(const FBufferWithRDG& Other);
+	RENDERCORE_API FBufferWithRDG& operator=(const FBufferWithRDG& Other);
+	RENDERCORE_API ~FBufferWithRDG() override;
 
-	void ReleaseRHI() override;
+	RENDERCORE_API void ReleaseRHI() override;
 
 	TRefCountPtr<FRDGPooledBuffer> Buffer;
 };
diff --git a/Engine/Source/Runtime/RenderCore/Public/RenderTargetPool.h b/Engine/Source/Runtime/RenderCore/Public/RenderTargetPool.h
--- a/Engine/Source/Runtime/RenderCore/Public/RenderTargetPool.h
+++ b/Engine/Source/Runtime/RenderCore/Public/RenderTargetPool.h
@@ -22,7 +22,7 @@ class FRHICommandList;
 class FRenderTargetPool;
 
 /** The reference to a pooled render target, use like this: TRefCountPtr<IPooledRenderTarget> */
-struct RENDERCORE_API FPooledRenderTarget final : public IPooledRenderTarget
+struct FPooledRenderTarget final : public IPooledRenderTarget
 {
 	FPooledRenderTarget(FRHITexture* Texture, const FPooledRenderTargetDesc& InDesc, FRenderTargetPool* InRenderTargetPool) 
 		: RenderTargetPool(InRenderTargetPool)
@@ -59,9 +59,9 @@ struct RENDERCORE_API FPooledRenderTarget final : public IPooledRenderTarget
 		return uint32(NumRefs);
 	}
 
-	bool IsFree() const override;
+	RENDERCORE_API bool IsFree() const override;
 	bool IsTracked() const override { return RenderTargetPool != nullptr; }
-	uint32 ComputeMemorySize() const override;
+	RENDERCORE_API uint32 ComputeMemorySize() const override;
 
 private:
 	/** Pointer back to the pool for render targets which are actually pooled, otherwise NULL. */
@@ -80,7 +80,7 @@ private:
 	FRDGPooledTexture PooledTexture;
 
 	/** @return true:release this one, false otherwise */
-	bool OnFrameStart();
+	RENDERCORE_API bool OnFrameStart();
 
 	friend class FRDGTexture;
 	friend class FRDGBuilder;
@@ -90,14 +90,14 @@ private:
 /**
  * Encapsulates the render targets pools that allows easy sharing (mostly used on the render thread side)
  */
-class RENDERCORE_API FRenderTargetPool : public FRenderResource
+class FRenderTargetPool : public FRenderResource
 {
 public:
 	FRenderTargetPool() = default;
 
-	TRefCountPtr<IPooledRenderTarget> FindFreeElement(FRHITextureCreateInfo Desc, const TCHAR* Name);
+	RENDERCORE_API TRefCountPtr<IPooledRenderTarget> FindFreeElement(FRHITextureCreateInfo Desc, const TCHAR* Name);
 
-	bool FindFreeElement(const FRHITextureCreateInfo& Desc, TRefCountPtr<IPooledRenderTarget>& Out, const TCHAR* Name);
+	RENDERCORE_API bool FindFreeElement(const FRHITextureCreateInfo& Desc, TRefCountPtr<IPooledRenderTarget>& Out, const TCHAR* Name);
 
 	/**
 	 * @param DebugName must not be 0, we only store the pointer
@@ -114,38 +114,38 @@ public:
 		return FindFreeElement(Translate(Desc), Out, InDebugName);
 	}
 
-	void CreateUntrackedElement(const FPooledRenderTargetDesc& Desc, TRefCountPtr<IPooledRenderTarget>& Out, const FSceneRenderTargetItem& Item);
+	RENDERCORE_API void CreateUntrackedElement(const FPooledRenderTargetDesc& Desc, TRefCountPtr<IPooledRenderTarget>& Out, const FSceneRenderTargetItem& Item);
 
 	/** Only to get statistics on usage and free elements. Normally only called in renderthread or if FlushRenderingCommands was called() */
-	void GetStats(uint32& OutWholeCount, uint32& OutWholePoolInKB, uint32& OutUsedInKB) const;
+	RENDERCORE_API void GetStats(uint32& OutWholeCount, uint32& OutWholePoolInKB, uint32& OutUsedInKB) const;
 	/**
 	 * Can release RT, should be called once per frame.
 	 * call from RenderThread only
 	 */
-	void TickPoolElements();
+	RENDERCORE_API void TickPoolElements();
 	/** Free renderer resources */
-	void ReleaseRHI();
+	RENDERCORE_API void ReleaseRHI();
 
 	/** Allows to remove a resource so it cannot be shared and gets released immediately instead a/some frame[s] later. */
-	void FreeUnusedResource(TRefCountPtr<IPooledRenderTarget>& In);
+	RENDERCORE_API void FreeUnusedResource(TRefCountPtr<IPooledRenderTarget>& In);
 
 	/** Good to call between levels or before memory intense operations. */
-	void FreeUnusedResources();
+	RENDERCORE_API void FreeUnusedResources();
 
 	// for debugging purpose, assumes you call FlushRenderingCommands() be
 	// @return can be 0, that doesn't mean iteration is done
-	FPooledRenderTarget* GetElementById(uint32 Id) const;
+	RENDERCORE_API FPooledRenderTarget* GetElementById(uint32 Id) const;
 
 	uint32 GetElementCount() const { return PooledRenderTargets.Num(); }
 
 	// @return -1 if not found
-	int32 FindIndex(IPooledRenderTarget* In) const;
+	RENDERCORE_API int32 FindIndex(IPooledRenderTarget* In) const;
 
 	// Logs out usage information.
-	void DumpMemoryUsage(FOutputDevice& OutputDevice);
+	RENDERCORE_API void DumpMemoryUsage(FOutputDevice& OutputDevice);
 
 private:
-	void FreeElementAtIndex(int32 Index);
+	RENDERCORE_API void FreeElementAtIndex(int32 Index);
 
 	/** Elements can be 0, we compact the buffer later. */
 	TArray<uint32> PooledRenderTargetHashes;
@@ -159,7 +159,7 @@ private:
 	bool bCurrentlyOverBudget = false;
 
 	// could be done on the fly but that makes the RenderTargetPoolEvents harder to read
-	void CompactPool();
+	RENDERCORE_API void CompactPool();
 
 	friend struct FPooledRenderTarget;
 	friend class FVisualizeTexture;
@@ -168,4 +168,4 @@ private:
 };
 
 /** The global render targets for easy shading. */
-extern RENDERCORE_API TGlobalResource<FRenderTargetPool> GRenderTargetPool;
\ No newline at end of file
+extern RENDERCORE_API TGlobalResource<FRenderTargetPool> GRenderTargetPool;
diff --git a/Engine/Source/Runtime/RenderCore/Public/RenderingThread.h b/Engine/Source/Runtime/RenderCore/Public/RenderingThread.h
--- a/Engine/Source/Runtime/RenderCore/Public/RenderingThread.h
+++ b/Engine/Source/Runtime/RenderCore/Public/RenderingThread.h
@@ -94,14 +94,14 @@ extern RENDERCORE_API void TickRenderingTickables();
 extern RENDERCORE_API void StartRenderCommandFenceBundler();
 extern RENDERCORE_API void StopRenderCommandFenceBundler();
 
-class RENDERCORE_API FCoreRenderDelegates
+class FCoreRenderDelegates
 {
 public:
 	DECLARE_MULTICAST_DELEGATE(FOnFlushRenderingCommandsStart);
-	static FOnFlushRenderingCommandsStart OnFlushRenderingCommandsStart;
+	static RENDERCORE_API FOnFlushRenderingCommandsStart OnFlushRenderingCommandsStart;
 
 	DECLARE_MULTICAST_DELEGATE(FOnFlushRenderingCommandsEnd);
-	static FOnFlushRenderingCommandsEnd OnFlushRenderingCommandsEnd;
+	static RENDERCORE_API FOnFlushRenderingCommandsEnd OnFlushRenderingCommandsEnd;
 };
 ////////////////////////////////////
 // Render thread suspension
@@ -110,17 +110,17 @@ public:
 /**
  * Encapsulates stopping and starting the renderthread so that other threads can manipulate graphics resources.
  */
-class RENDERCORE_API FSuspendRenderingThread
+class FSuspendRenderingThread
 {
 public:
 	/**
 	 *	Constructor that flushes and suspends the renderthread
 	 *	@param bRecreateThread	- Whether the rendering thread should be completely destroyed and recreated, or just suspended.
 	 */
-	FSuspendRenderingThread( bool bRecreateThread );
+	RENDERCORE_API FSuspendRenderingThread( bool bRecreateThread );
 
 	/** Destructor that starts the renderthread again */
-	~FSuspendRenderingThread();
+	RENDERCORE_API ~FSuspendRenderingThread();
 
 private:
 	/** Whether we should use a rendering thread or not */
@@ -143,7 +143,7 @@ private:
 UE_TRACE_CHANNEL_EXTERN(RenderCommandsChannel, RENDERCORE_API);
 
 /** The parent class of commands stored in the rendering command queue. */
-class RENDERCORE_API FRenderCommand
+class FRenderCommand
 {
 public:
 	// All render commands run on the render thread
@@ -343,7 +343,7 @@ private:
 	FRHICommandList RHICmdListStack;
 };
 
-class RENDERCORE_API FRenderThreadScope
+class FRenderThreadScope
 {
 	typedef TFunction<void(FRHICommandListImmediate&)> RenderCommandFunction;
 	typedef TArray<RenderCommandFunction> RenderCommandFunctionArray;
@@ -507,4 +507,4 @@ using FStopRenderingThreadDelegate = FStopRenderingThread::FDelegate;
 
 extern RENDERCORE_API FDelegateHandle RegisterStopRenderingThreadDelegate(const FStopRenderingThreadDelegate& InDelegate);
 
-extern RENDERCORE_API void UnregisterStopRenderingThreadDelegate(FDelegateHandle InDelegateHandle);
\ No newline at end of file
+extern RENDERCORE_API void UnregisterStopRenderingThreadDelegate(FDelegateHandle InDelegateHandle);
diff --git a/Engine/Source/Runtime/RenderCore/Public/Shader.h b/Engine/Source/Runtime/RenderCore/Public/Shader.h
--- a/Engine/Source/Runtime/RenderCore/Public/Shader.h
+++ b/Engine/Source/Runtime/RenderCore/Public/Shader.h
@@ -138,7 +138,7 @@ public:
 
 class FShaderUniformBufferParameterInfo
 {
-	DECLARE_TYPE_LAYOUT(FShaderUniformBufferParameterInfo, NonVirtual);
+	DECLARE_EXPORTED_TYPE_LAYOUT(FShaderUniformBufferParameterInfo, RENDERCORE_API, NonVirtual);
 public:
 	LAYOUT_FIELD(uint16, BaseIndex);
 
@@ -169,7 +169,7 @@ public:
 
 class FShaderResourceParameterInfo
 {
-	DECLARE_TYPE_LAYOUT(FShaderResourceParameterInfo, NonVirtual);
+	DECLARE_EXPORTED_TYPE_LAYOUT(FShaderResourceParameterInfo, RENDERCORE_API, NonVirtual);
 public:
 	LAYOUT_FIELD(uint16, BaseIndex);
 	LAYOUT_FIELD(uint8, BufferIndex);
@@ -208,7 +208,7 @@ public:
 
 class FShaderLooseParameterInfo
 {
-	DECLARE_TYPE_LAYOUT(FShaderLooseParameterInfo, NonVirtual);
+	DECLARE_EXPORTED_TYPE_LAYOUT(FShaderLooseParameterInfo, RENDERCORE_API, NonVirtual);
 public:
 	LAYOUT_FIELD(uint16, BaseIndex);
 	LAYOUT_FIELD(uint16, Size);
@@ -243,7 +243,7 @@ public:
 
 class FShaderLooseParameterBufferInfo
 {
-	DECLARE_TYPE_LAYOUT(FShaderLooseParameterBufferInfo, NonVirtual);
+	DECLARE_EXPORTED_TYPE_LAYOUT(FShaderLooseParameterBufferInfo, RENDERCORE_API, NonVirtual);
 public:
 	LAYOUT_FIELD(uint16, BaseIndex);
 	LAYOUT_FIELD(uint16, Size);
@@ -282,7 +282,7 @@ public:
 
 class FShaderParameterMapInfo
 {
-	DECLARE_TYPE_LAYOUT(FShaderParameterMapInfo, NonVirtual);
+	DECLARE_EXPORTED_TYPE_LAYOUT(FShaderParameterMapInfo, RENDERCORE_API, NonVirtual);
 public:
 	LAYOUT_FIELD(TMemoryImageArray<FShaderUniformBufferParameterInfo>, UniformBuffers);
 	LAYOUT_FIELD(TMemoryImageArray<FShaderResourceParameterInfo>, TextureSamplers);
@@ -326,19 +326,19 @@ struct FRegisterRayTracingPayloadTypeHelper
 #define IMPLEMENT_RT_PAYLOAD_TYPE(PayloadType, PayloadSize)  static FRegisterRayTracingPayloadTypeHelper IMPLEMENT_RT_PAYLOAD_TYPE_CONCAT(PayloadHelper, __COUNTER__) = FRegisterRayTracingPayloadTypeHelper(PayloadType, PayloadSize, nullptr);
 #define IMPLEMENT_RT_PAYLOAD_TYPE_FUNCTION(PayloadType, PayloadSizeFunction)  static FRegisterRayTracingPayloadTypeHelper IMPLEMENT_RT_PAYLOAD_TYPE_CONCAT(PayloadHelper, __COUNTER__) = FRegisterRayTracingPayloadTypeHelper(PayloadType, 0u, PayloadSizeFunction);
 
-class RENDERCORE_API FShaderMapResource : public FRenderResource, public FDeferredCleanupInterface
+class FShaderMapResource : public FRenderResource, public FDeferredCleanupInterface
 {
 public:
-	static bool ArePlatformsCompatible(EShaderPlatform CurrentPlatform, EShaderPlatform TargetPlatform);
+	static RENDERCORE_API bool ArePlatformsCompatible(EShaderPlatform CurrentPlatform, EShaderPlatform TargetPlatform);
 
 	EShaderPlatform GetPlatform() const { return Platform; }
 
-	void AddRef();
-	void Release();
+	RENDERCORE_API void AddRef();
+	RENDERCORE_API void Release();
 	inline int32 GetNumRefs() const { return NumRefs.load(std::memory_order_relaxed); }
 
 	// FRenderResource interface.
-	virtual void ReleaseRHI();
+	RENDERCORE_API virtual void ReleaseRHI();
 
 	inline int32 GetNumShaders() const
 	{
@@ -372,7 +372,7 @@ public:
 	/** Return shader hash for a particular shader without creating it. */
 	virtual FSHAHash GetShaderHash(int32 ShaderIndex) = 0;
 
-	void BeginCreateAllShaders();
+	RENDERCORE_API void BeginCreateAllShaders();
 
 #if RHI_RAYTRACING
 
@@ -388,9 +388,9 @@ public:
 		return GetRayTracingHitGroupLibraryIndex(ShaderIndex);
 	}
 
-	static void GetRayTracingHitGroupLibrary(TArray<FRHIRayTracingShader*>& RayTracingHitGroupShaders, FRHIRayTracingShader* DefaultShader);
-	static void GetRayTracingCallableShaderLibrary(TArray<FRHIRayTracingShader*>& RayTracingCallableShaders, FRHIRayTracingShader* DefaultShader);
-	static void GetRayTracingMissShaderLibrary(TArray<FRHIRayTracingShader*>& RayTracingMissShaders, FRHIRayTracingShader* DefaultShader);
+	static RENDERCORE_API void GetRayTracingHitGroupLibrary(TArray<FRHIRayTracingShader*>& RayTracingHitGroupShaders, FRHIRayTracingShader* DefaultShader);
+	static RENDERCORE_API void GetRayTracingCallableShaderLibrary(TArray<FRHIRayTracingShader*>& RayTracingCallableShaders, FRHIRayTracingShader* DefaultShader);
+	static RENDERCORE_API void GetRayTracingMissShaderLibrary(TArray<FRHIRayTracingShader*>& RayTracingMissShaders, FRHIRayTracingShader* DefaultShader);
 
 	inline uint32 GetRayTracingHitGroupLibraryIndex(int32 ShaderIndex)
 	{
@@ -417,8 +417,8 @@ public:
 	virtual uint32 GetSizeBytes() const = 0;
 
 protected:
-	explicit FShaderMapResource(EShaderPlatform InPlatform, int32 NumShaders);
-	virtual ~FShaderMapResource();
+	RENDERCORE_API explicit FShaderMapResource(EShaderPlatform InPlatform, int32 NumShaders);
+	RENDERCORE_API virtual ~FShaderMapResource();
 
 	SIZE_T GetAllocatedSize() const
 	{
@@ -437,12 +437,12 @@ protected:
 
 	virtual bool TryRelease() { return true; }
 
-	void ReleaseShaders();
+	RENDERCORE_API void ReleaseShaders();
 
 private:
 
 	/** Creates an entry in RHIShaders array and registers it among the raytracing libs if needed. Created shader is returned. */
-	FRHIShader* CreateShaderOrCrash(int32 ShaderIndex);
+	RENDERCORE_API FRHIShader* CreateShaderOrCrash(int32 ShaderIndex);
 
 	/** This lock is to prevent two threads creating the same RHIShaders element. It is only taken if the element is to be created. */
 	FCriticalSection RHIShadersCreationGuard;
@@ -498,7 +498,7 @@ public:
 	};
 #endif // WITH_EDITORONLY_DATA
 
-	RENDERCORE_API FShaderMapResourceCode() {}
+	FShaderMapResourceCode() {}
 	RENDERCORE_API FShaderMapResourceCode(const FShaderMapResourceCode& Other);
     RENDERCORE_API ~FShaderMapResourceCode();
 
@@ -534,7 +534,7 @@ public:
 #endif // WITH_EDITORONLY_DATA
 };
 	
-class RENDERCORE_API FShaderMapResource_InlineCode : public FShaderMapResource
+class FShaderMapResource_InlineCode : public FShaderMapResource
 {
 public:
 	FShaderMapResource_InlineCode(EShaderPlatform InPlatform, FShaderMapResourceCode* InCode)
@@ -543,23 +543,23 @@ public:
 	{}
 
 	// FShaderMapResource interface
-	virtual FSHAHash GetShaderHash(int32 ShaderIndex) override;
-	virtual FRHIShader* CreateRHIShaderOrCrash(int32 ShaderIndex) override;
+	RENDERCORE_API virtual FSHAHash GetShaderHash(int32 ShaderIndex) override;
+	RENDERCORE_API virtual FRHIShader* CreateRHIShaderOrCrash(int32 ShaderIndex) override;
 	virtual uint32 GetSizeBytes() const override { return sizeof(*this) + GetAllocatedSize(); }
 
 	TRefCountPtr<FShaderMapResourceCode> Code;
 };
 
-class RENDERCORE_API FShaderMapPointerTable : public FPointerTableBase
+class FShaderMapPointerTable : public FPointerTableBase
 {
 public:
-	virtual int32 AddIndexedPointer(const FTypeLayoutDesc& TypeDesc, void* Ptr) override;
-	virtual void* GetIndexedPointer(const FTypeLayoutDesc& TypeDesc, uint32 i) const override;
+	RENDERCORE_API virtual int32 AddIndexedPointer(const FTypeLayoutDesc& TypeDesc, void* Ptr) override;
+	RENDERCORE_API virtual void* GetIndexedPointer(const FTypeLayoutDesc& TypeDesc, uint32 i) const override;
 
 	virtual FShaderMapPointerTable* Clone() const { return new FShaderMapPointerTable(*this); }
 
-	virtual void SaveToArchive(FArchive& Ar, const FPlatformTypeLayoutParameters& LayoutParams, const void* FrozenObject) const override;
-	virtual bool LoadFromArchive(FArchive& Ar, const FPlatformTypeLayoutParameters& LayoutParams, void* FrozenObject) override;
+	RENDERCORE_API virtual void SaveToArchive(FArchive& Ar, const FPlatformTypeLayoutParameters& LayoutParams, const void* FrozenObject) const override;
+	RENDERCORE_API virtual bool LoadFromArchive(FArchive& Ar, const FPlatformTypeLayoutParameters& LayoutParams, void* FrozenObject) override;
 
 	TPtrTable<FShaderType> ShaderTypes;
 	TPtrTable<FVertexFactoryType> VFTypes;
@@ -734,7 +734,7 @@ class FMaterial;
 /**
  * Stores all shader parameter bindings and their corresponding offset and size in the shader's parameters struct.
  */
-class RENDERCORE_API FShaderParameterBindings
+class FShaderParameterBindings
 {
 public:
 	struct FParameter
@@ -770,12 +770,12 @@ public:
 		LAYOUT_FIELD(uint16, ByteOffset);
 	};
 
-	DECLARE_TYPE_LAYOUT(FShaderParameterBindings, NonVirtual);
+	DECLARE_EXPORTED_TYPE_LAYOUT(FShaderParameterBindings, RENDERCORE_API, NonVirtual);
 public:
 	static constexpr uint16 kInvalidBufferIndex = 0xFFFF;
 
-	void BindForLegacyShaderParameters(const FShader* Shader, int32 PermutationId, const FShaderParameterMap& ParameterMaps, const FShaderParametersMetadata& StructMetaData, bool bShouldBindEverything = false);
-	void BindForRootShaderParameters(const FShader* Shader, int32 PermutationId, const FShaderParameterMap& ParameterMaps);
+	RENDERCORE_API void BindForLegacyShaderParameters(const FShader* Shader, int32 PermutationId, const FShaderParameterMap& ParameterMaps, const FShaderParametersMetadata& StructMetaData, bool bShouldBindEverything = false);
+	RENDERCORE_API void BindForRootShaderParameters(const FShader* Shader, int32 PermutationId, const FShaderParameterMap& ParameterMaps);
 
 	LAYOUT_FIELD(TMemoryImageArray<FParameter>, Parameters);
 	LAYOUT_FIELD(TMemoryImageArray<FResourceParameter>, ResourceParameters);
@@ -834,10 +834,10 @@ DECLARE_EXPORTED_TEMPLATE_INTRINSIC_TYPE_LAYOUT(template<>, TIndexedPtr<FShaderT
 DECLARE_EXPORTED_TEMPLATE_INTRINSIC_TYPE_LAYOUT(template<>, TIndexedPtr<FVertexFactoryType>, RENDERCORE_API);
 
 /** A compiled shader and its parameter bindings. */
-class RENDERCORE_API FShader
+class FShader
 {
 	friend class FShaderType;
-	DECLARE_TYPE_LAYOUT(FShader, NonVirtual);
+	DECLARE_EXPORTED_TYPE_LAYOUT(FShader, RENDERCORE_API, NonVirtual);
 public:
 	using FPermutationDomain = FShaderPermutationNone;
 	using FPermutationParameters = FShaderPermutationParameters;
@@ -848,14 +848,14 @@ public:
 	 * Used to construct a shader for deserialization.
 	 * This still needs to initialize members to safe values since FShaderType::GenerateSerializationHistory uses this constructor.
 	 */
-	FShader();
+	RENDERCORE_API FShader();
 
 	/**
 	 * Construct a shader from shader compiler output.
 	 */
-	FShader(const CompiledShaderInitializerType& Initializer);
+	RENDERCORE_API FShader(const CompiledShaderInitializerType& Initializer);
 
-	~FShader();
+	RENDERCORE_API ~FShader();
 
 	/** Can be overridden by FShader subclasses to modify their compile environment just before compilation occurs. */
 	static void ModifyCompilationEnvironment(const FShaderPermutationParameters&, FShaderCompilerEnvironment&) {}
@@ -870,16 +870,16 @@ public:
 	static ERayTracingPayloadType GetRayTracingPayloadType(const int32 PermutationId) { return static_cast<ERayTracingPayloadType>(0); }
 
 	/** Returns the hash of the shader file that this shader was compiled with. */
-	const FSHAHash& GetHash() const;
+	RENDERCORE_API const FSHAHash& GetHash() const;
 
-	const FSHAHash& GetVertexFactoryHash() const;
+	RENDERCORE_API const FSHAHash& GetVertexFactoryHash() const;
 	
-	const FSHAHash& GetOutputHash() const;
+	RENDERCORE_API const FSHAHash& GetOutputHash() const;
 
 	/** Returns an identifier suitable for deterministic sorting of shaders between sessions. */
 	uint32 GetSortKey() const { return SortKey; }
 
-	void Finalize(const FShaderMapResourceCode* Code);
+	RENDERCORE_API void Finalize(const FShaderMapResourceCode* Code);
 
 	// Accessors.
 	inline FShaderType* GetType(const FShaderMapPointerTable& InPointerTable) const { return Type.Get(InPointerTable.ShaderTypes); }
@@ -948,12 +948,12 @@ public:
 		}
 	}
 
-	const FShaderParametersMetadata* FindAutomaticallyBoundUniformBufferStruct(int32 BaseIndex) const;
+	RENDERCORE_API const FShaderParametersMetadata* FindAutomaticallyBoundUniformBufferStruct(int32 BaseIndex) const;
 
-	void DumpDebugInfo(const FShaderMapPointerTable& InPtrTable);
+	RENDERCORE_API void DumpDebugInfo(const FShaderMapPointerTable& InPtrTable);
 
 #if WITH_EDITOR
-	void SaveShaderStableKeys(const FShaderMapPointerTable& InPtrTable, EShaderPlatform TargetShaderPlatform, int32 PermutationId, const struct FStableShaderKeyAndValue& SaveKeyVal);
+	RENDERCORE_API void SaveShaderStableKeys(const FShaderMapPointerTable& InPtrTable, EShaderPlatform TargetShaderPlatform, int32 PermutationId, const struct FStableShaderKeyAndValue& SaveKeyVal);
 #endif // WITH_EDITOR
 
 	/** Returns the meta data for the root shader parameter struct. */
@@ -963,7 +963,7 @@ public:
 	}
 
 private:
-	void BuildParameterMapInfo(const TMap<FString, FParameterAllocation>& ParameterMap);
+	RENDERCORE_API void BuildParameterMapInfo(const TMap<FString, FParameterAllocation>& ParameterMap);
 
 public:
 	LAYOUT_FIELD(FShaderParameterBindings, Bindings);
@@ -1186,7 +1186,7 @@ using TShaderRef = TShaderRefBase<ShaderType, FShaderMapPointerTable>;
  * A shader type can manage multiple instance of FShader across mutiple dimensions such as EShaderPlatform, or permutation id.
  * The number of permutation of a shader type is simply given by GetPermutationCount().
  */
-class RENDERCORE_API FShaderType
+class FShaderType
 {
 public:
 	enum class EShaderTypeForDynamicCast : uint32
@@ -1219,24 +1219,24 @@ public:
 #endif // WITH_EDITOR
 
 	/** @return The global shader factory list. */
-	static TLinkedList<FShaderType*>*& GetTypeList();
+	static RENDERCORE_API TLinkedList<FShaderType*>*& GetTypeList();
 
-	static FShaderType* GetShaderTypeByName(const TCHAR* Name);
-	static TArray<const FShaderType*> GetShaderTypesByFilename(const TCHAR* Filename);
+	static RENDERCORE_API FShaderType* GetShaderTypeByName(const TCHAR* Name);
+	static RENDERCORE_API TArray<const FShaderType*> GetShaderTypesByFilename(const TCHAR* Filename);
 
 	/** @return The global shader name to type map */
-	static TMap<FHashedName, FShaderType*>& GetNameToTypeMap();
+	static RENDERCORE_API TMap<FHashedName, FShaderType*>& GetNameToTypeMap();
 
-	static const TArray<FShaderType*>& GetSortedTypes(EShaderTypeForDynamicCast Type);
+	static RENDERCORE_API const TArray<FShaderType*>& GetSortedTypes(EShaderTypeForDynamicCast Type);
 
 	/** Initialize FShaderType static members, this must be called before any shader types are created. */
-	static void Initialize(const TMap<FString, TArray<const TCHAR*> >& ShaderFileToUniformBufferVariables);
+	static RENDERCORE_API void Initialize(const TMap<FString, TArray<const TCHAR*> >& ShaderFileToUniformBufferVariables);
 
 	/** Uninitializes FShaderType cached data. */
-	static void Uninitialize();
+	static RENDERCORE_API void Uninitialize();
 
 	/** Minimal initialization constructor. */
-	FShaderType(
+	RENDERCORE_API FShaderType(
 		EShaderTypeForDynamicCast InShaderTypeForDynamicCast,
 		FTypeLayoutDesc& InTypeLayout,
 		const TCHAR* InName,
@@ -1256,16 +1256,16 @@ public:
 		const FShaderParametersMetadata* InRootParametersMetadata
 	);
 
-	virtual ~FShaderType();
+	RENDERCORE_API virtual ~FShaderType();
 
 	/** Constructs a new instance of the shader type for deserialization. */
-	FShader* ConstructForDeserialization() const;
-	FShader* ConstructCompiled(const FShader::CompiledShaderInitializerType& Initializer) const;
+	RENDERCORE_API FShader* ConstructForDeserialization() const;
+	RENDERCORE_API FShader* ConstructCompiled(const FShader::CompiledShaderInitializerType& Initializer) const;
 
-	bool ShouldCompilePermutation(const FShaderPermutationParameters& Parameters) const;
+	RENDERCORE_API bool ShouldCompilePermutation(const FShaderPermutationParameters& Parameters) const;
 
 #if WITH_EDITOR
-	void ModifyCompilationEnvironment(const FShaderPermutationParameters& Parameters, FShaderCompilerEnvironment& OutEnvironment) const;
+	RENDERCORE_API void ModifyCompilationEnvironment(const FShaderPermutationParameters& Parameters, FShaderCompilerEnvironment& OutEnvironment) const;
 
 	/**
 	* Checks if the shader type should pass compilation for a particular set of parameters.
@@ -1273,13 +1273,13 @@ public:
 	* @param ParameterMap - Shader parameters to validate.
 	* @param OutError - List for appending validation errors.
 	*/
-	bool ValidateCompiledResult(EShaderPlatform Platform, const FShaderParameterMap& ParameterMap, TArray<FString>& OutError) const;
+	RENDERCORE_API bool ValidateCompiledResult(EShaderPlatform Platform, const FShaderParameterMap& ParameterMap, TArray<FString>& OutError) const;
 #endif // WITH_EDITOR
 
-	ERayTracingPayloadType GetRayTracingPayloadType(const int32 PermutationId) const;
+	RENDERCORE_API ERayTracingPayloadType GetRayTracingPayloadType(const int32 PermutationId) const;
 
 	/** Calculates a Hash based on this shader type's source code and includes */
-	const FSHAHash& GetSourceHash(EShaderPlatform ShaderPlatform) const;
+	RENDERCORE_API const FSHAHash& GetSourceHash(EShaderPlatform ShaderPlatform) const;
 
 	/** Serializes a shader type reference by name. */
 	RENDERCORE_API friend FArchive& operator<<(FArchive& Ar,FShaderType*& Ref);
@@ -1437,7 +1437,7 @@ public:
 	}
 
 	/** Adds include statements for uniform buffers that this shader type references. */
-	void AddUniformBufferIncludesToEnvironment(FShaderCompilerEnvironment& OutEnvironment, EShaderPlatform Platform) const;
+	RENDERCORE_API void AddUniformBufferIncludesToEnvironment(FShaderCompilerEnvironment& OutEnvironment, EShaderPlatform Platform) const;
 
 	UE_DEPRECATED(5.2, "AddReferencedUniformBufferIncludes has moved to AddUniformBufferIncludesToEnvironment and no longer takes a prefix argument.")
 	inline void AddReferencedUniformBufferIncludes(FShaderCompilerEnvironment& OutEnvironment, FString& OutSourceFilePrefix, EShaderPlatform Platform) const
@@ -1445,15 +1445,15 @@ public:
 		AddUniformBufferIncludesToEnvironment(OutEnvironment, Platform);
 	}
 
-	void UpdateReferencedUniformBufferNames(const TMap<FString, TArray<const TCHAR*>>& ShaderFileToUniformBufferVariables);
+	RENDERCORE_API void UpdateReferencedUniformBufferNames(const TMap<FString, TArray<const TCHAR*>>& ShaderFileToUniformBufferVariables);
 
 	UE_DEPRECATED(5.2, "FlushShaderFileCache is deprecated. UpdateReferencedUniformBufferNames should be used to flush any uniform buffer changes")
-	void FlushShaderFileCache(const TMap<FString, TArray<const TCHAR*> >& ShaderFileToUniformBufferVariables);
+	RENDERCORE_API void FlushShaderFileCache(const TMap<FString, TArray<const TCHAR*> >& ShaderFileToUniformBufferVariables);
 
-	void GetShaderStableKeyParts(struct FStableShaderKeyAndValue& SaveKeyVal);
+	RENDERCORE_API void GetShaderStableKeyParts(struct FStableShaderKeyAndValue& SaveKeyVal);
 #endif // WITH_EDITOR
 
-	void DumpDebugInfo();
+	RENDERCORE_API void DumpDebugInfo();
 
 private:
 	EShaderTypeForDynamicCast ShaderTypeForDynamicCast;
@@ -1483,7 +1483,7 @@ private:
 	friend void RENDERCORE_API DumpShaderStats( EShaderPlatform Platform, EShaderFrequency Frequency );
 
 	/** Tracks whether serialization history for all shader types has been initialized. */
-	static bool bInitializedSerializationHistory;
+	static RENDERCORE_API bool bInitializedSerializationHistory;
 
 #if WITH_EDITOR
 protected:
@@ -1503,7 +1503,7 @@ protected:
  * The collection happens during static init, the actual construction happens later during launch.
  * The purpose of collecting the types is the CommitAll function, called in LaunchEngineLoop, to ensure all type instances are constructed and registered before other systems start iterating them.
  */
-class RENDERCORE_API FShaderTypeRegistration
+class FShaderTypeRegistration
 {
 public:
 	FShaderTypeRegistration(TFunctionRef<FShaderType& ()> LazyShaderTypeAccessor)
@@ -1512,10 +1512,10 @@ public:
 		GetInstances().Add(this);
 	}
 
-	static TArray<const FShaderTypeRegistration*>& GetInstances();
+	static RENDERCORE_API TArray<const FShaderTypeRegistration*>& GetInstances();
 
 	// Actually register all the types and clear the array
-	static void CommitAll();
+	static RENDERCORE_API void CommitAll();
 
 private:
 	TFunctionRef<FShaderType& ()> LazyShaderTypeAccessor;
@@ -1724,19 +1724,19 @@ struct FShaderCompiledShaderInitializerType
 	TemplatePrefix FShaderTypeRegistration ShaderClass::ShaderTypeRegistration{TFunctionRef<::FShaderType&()>{ShaderClass::GetStaticType}};
 
 // Binding of a set of shader stages in a single pipeline
-class RENDERCORE_API FShaderPipelineType
+class FShaderPipelineType
 {
 public:
 	// Set bShouldOptimizeUnusedOutputs to true if we want unique FShaders for each shader pipeline
 	// Set bShouldOptimizeUnusedOutputs to false if the FShaders will point to the individual shaders in the map
-	FShaderPipelineType(
+	RENDERCORE_API FShaderPipelineType(
 		const TCHAR* InName,
 		const FShaderType* InVertexOrMeshShader,
 		const FShaderType* InGeometryOrAmplificationShader,
 		const FShaderType* InPixelShader,
 		bool bInIsMeshPipeline,
 		bool bInShouldOptimizeUnusedOutputs);
-	~FShaderPipelineType();
+	RENDERCORE_API ~FShaderPipelineType();
 
 	FORCEINLINE bool HasMeshShader() const { return AllStages[SF_Mesh] != nullptr; }
 	FORCEINLINE bool HasGeometry() const { return AllStages[SF_Geometry] != nullptr; }
@@ -1756,19 +1756,19 @@ public:
 	// Returns an array of valid stages, sorted from PS->GS->DS->HS->VS, no gaps if missing stages
 	FORCEINLINE const TArray<const FShaderType*>& GetStages() const { return Stages; }
 
-	static TLinkedList<FShaderPipelineType*>*& GetTypeList();
+	static RENDERCORE_API TLinkedList<FShaderPipelineType*>*& GetTypeList();
 
-	static const TArray<FShaderPipelineType*>& GetSortedTypes(FShaderType::EShaderTypeForDynamicCast Type);
+	static RENDERCORE_API const TArray<FShaderPipelineType*>& GetSortedTypes(FShaderType::EShaderTypeForDynamicCast Type);
 
 	/** @return The global shader pipeline name to type map */
-	static TMap<FHashedName, FShaderPipelineType*>& GetNameToTypeMap();
-	static const FShaderPipelineType* GetShaderPipelineTypeByName(const FHashedName& Name);
+	static RENDERCORE_API TMap<FHashedName, FShaderPipelineType*>& GetNameToTypeMap();
+	static RENDERCORE_API const FShaderPipelineType* GetShaderPipelineTypeByName(const FHashedName& Name);
 
 	/** Initialize static members, this must be called before any shader types are created. */
-	static void Initialize();
-	static void Uninitialize();
+	static RENDERCORE_API void Initialize();
+	static RENDERCORE_API void Uninitialize();
 
-	static TArray<const FShaderPipelineType*> GetShaderPipelineTypesByFilename(const TCHAR* Filename);
+	static RENDERCORE_API TArray<const FShaderPipelineType*> GetShaderPipelineTypesByFilename(const TCHAR* Filename);
 
 	/** Serializes a shader type reference by name. */
 	RENDERCORE_API friend FArchive& operator<<(FArchive& Ar, const FShaderPipelineType*& Ref);
@@ -1782,12 +1782,12 @@ public:
 	bool IsMaterialTypePipeline() const { return Stages[0]->GetMaterialShaderType() != nullptr; }
 	bool IsMeshMaterialTypePipeline() const { return Stages[0]->GetMeshMaterialShaderType() != nullptr; }
 
-	bool ShouldOptimizeUnusedOutputs(EShaderPlatform Platform) const;
+	RENDERCORE_API bool ShouldOptimizeUnusedOutputs(EShaderPlatform Platform) const;
 
 	/** Calculates a Hash based on this shader pipeline type stages' source code and includes */
-	const FSHAHash& GetSourceHash(EShaderPlatform ShaderPlatform) const;
+	RENDERCORE_API const FSHAHash& GetSourceHash(EShaderPlatform ShaderPlatform) const;
 
-	bool ShouldCompilePermutation(const FShaderPermutationParameters& Parameters) const;
+	RENDERCORE_API bool ShouldCompilePermutation(const FShaderPermutationParameters& Parameters) const;
 
 protected:
 	const TCHAR* const Name;
@@ -1805,7 +1805,7 @@ protected:
 	uint32 HashIndex;
 	bool bShouldOptimizeUnusedOutputs;
 
-	static bool bInitialized;
+	static RENDERCORE_API bool bInitialized;
 };
 
 #if !UE_BUILD_DOCS
@@ -1974,15 +1974,15 @@ public:
 };
 
 // A Shader Pipeline instance with compiled stages
-class RENDERCORE_API FShaderPipeline
+class FShaderPipeline
 {
-	DECLARE_TYPE_LAYOUT(FShaderPipeline, NonVirtual);
+	DECLARE_EXPORTED_TYPE_LAYOUT(FShaderPipeline, RENDERCORE_API, NonVirtual);
 public:
 	explicit FShaderPipeline(const FShaderPipelineType* InType) : TypeName(InType->GetHashedName()) { FMemory::Memzero(&PermutationIds, sizeof(PermutationIds)); }
-	~FShaderPipeline();
+	RENDERCORE_API ~FShaderPipeline();
 
-	void AddShader(FShader* Shader, int32 PermutationId);
-	FShader* FindOrAddShader(FShader* Shader, int32 PermutationId);
+	RENDERCORE_API void AddShader(FShader* Shader, int32 PermutationId);
+	RENDERCORE_API FShader* FindOrAddShader(FShader* Shader, int32 PermutationId);
 
 	inline uint32 GetNumShaders() const
 	{
@@ -2039,9 +2039,9 @@ public:
 		return Result;
 	}
 
-	void Validate(const FShaderPipelineType* InPipelineType) const;
+	RENDERCORE_API void Validate(const FShaderPipelineType* InPipelineType) const;
 
-	void Finalize(const FShaderMapResourceCode* Code);
+	RENDERCORE_API void Finalize(const FShaderMapResourceCode* Code);
 
 	enum EFilter
 	{
@@ -2052,7 +2052,7 @@ public:
 
 	/** Saves stable keys for the shaders in the pipeline */
 #if WITH_EDITOR
-	void SaveShaderStableKeys(const FShaderMapPointerTable& InPtrTable, EShaderPlatform TargetShaderPlatform, const struct FStableShaderKeyAndValue& SaveKeyVal) const;
+	RENDERCORE_API void SaveShaderStableKeys(const FShaderMapPointerTable& InPtrTable, EShaderPlatform TargetShaderPlatform, const struct FStableShaderKeyAndValue& SaveKeyVal) const;
 #endif // WITH_EDITOR
 
 	LAYOUT_FIELD(FHashedName, TypeName);
@@ -2102,9 +2102,9 @@ private:
 };
 
 /** A collection of shaders of different types */
-class RENDERCORE_API FShaderMapContent
+class FShaderMapContent
 {
-	DECLARE_TYPE_LAYOUT(FShaderMapContent, NonVirtual);
+	DECLARE_EXPORTED_TYPE_LAYOUT(FShaderMapContent, RENDERCORE_API, NonVirtual);
 public:
 	struct FProjectShaderPipelineToKey
 	{
@@ -2112,14 +2112,14 @@ public:
 	};
 
 	/** Default constructor. */
-	explicit FShaderMapContent(EShaderPlatform InPlatform);
+	RENDERCORE_API explicit FShaderMapContent(EShaderPlatform InPlatform);
 
 	/** Destructor ensures pipelines cleared up. */
-	~FShaderMapContent();
+	RENDERCORE_API ~FShaderMapContent();
 
-	EShaderPlatform GetShaderPlatform() const;
+	RENDERCORE_API EShaderPlatform GetShaderPlatform() const;
 
-	void Validate(const FShaderMapBase& InShaderMap) const;
+	RENDERCORE_API void Validate(const FShaderMapBase& InShaderMap) const;
 
 	/** Finds the shader with the given type.  Asserts on failure. */
 	template<typename ShaderType>
@@ -2144,7 +2144,7 @@ public:
 	}
 
 	/** Finds the shader with the given type name.  May return NULL. */
-	FShader* GetShader(const FHashedName& TypeName, int32 PermutationId = 0) const;
+	RENDERCORE_API FShader* GetShader(const FHashedName& TypeName, int32 PermutationId = 0) const;
 
 	/** Finds the shader with the given type. */
 	bool HasShader(const FHashedName& TypeName, int32 PermutationId) const
@@ -2168,42 +2168,42 @@ public:
 		return ShaderPipelines;
 	}
 
-	void AddShader(const FHashedName& TypeName, int32 PermutationId, FShader* Shader);
+	RENDERCORE_API void AddShader(const FHashedName& TypeName, int32 PermutationId, FShader* Shader);
 
-	FShader* FindOrAddShader(const FHashedName& TypeName, int32 PermutationId, FShader* Shader);
+	RENDERCORE_API FShader* FindOrAddShader(const FHashedName& TypeName, int32 PermutationId, FShader* Shader);
 
-	void AddShaderPipeline(FShaderPipeline* Pipeline);
+	RENDERCORE_API void AddShaderPipeline(FShaderPipeline* Pipeline);
 
-	FShaderPipeline* FindOrAddShaderPipeline(FShaderPipeline* Pipeline);
+	RENDERCORE_API FShaderPipeline* FindOrAddShaderPipeline(FShaderPipeline* Pipeline);
 
 	/**
 	 * Removes the shader of the given type from the shader map
 	 * @param Type Shader type to remove the entry for 
 	 */
-	void RemoveShaderTypePermutaion(const FHashedName& TypeName, int32 PermutationId);
+	RENDERCORE_API void RemoveShaderTypePermutaion(const FHashedName& TypeName, int32 PermutationId);
 
 	inline void RemoveShaderTypePermutaion(const FShaderType* Type, int32 PermutationId)
 	{
 		RemoveShaderTypePermutaion(Type->GetHashedName(), PermutationId);
 	}
 
-	void RemoveShaderPipelineType(const FShaderPipelineType* ShaderPipelineType);
+	RENDERCORE_API void RemoveShaderPipelineType(const FShaderPipelineType* ShaderPipelineType);
 
 	/** Builds a list of the shaders in a shader map. */
-	void GetShaderList(const FShaderMapBase& InShaderMap, const FSHAHash& InMaterialShaderMapHash, TMap<FShaderId, TShaderRef<FShader>>& OutShaders) const;
+	RENDERCORE_API void GetShaderList(const FShaderMapBase& InShaderMap, const FSHAHash& InMaterialShaderMapHash, TMap<FShaderId, TShaderRef<FShader>>& OutShaders) const;
 
 	/** Builds a list of the shaders in a shader map. Key is FShaderType::TypeName */
-	void GetShaderList(const FShaderMapBase& InShaderMap, TMap<FHashedName, TShaderRef<FShader>>& OutShaders) const;
+	RENDERCORE_API void GetShaderList(const FShaderMapBase& InShaderMap, TMap<FHashedName, TShaderRef<FShader>>& OutShaders) const;
 
 	/** Builds a list of the shader pipelines in a shader map. */
-	void GetShaderPipelineList(const FShaderMapBase& InShaderMap, TArray<FShaderPipelineRef>& OutShaderPipelines, FShaderPipeline::EFilter Filter) const;
+	RENDERCORE_API void GetShaderPipelineList(const FShaderMapBase& InShaderMap, TArray<FShaderPipelineRef>& OutShaderPipelines, FShaderPipeline::EFilter Filter) const;
 
 #if WITH_EDITOR
-	uint32 GetMaxTextureSamplersShaderMap(const FShaderMapBase& InShaderMap) const;
+	RENDERCORE_API uint32 GetMaxTextureSamplersShaderMap(const FShaderMapBase& InShaderMap) const;
 
-	void GetOutdatedTypes(const FShaderMapBase& InShaderMap, TArray<const FShaderType*>& OutdatedShaderTypes, TArray<const FShaderPipelineType*>& OutdatedShaderPipelineTypes, TArray<const FVertexFactoryType*>& OutdatedFactoryTypes) const;
+	RENDERCORE_API void GetOutdatedTypes(const FShaderMapBase& InShaderMap, TArray<const FShaderType*>& OutdatedShaderTypes, TArray<const FShaderPipelineType*>& OutdatedShaderPipelineTypes, TArray<const FVertexFactoryType*>& OutdatedFactoryTypes) const;
 
-	void SaveShaderStableKeys(const FShaderMapBase& InShaderMap, EShaderPlatform TargetShaderPlatform, const struct FStableShaderKeyAndValue& SaveKeyVal);
+	RENDERCORE_API void SaveShaderStableKeys(const FShaderMapBase& InShaderMap, EShaderPlatform TargetShaderPlatform, const struct FStableShaderKeyAndValue& SaveKeyVal);
 #endif // WITH_EDITOR
 
 	/** @return true if the map is empty */
@@ -2213,7 +2213,7 @@ public:
 	}
 
 	/** @return The number of shaders in the map. */
-	uint32 GetNumShaders() const;
+	RENDERCORE_API uint32 GetNumShaders() const;
 
 	/** @return The number of shader pipelines in the map. */
 	inline uint32 GetNumShaderPipelines() const
@@ -2222,7 +2222,7 @@ public:
 	}
 
 	/** clears out all shaders and deletes shader pipelines held in the map */
-	void Empty();
+	RENDERCORE_API void Empty();
 
 	inline FShaderPipeline* GetShaderPipeline(const FHashedName& PipelineTypeName) const
 	{
@@ -2238,14 +2238,14 @@ public:
 	inline bool HasShaderPipeline(const FHashedName& PipelineTypeName) const { return GetShaderPipeline(PipelineTypeName) != nullptr; }
 	inline bool HasShaderPipeline(const FShaderPipelineType* PipelineType) const { return (GetShaderPipeline(PipelineType) != nullptr); }
 
-	uint32 GetMaxNumInstructionsForShader(const FShaderMapBase& InShaderMap, FShaderType* ShaderType) const;
+	RENDERCORE_API uint32 GetMaxNumInstructionsForShader(const FShaderMapBase& InShaderMap, FShaderType* ShaderType) const;
 
-	void Finalize(const FShaderMapResourceCode* Code);
+	RENDERCORE_API void Finalize(const FShaderMapResourceCode* Code);
 
-	void UpdateHash(FSHA1& Hasher) const;
+	RENDERCORE_API void UpdateHash(FSHA1& Hasher) const;
 
 protected:
-	void EmptyShaderPipelines();
+	RENDERCORE_API void EmptyShaderPipelines();
 
 	using FMemoryImageHashTable = THashTable<FMemoryImageAllocator>;
 
@@ -2258,12 +2258,12 @@ protected:
 	LAYOUT_FIELD(FMemoryImageName, ShaderPlatformName);
 };
 
-class RENDERCORE_API FShaderMapBase
+class FShaderMapBase
 {
 public:
-	virtual ~FShaderMapBase();
+	RENDERCORE_API virtual ~FShaderMapBase();
 
-	FShaderMapResourceCode* GetResourceCode();
+	RENDERCORE_API FShaderMapResourceCode* GetResourceCode();
 
 	inline FShaderMapResource* GetResource() const { return Resource; }
 	inline FShaderMapResource* GetResourceChecked() const { check(Resource); return Resource; }
@@ -2278,18 +2278,18 @@ public:
 	inline EShaderPlatform GetShaderPlatform() const { return Content.Object ? Content.Object->GetShaderPlatform() : SP_NumPlatforms; }
 	inline uint32 GetFrozenContentSize() const { return Content.FrozenSize; }
 
-	void AssignContent(TMemoryImageObject<FShaderMapContent> InContent);
+	RENDERCORE_API void AssignContent(TMemoryImageObject<FShaderMapContent> InContent);
 
-	void FinalizeContent();
-	void UnfreezeContent();
-	bool Serialize(FArchive& Ar, bool bInlineShaderResources, bool bLoadedByCookedMaterial, bool bInlineShaderCode=false);
+	RENDERCORE_API void FinalizeContent();
+	RENDERCORE_API void UnfreezeContent();
+	RENDERCORE_API bool Serialize(FArchive& Ar, bool bInlineShaderResources, bool bLoadedByCookedMaterial, bool bInlineShaderCode=false);
 
 	EShaderPermutationFlags GetPermutationFlags() const
 	{
 		return PermutationFlags;
 	}
 
-	FString ToString() const;
+	RENDERCORE_API FString ToString() const;
 
 #if WITH_EDITOR
 	inline void GetOutdatedTypes(TArray<const FShaderType*>& OutdatedShaderTypes, TArray<const FShaderPipelineType*>& OutdatedShaderPipelineTypes, TArray<const FVertexFactoryType*>& OutdatedFactoryTypes) const
@@ -2321,12 +2321,12 @@ public:
 #endif // WITH_EDITOR
 
 protected:
-	FShaderMapBase();
+	RENDERCORE_API FShaderMapBase();
 
-	void AssignCopy(const FShaderMapBase& Source);
+	RENDERCORE_API void AssignCopy(const FShaderMapBase& Source);
 
-	void InitResource();
-	void DestroyContent();
+	RENDERCORE_API void InitResource();
+	RENDERCORE_API void DestroyContent();
 
 protected:
 	virtual const FTypeLayoutDesc& GetContentTypeDesc() const = 0;
diff --git a/Engine/Source/Runtime/RenderCore/Public/ShaderCodeArchive.h b/Engine/Source/Runtime/RenderCore/Public/ShaderCodeArchive.h
--- a/Engine/Source/Runtime/RenderCore/Public/ShaderCodeArchive.h
+++ b/Engine/Source/Runtime/RenderCore/Public/ShaderCodeArchive.h
@@ -73,7 +73,7 @@ struct FShaderCodeEntry
 };
 
 // Portion of shader code archive that's serialize to disk
-class RENDERCORE_API FSerializedShaderArchive
+class FSerializedShaderArchive
 {
 public:
 
@@ -197,25 +197,25 @@ public:
 			;
 	}
 
-	int32 FindShaderMapWithKey(const FSHAHash& Hash, uint32 Key) const;
-	int32 FindShaderMap(const FSHAHash& Hash) const;
-	bool FindOrAddShaderMap(const FSHAHash& Hash, int32& OutIndex, const FShaderMapAssetPaths* AssociatedAssets);
+	RENDERCORE_API int32 FindShaderMapWithKey(const FSHAHash& Hash, uint32 Key) const;
+	RENDERCORE_API int32 FindShaderMap(const FSHAHash& Hash) const;
+	RENDERCORE_API bool FindOrAddShaderMap(const FSHAHash& Hash, int32& OutIndex, const FShaderMapAssetPaths* AssociatedAssets);
 
-	int32 FindShaderWithKey(const FSHAHash& Hash, uint32 Key) const;
-	int32 FindShader(const FSHAHash& Hash) const;
-	bool FindOrAddShader(const FSHAHash& Hash, int32& OutIndex);
-	void RemoveLastAddedShader();
+	RENDERCORE_API int32 FindShaderWithKey(const FSHAHash& Hash, uint32 Key) const;
+	RENDERCORE_API int32 FindShader(const FSHAHash& Hash) const;
+	RENDERCORE_API bool FindOrAddShader(const FSHAHash& Hash, int32& OutIndex);
+	RENDERCORE_API void RemoveLastAddedShader();
 
-	void DecompressShader(int32 Index, const TArray<TArray<uint8>>& ShaderCode, TArray<uint8>& OutDecompressedShader) const;
+	RENDERCORE_API void DecompressShader(int32 Index, const TArray<TArray<uint8>>& ShaderCode, TArray<uint8>& OutDecompressedShader) const;
 
-	void Finalize();
-	void Serialize(FArchive& Ar);
+	RENDERCORE_API void Finalize();
+	RENDERCORE_API void Serialize(FArchive& Ar);
 #if WITH_EDITOR
-	void SaveAssetInfo(FArchive& Ar);
-	bool LoadAssetInfo(const FString& Filename);
-	void CreateAsChunkFrom(const FSerializedShaderArchive& Parent, const TSet<FName>& PackagesInChunk, TArray<int32>& OutShaderCodeEntriesNeeded);
-	void CollectStatsAndDebugInfo(FDebugStats& OutDebugStats, FExtendedDebugStats* OutExtendedDebugStats);
-	void DumpContentsInPlaintext(FString& OutText) const;
+	RENDERCORE_API void SaveAssetInfo(FArchive& Ar);
+	RENDERCORE_API bool LoadAssetInfo(const FString& Filename);
+	RENDERCORE_API void CreateAsChunkFrom(const FSerializedShaderArchive& Parent, const TSet<FName>& PackagesInChunk, TArray<int32>& OutShaderCodeEntriesNeeded);
+	RENDERCORE_API void CollectStatsAndDebugInfo(FDebugStats& OutDebugStats, FExtendedDebugStats* OutExtendedDebugStats);
+	RENDERCORE_API void DumpContentsInPlaintext(FString& OutText) const;
 	RENDERCORE_API friend FCbWriter& operator<<(FCbWriter& Writer, const FSerializedShaderArchive& Archive);
 	RENDERCORE_API friend bool LoadFromCompactBinary(FCbFieldView Field, FSerializedShaderArchive& OutArchive);
 #endif
diff --git a/Engine/Source/Runtime/RenderCore/Public/ShaderCodeLibrary.h b/Engine/Source/Runtime/RenderCore/Public/ShaderCodeLibrary.h
--- a/Engine/Source/Runtime/RenderCore/Public/ShaderCodeLibrary.h
+++ b/Engine/Source/Runtime/RenderCore/Public/ShaderCodeLibrary.h
@@ -42,13 +42,13 @@ class FCbFieldView;
 class FCbWriter;
 #endif
 
-struct RENDERCORE_API FShaderCodeLibraryPipeline
+struct FShaderCodeLibraryPipeline
 {
 	FSHAHash Shaders[SF_NumGraphicsFrequencies];
 	mutable uint32 Hash;
 	
 	/** Fills the hashes from the pipeline stage shaders */
-	void Initialize(const FShaderPipeline* Pipeline);
+	RENDERCORE_API void Initialize(const FShaderPipeline* Pipeline);
 
 	FShaderCodeLibraryPipeline() : Hash(0) {}
 	
@@ -77,7 +77,7 @@ struct RENDERCORE_API FShaderCodeLibraryPipeline
 	}
 
 	/** Computes a longer hash that uniquely identifies the whole pipeline, used in FStableShaderKeyAndValue */
-	void GetPipelineHash(FSHAHash& Output);
+	RENDERCORE_API void GetPipelineHash(FSHAHash& Output);
 	
 	friend FArchive& operator<<( FArchive& Ar, FShaderCodeLibraryPipeline& Info )
 	{
@@ -89,7 +89,7 @@ struct RENDERCORE_API FShaderCodeLibraryPipeline
 	}
 };
 
-struct RENDERCORE_API FCompactFullName
+struct FCompactFullName
 {
 	TArray<FName, TInlineAllocator<16>> ObjectClassAndPath;
 
@@ -98,21 +98,21 @@ struct RENDERCORE_API FCompactFullName
 		return ObjectClassAndPath == Other.ObjectClassAndPath;
 	}
 
-	FString ToString() const;
-	FString ToStringPathOnly() const;
-	void AppendString(FStringBuilderBase& Out) const;
-	void AppendString(FAnsiStringBuilderBase& Out) const;
-	void ParseFromString(const FStringView& Src);
+	RENDERCORE_API FString ToString() const;
+	RENDERCORE_API FString ToStringPathOnly() const;
+	RENDERCORE_API void AppendString(FStringBuilderBase& Out) const;
+	RENDERCORE_API void AppendString(FAnsiStringBuilderBase& Out) const;
+	RENDERCORE_API void ParseFromString(const FStringView& Src);
 	friend RENDERCORE_API uint32 GetTypeHash(const FCompactFullName& A);
 
 #if WITH_EDITOR
 	/** Used to set up some compact FName paths for the FCompactFullName */
-	void SetCompactFullNameFromObject(UObject* InDepObject);
+	RENDERCORE_API void SetCompactFullNameFromObject(UObject* InDepObject);
 #endif
 };
 
 
-struct RENDERCORE_API FStableShaderKeyAndValue
+struct FStableShaderKeyAndValue
 {
 	FCompactFullName ClassNameAndObjectPath;
 	FName ShaderType;
@@ -135,16 +135,16 @@ struct RENDERCORE_API FStableShaderKeyAndValue
 	{
 	}
 
-	void ComputeKeyHash();
-	void ParseFromString(const FStringView& Src);
-	void ParseFromStringCached(const FStringView& Src, class TMap<uint32, FName>& NameCache);
-	FString ToString() const;
-	void ToString(FString& OutResult) const;
-	void AppendString(FAnsiStringBuilderBase& Out) const;
-	static FString HeaderLine();
+	RENDERCORE_API void ComputeKeyHash();
+	RENDERCORE_API void ParseFromString(const FStringView& Src);
+	RENDERCORE_API void ParseFromStringCached(const FStringView& Src, class TMap<uint32, FName>& NameCache);
+	RENDERCORE_API FString ToString() const;
+	RENDERCORE_API void ToString(FString& OutResult) const;
+	RENDERCORE_API void AppendString(FAnsiStringBuilderBase& Out) const;
+	static RENDERCORE_API FString HeaderLine();
 
 	/** Computes pipeline hash from the passed pipeline. Pass nullptr to clear */
-	void SetPipelineHash(const FShaderPipeline* Pipeline);
+	RENDERCORE_API void SetPipelineHash(const FShaderPipeline* Pipeline);
 
 	friend bool operator ==(const FStableShaderKeyAndValue& A, const FStableShaderKeyAndValue& B)
 	{
@@ -183,16 +183,16 @@ using FSharedShaderCodeRelease = TMulticastDelegate<void(const FSHAHash&)>;
 
 // Collection of unique shader code
 // Populated at cook time
-struct RENDERCORE_API FShaderCodeLibrary
+struct FShaderCodeLibrary
 {
 	/** Adds the hooks for OnPakFileMounted, since pak opening normally happens earlier. */
-	static void PreInit();
+	static RENDERCORE_API void PreInit();
 
 	/** This is the real initialization function. */
-	static void InitForRuntime(EShaderPlatform ShaderPlatform);
-	static void Shutdown();
+	static RENDERCORE_API void InitForRuntime(EShaderPlatform ShaderPlatform);
+	static RENDERCORE_API void Shutdown();
 
-	static bool IsEnabled();
+	static RENDERCORE_API bool IsEnabled();
 
 	/**
 	 * Makes a number of ChunkIDs known to the library.
@@ -202,7 +202,7 @@ struct RENDERCORE_API FShaderCodeLibrary
 	 * preloaded.
 	 * Takes C array and not TArray because that makes easier to hardcode the chunk IDs (which is the intended use case for this function).
 	 */
-	static void AddKnownChunkIDs(const int32* IDs, const int32 NumChunkIDs);
+	static RENDERCORE_API void AddKnownChunkIDs(const int32* IDs, const int32 NumChunkIDs);
 
 	/** 
 	 * Open a named library.
@@ -212,57 +212,57 @@ struct RENDERCORE_API FShaderCodeLibrary
 	 * @param Directory location of the .ushadercode file
 	 * @return true if successful
 	 */
-	static bool OpenLibrary(FString const& Name, FString const& Directory);
+	static RENDERCORE_API bool OpenLibrary(FString const& Name, FString const& Directory);
 
 	/**
 	 * Close a named library.
 	 *
 	 * At runtime this will release the library data and further requests for shaders from this library will fail.
 	 */
-	static void CloseLibrary(FString const& Name);
+	static RENDERCORE_API void CloseLibrary(FString const& Name);
 
-	static bool ContainsShaderCode(const FSHAHash& Hash);
-    static bool ContainsShaderCode(const FSHAHash& Hash, const FString& LogicalLibraryName);
+	static RENDERCORE_API bool ContainsShaderCode(const FSHAHash& Hash);
+    static RENDERCORE_API bool ContainsShaderCode(const FSHAHash& Hash, const FString& LogicalLibraryName);
 
-	static TRefCountPtr<FShaderMapResource> LoadResource(const FSHAHash& Hash, FArchive* Ar);
+	static RENDERCORE_API TRefCountPtr<FShaderMapResource> LoadResource(const FSHAHash& Hash, FArchive* Ar);
 
-	static bool PreloadShader(const FSHAHash& Hash, FArchive* Ar);
-	static bool ReleasePreloadedShader(const FSHAHash& Hash);
+	static RENDERCORE_API bool PreloadShader(const FSHAHash& Hash, FArchive* Ar);
+	static RENDERCORE_API bool ReleasePreloadedShader(const FSHAHash& Hash);
 
-	static FVertexShaderRHIRef CreateVertexShader(EShaderPlatform Platform, const FSHAHash& Hash);
-	static FPixelShaderRHIRef CreatePixelShader(EShaderPlatform Platform, const FSHAHash& Hash);
-	static FGeometryShaderRHIRef CreateGeometryShader(EShaderPlatform Platform, const FSHAHash& Hash);
-	static FComputeShaderRHIRef CreateComputeShader(EShaderPlatform Platform, const FSHAHash& Hash);
-	static FMeshShaderRHIRef CreateMeshShader(EShaderPlatform Platform, const FSHAHash& Hash);
-	static FAmplificationShaderRHIRef CreateAmplificationShader(EShaderPlatform Platform, const FSHAHash& Hash);
-	static FRayTracingShaderRHIRef CreateRayTracingShader(EShaderPlatform Platform, const FSHAHash& Hash, EShaderFrequency Frequency);
+	static RENDERCORE_API FVertexShaderRHIRef CreateVertexShader(EShaderPlatform Platform, const FSHAHash& Hash);
+	static RENDERCORE_API FPixelShaderRHIRef CreatePixelShader(EShaderPlatform Platform, const FSHAHash& Hash);
+	static RENDERCORE_API FGeometryShaderRHIRef CreateGeometryShader(EShaderPlatform Platform, const FSHAHash& Hash);
+	static RENDERCORE_API FComputeShaderRHIRef CreateComputeShader(EShaderPlatform Platform, const FSHAHash& Hash);
+	static RENDERCORE_API FMeshShaderRHIRef CreateMeshShader(EShaderPlatform Platform, const FSHAHash& Hash);
+	static RENDERCORE_API FAmplificationShaderRHIRef CreateAmplificationShader(EShaderPlatform Platform, const FSHAHash& Hash);
+	static RENDERCORE_API FRayTracingShaderRHIRef CreateRayTracingShader(EShaderPlatform Platform, const FSHAHash& Hash, EShaderFrequency Frequency);
 
 	// Total number of shader entries in the library
-	static uint32 GetShaderCount(void);
+	static RENDERCORE_API uint32 GetShaderCount(void);
 	
 	// The shader platform that the library manages - at runtime this will only be one
-	static EShaderPlatform GetRuntimeShaderPlatform(void);
+	static RENDERCORE_API EShaderPlatform GetRuntimeShaderPlatform(void);
 
 	// Safely assign the hash to a shader object
-	static void SafeAssignHash(FRHIShader* InShader, const FSHAHash& Hash);
+	static RENDERCORE_API void SafeAssignHash(FRHIShader* InShader, const FSHAHash& Hash);
 
 	// Delegate called whenever shader code is requested.
-	static FDelegateHandle RegisterSharedShaderCodeRequestDelegate_Handle(const FSharedShaderCodeRequest::FDelegate& Delegate);
-	static void UnregisterSharedShaderCodeRequestDelegate_Handle(FDelegateHandle Handle);
+	static RENDERCORE_API FDelegateHandle RegisterSharedShaderCodeRequestDelegate_Handle(const FSharedShaderCodeRequest::FDelegate& Delegate);
+	static RENDERCORE_API void UnregisterSharedShaderCodeRequestDelegate_Handle(FDelegateHandle Handle);
 };
 
 #if WITH_EDITOR
 class ITargetPlatform;
 
-struct RENDERCORE_API FShaderLibraryCooker
+struct FShaderLibraryCooker
 {
 	// Initialize the library cooker
-	static void InitForCooking(bool bNativeFormat);
+	static RENDERCORE_API void InitForCooking(bool bNativeFormat);
 	// Shutdown the library cooker
-	static void Shutdown();
+	static RENDERCORE_API void Shutdown();
 
 	// Clean the cook directories
-	static void CleanDirectories(TArray<FName> const& ShaderFormats);
+	static RENDERCORE_API void CleanDirectories(TArray<FName> const& ShaderFormats);
 
 	struct FShaderFormatDescriptor
 	{
@@ -278,43 +278,43 @@ struct RENDERCORE_API FShaderLibraryCooker
 	 * @param Name is a high level description of the library (usually a project name or "Global")
 	 * @return true if successful
 	 */
-	static bool BeginCookingLibrary(FString const& Name);
+	static RENDERCORE_API bool BeginCookingLibrary(FString const& Name);
 
 	/**
 	 * Close a named library.
 	 *  For cooking, after this point any AddShaderCode/AddShaderPipeline calls will be invalid until OpenLibrary is called again.
 	 */
-	static void EndCookingLibrary(FString const& Name);
+	static RENDERCORE_API void EndCookingLibrary(FString const& Name);
 
 	/**
 	 * Whether storing shaders in the shader library is enabled
 	 */
-	static bool IsShaderLibraryEnabled();
+	static RENDERCORE_API bool IsShaderLibraryEnabled();
 
 	// Specify the shader formats to cook and which ones needs stable keys. Provide an array of FShaderFormatDescriptors
-	static void CookShaderFormats(TArray<FShaderFormatDescriptor> const& ShaderFormats);
+	static RENDERCORE_API void CookShaderFormats(TArray<FShaderFormatDescriptor> const& ShaderFormats);
 
 	// At cook time, add shader code to collection
-	static bool AddShaderCode(EShaderPlatform ShaderPlatform, const FShaderMapResourceCode* Code, const FShaderMapAssetPaths& AssociatedAssets);
+	static RENDERCORE_API bool AddShaderCode(EShaderPlatform ShaderPlatform, const FShaderMapResourceCode* Code, const FShaderMapAssetPaths& AssociatedAssets);
 
 #if WITH_EDITOR
 	/** Called from a CookWorker to send all contents of the ShaderLibrary to the CookDirector */
-	static void CopyToCompactBinaryAndClear(FCbWriter& Writer, bool& bOutHasData, bool& bOutRanOutOfRoom, int64 MaxShaderSize);
+	static RENDERCORE_API void CopyToCompactBinaryAndClear(FCbWriter& Writer, bool& bOutHasData, bool& bOutRanOutOfRoom, int64 MaxShaderSize);
 
 	/** Called On the CookDirector to receive ShaderLibrary contents from a CookWorker */
-	static bool AppendFromCompactBinary(FCbFieldView Field);
+	static RENDERCORE_API bool AppendFromCompactBinary(FCbFieldView Field);
 #endif
 
 	// We check this early in the callstack to avoid creating a bunch of FName and keys and things we will never save anyway. 
 	// Pass the shader platform to check or EShaderPlatform::SP_NumPlatforms to check if any of the registered types require
 	// stable keys.
-	static bool NeedsShaderStableKeys(EShaderPlatform ShaderPlatform);
+	static RENDERCORE_API bool NeedsShaderStableKeys(EShaderPlatform ShaderPlatform);
 
 	// At cook time, add the human readable key value information
-	static void AddShaderStableKeyValue(EShaderPlatform ShaderPlatform, FStableShaderKeyAndValue& StableKeyValue);
+	static RENDERCORE_API void AddShaderStableKeyValue(EShaderPlatform ShaderPlatform, FStableShaderKeyAndValue& StableKeyValue);
 
 	/** Finishes collection of data that should be in the named code library. This includes loading data from a previous iterative cook. */
-	static void FinishPopulateShaderLibrary(const ITargetPlatform* TargetPlatform, FString const& Name, FString const& SandboxDestinationPath,
+	static RENDERCORE_API void FinishPopulateShaderLibrary(const ITargetPlatform* TargetPlatform, FString const& Name, FString const& SandboxDestinationPath,
 		FString const& SandboxMetadataPath);
 
 	/**
@@ -330,7 +330,7 @@ struct RENDERCORE_API FShaderLibraryCooker
 	 * @param bOutHasData Reports whether any files were written to PlatformSCLCSVPaths
 	 * @return true if successful
 	 */
-	static bool SaveShaderLibraryWithoutChunking(const ITargetPlatform* TargetPlatform, FString const& Name, FString const& SandboxDestinationPath,
+	static RENDERCORE_API bool SaveShaderLibraryWithoutChunking(const ITargetPlatform* TargetPlatform, FString const& Name, FString const& SandboxDestinationPath,
 		FString const& SandboxMetadataPath, TArray<FString>& PlatformSCLCSVPaths, FString& OutErrorMessage, bool& bOutHasData);
 
 	/** 
@@ -346,13 +346,13 @@ struct RENDERCORE_API FShaderLibraryCooker
 	 * @param bOutHasData Reports whether any files were written to OutChunkFilenames
 	 * @return true if successful
 	 */
-	static bool SaveShaderLibraryChunk(int32 ChunkId, const TSet<FName>& InPackagesInChunk, const ITargetPlatform* TargetPlatform,
+	static RENDERCORE_API bool SaveShaderLibraryChunk(int32 ChunkId, const TSet<FName>& InPackagesInChunk, const ITargetPlatform* TargetPlatform,
 		const FString& SandboxDestinationPath, const FString& SandboxMetadataPath, TArray<FString>& OutChunkFilenames, bool& bOutHasData);
 
 	// Dump collected stats for each shader platform
-	static void DumpShaderCodeStats();
+	static RENDERCORE_API void DumpShaderCodeStats();
 
 	// Create a smaller 'patch' library that only contains data from 'NewMetaDataDir' not contained in any of 'OldMetaDataDirs'
-	static bool CreatePatchLibrary(TArray<FString> const& OldMetaDataDirs, FString const& NewMetaDataDir, FString const& OutDir, bool bNativeFormat, bool bNeedsDeterministicOrder);
+	static RENDERCORE_API bool CreatePatchLibrary(TArray<FString> const& OldMetaDataDirs, FString const& NewMetaDataDir, FString const& OutDir, bool bNativeFormat, bool bNeedsDeterministicOrder);
 };
 #endif
diff --git a/Engine/Source/Runtime/RenderCore/Public/ShaderCompilerCore.h b/Engine/Source/Runtime/RenderCore/Public/ShaderCompilerCore.h
--- a/Engine/Source/Runtime/RenderCore/Public/ShaderCompilerCore.h
+++ b/Engine/Source/Runtime/RenderCore/Public/ShaderCompilerCore.h
@@ -625,7 +625,7 @@ struct FShaderCompilerOutput
 	}
 };
 
-struct RENDERCORE_API FSCWErrorCode
+struct FSCWErrorCode
 {
 	enum ECode : int32
 	{
@@ -650,19 +650,19 @@ struct RENDERCORE_API FSCWErrorCode
 	Call Reset first before setting a new value.
 	Returns true on success, otherwise the error code has already been set.
 	*/
-	static void Report(ECode Code, const FStringView& Info = {});
+	static RENDERCORE_API void Report(ECode Code, const FStringView& Info = {});
 
 	/** Resets the global SCW error code to NotSet. */
-	static void Reset();
+	static RENDERCORE_API void Reset();
 
 	/** Returns the global SCW error code. */
-	static ECode Get();
+	static RENDERCORE_API ECode Get();
 
 	/** Returns the global SCW error code information string. Empty string if not set. */
-	static const FString& GetInfo();
+	static RENDERCORE_API const FString& GetInfo();
 
 	/** Returns true if the SCW global error code has been set. Equivalent to 'Get() != NotSet'. */
-	static bool IsSet();
+	static RENDERCORE_API bool IsSet();
 };
 
 #if PLATFORM_WINDOWS
diff --git a/Engine/Source/Runtime/RenderCore/Public/ShaderCompilerJobTypes.h b/Engine/Source/Runtime/RenderCore/Public/ShaderCompilerJobTypes.h
--- a/Engine/Source/Runtime/RenderCore/Public/ShaderCompilerJobTypes.h
+++ b/Engine/Source/Runtime/RenderCore/Public/ShaderCompilerJobTypes.h
@@ -219,10 +219,10 @@ public:
 
 	RENDERCORE_API const FString& GetFinalSource() const;
 
-	RENDERCORE_API FShaderCompileJob() : FShaderCommonCompileJob(Type, 0u, 0u, EShaderCompileJobPriority::Num)
+	FShaderCompileJob() : FShaderCommonCompileJob(Type, 0u, 0u, EShaderCompileJobPriority::Num)
 	{}
 
-	RENDERCORE_API FShaderCompileJob(uint32 InHash, uint32 InId, EShaderCompileJobPriority InPriroity, const FShaderCompileJobKey& InKey) :
+	FShaderCompileJob(uint32 InHash, uint32 InId, EShaderCompileJobPriority InPriroity, const FShaderCompileJobKey& InKey) :
 		FShaderCommonCompileJob(Type, InHash, InId, InPriroity),
 		Key(InKey)
 	{}
diff --git a/Engine/Source/Runtime/RenderCore/Public/ShaderParameterMacros.h b/Engine/Source/Runtime/RenderCore/Public/ShaderParameterMacros.h
--- a/Engine/Source/Runtime/RenderCore/Public/ShaderParameterMacros.h
+++ b/Engine/Source/Runtime/RenderCore/Public/ShaderParameterMacros.h
@@ -1428,7 +1428,7 @@ extern RENDERCORE_API FShaderParametersMetadata* FindUniformBufferStructByShader
 	END_SHADER_PARAMETER_STRUCT()
 
 
-class RENDERCORE_API FShaderParametersMetadataRegistration
+class FShaderParametersMetadataRegistration
 {
 public:
 	FShaderParametersMetadataRegistration(TFunctionRef<const FShaderParametersMetadata* ()> LazyShaderParametersMetadataAccessor)
@@ -1437,11 +1437,11 @@ public:
 		GetInstances().Add(this);
 	}
 
-	static TArray<const FShaderParametersMetadataRegistration*>& GetInstances();
+	static RENDERCORE_API TArray<const FShaderParametersMetadataRegistration*>& GetInstances();
 
 	// Actually register all the instances and clear the array
-	static void CommitAll();
-	static bool IsReadyForRegistration();
+	static RENDERCORE_API void CommitAll();
+	static RENDERCORE_API bool IsReadyForRegistration();
 
 private:
 	TFunctionRef<const FShaderParametersMetadata* ()> LazyShaderParametersMetadataAccessor;
diff --git a/Engine/Source/Runtime/RenderCore/Public/ShaderParameterMetadata.h b/Engine/Source/Runtime/RenderCore/Public/ShaderParameterMetadata.h
--- a/Engine/Source/Runtime/RenderCore/Public/ShaderParameterMetadata.h
+++ b/Engine/Source/Runtime/RenderCore/Public/ShaderParameterMetadata.h
@@ -87,19 +87,19 @@ const TCHAR* GetShaderCodeResourceBindingTypeName(EShaderCodeResourceBindingType
 
 
 /** Simple class that registers a uniform buffer static slot in the constructor. */
-class RENDERCORE_API FUniformBufferStaticSlotRegistrar
+class FUniformBufferStaticSlotRegistrar
 {
 public:
-	FUniformBufferStaticSlotRegistrar(const TCHAR* InName);
+	RENDERCORE_API FUniformBufferStaticSlotRegistrar(const TCHAR* InName);
 };
 
 /** Registry for uniform buffer static slots. */
-class RENDERCORE_API FUniformBufferStaticSlotRegistry
+class FUniformBufferStaticSlotRegistry
 {
 public:
-	static FUniformBufferStaticSlotRegistry& Get();
+	static RENDERCORE_API FUniformBufferStaticSlotRegistry& Get();
 
-	void RegisterSlot(FName SlotName);
+	RENDERCORE_API void RegisterSlot(FName SlotName);
 
 	inline int32 GetSlotCount() const
 	{
@@ -135,7 +135,7 @@ private:
 };
 
 /** A uniform buffer struct. */
-class RENDERCORE_API FShaderParametersMetadata
+class FShaderParametersMetadata
 {
 public:
 	/** The use case of the uniform buffer structures. */
@@ -167,7 +167,7 @@ public:
 	static constexpr int32 kRootCBufferBindingIndex = 0;
 
 	/** A member of a shader parameter structure. */
-	class RENDERCORE_API FMember
+	class FMember
 	{
 	public:
 
@@ -250,7 +250,7 @@ public:
 			return ElementSize;
 		}
 
-		static void GenerateShaderParameterType(
+		static RENDERCORE_API void GenerateShaderParameterType(
 			FString& Result,
 			bool bSupportsPrecisionModifier,
 			EUniformBufferBaseType BaseType,
@@ -258,8 +258,8 @@ public:
 			uint32 NumRows,
 			uint32 NumColumns
 		);
-		void GenerateShaderParameterType(FString& Result, bool bSupportsPrecisionModifier) const;
-		void GenerateShaderParameterType(FString& Result, EShaderPlatform ShaderPlatform) const;
+		RENDERCORE_API void GenerateShaderParameterType(FString& Result, bool bSupportsPrecisionModifier) const;
+		RENDERCORE_API void GenerateShaderParameterType(FString& Result, EShaderPlatform ShaderPlatform) const;
 
 	private:
 
@@ -282,7 +282,7 @@ public:
 	 * traversal. bForceCompleteInitialization force to ignore the list for EUseCase::UniformBuffer and instead handle it like a standalone non
 	 * globally listed EUseCase::ShaderParameterStruct. This is required for the ShaderCompileWorker to deserialize them without side global effects.
 	 */
-	FShaderParametersMetadata(
+	RENDERCORE_API FShaderParametersMetadata(
 		EUseCase UseCase,
 		EUniformBufferBindingFlags InBindingFlags,
 		const TCHAR* InLayoutName,
@@ -297,16 +297,16 @@ public:
 		FRHIUniformBufferLayoutInitializer* OutLayoutInitializer = nullptr,
 		uint32 InUsageFlags = 0);
 
-	virtual ~FShaderParametersMetadata();
+	RENDERCORE_API virtual ~FShaderParametersMetadata();
 
-	void GetNestedStructs(TArray<const FShaderParametersMetadata*>& OutNestedStructs) const;
+	RENDERCORE_API void GetNestedStructs(TArray<const FShaderParametersMetadata*>& OutNestedStructs) const;
 
 #if WITH_EDITOR
-	void AddResourceTableEntries(FShaderResourceTableMap& ResourceTableMap, TMap<FString, FUniformBufferEntry>& UniformBufferMap) const;
+	RENDERCORE_API void AddResourceTableEntries(FShaderResourceTableMap& ResourceTableMap, TMap<FString, FUniformBufferEntry>& UniformBufferMap) const;
 
 	PRAGMA_DISABLE_DEPRECATION_WARNINGS
 	UE_DEPRECATED(5.3, "Resource table entries are now stored in FShaderResourceTableMap, rather than a TMap.")
-	void AddResourceTableEntries(TMap<FString, FResourceTableEntry>& ResourceTableMap, TMap<FString, FUniformBufferEntry>& UniformBufferMap) const;
+	RENDERCORE_API void AddResourceTableEntries(TMap<FString, FResourceTableEntry>& ResourceTableMap, TMap<FString, FUniformBufferEntry>& UniformBufferMap) const;
 	PRAGMA_ENABLE_DEPRECATION_WARNINGS
 #endif
 
@@ -360,21 +360,21 @@ public:
 #endif // WITH_EDITOR
 
 	/** Find a member for a given offset. */
-	void FindMemberFromOffset(
+	RENDERCORE_API void FindMemberFromOffset(
 		uint16 MemberOffset,
 		const FShaderParametersMetadata** OutContainingStruct,
 		const FShaderParametersMetadata::FMember** OutMember,
 		int32* ArrayElementId, FString* NamePrefix) const;
 
 	/** Returns the full C++ member name from it's byte offset in the structure. */
-	FString GetFullMemberCodeName(uint16 MemberOffset) const;
+	RENDERCORE_API FString GetFullMemberCodeName(uint16 MemberOffset) const;
 
-	static TLinkedList<FShaderParametersMetadata*>*& GetStructList();
+	static RENDERCORE_API TLinkedList<FShaderParametersMetadata*>*& GetStructList();
 	/** Speed up finding the uniform buffer by its name */
-	static TMap<FHashedName, FShaderParametersMetadata*>& GetNameStructMap();
+	static RENDERCORE_API TMap<FHashedName, FShaderParametersMetadata*>& GetNameStructMap();
 
 	/** Initialize all the global shader parameter structs. */
-	static void InitializeAllUniformBufferStructs();
+	static RENDERCORE_API void InitializeAllUniformBufferStructs();
 
 	/** Returns a hash about the entire layout of the structure. */
 	uint32 GetLayoutHash() const
@@ -462,9 +462,9 @@ private:
 	/** Additional flags for how to use the buffer */
 	uint32 UsageFlags = 0;
 
-	void InitializeLayout(FRHIUniformBufferLayoutInitializer* OutLayoutInitializer = nullptr);
+	RENDERCORE_API void InitializeLayout(FRHIUniformBufferLayoutInitializer* OutLayoutInitializer = nullptr);
 
 #if WITH_EDITOR
-	void InitializeUniformBufferDeclaration();
+	RENDERCORE_API void InitializeUniformBufferDeclaration();
 #endif
-};
\ No newline at end of file
+};
diff --git a/Engine/Source/Runtime/RenderCore/Public/ShaderParameterMetadataBuilder.h b/Engine/Source/Runtime/RenderCore/Public/ShaderParameterMetadataBuilder.h
--- a/Engine/Source/Runtime/RenderCore/Public/ShaderParameterMetadataBuilder.h
+++ b/Engine/Source/Runtime/RenderCore/Public/ShaderParameterMetadataBuilder.h
@@ -11,7 +11,7 @@
 #include "ShaderParameterMetadata.h"
 #include "Templates/AlignmentTemplates.h"
 
-class RENDERCORE_API FShaderParametersMetadataBuilder
+class FShaderParametersMetadataBuilder
 {
 public:
 	FShaderParametersMetadataBuilder() {}
@@ -87,7 +87,7 @@ public:
 		AddReferencedStruct(Name, TShaderParameterStructTypeInfo<T>::GetStructMetadata(), Precision);
 	}
 
-	void AddReferencedStruct(
+	RENDERCORE_API void AddReferencedStruct(
 		const TCHAR* Name,
 		const FShaderParametersMetadata* StructMetadata,
 		EShaderPrecisionModifier::Type Precision = EShaderPrecisionModifier::Float
@@ -101,7 +101,7 @@ public:
 		AddIncludedStruct(TShaderParameterStructTypeInfo<T>::GetStructMetadata(), Precision);
 	}
 
-	void AddIncludedStruct(
+	RENDERCORE_API void AddIncludedStruct(
 		const FShaderParametersMetadata* StructMetadata,
 		EShaderPrecisionModifier::Type Precision = EShaderPrecisionModifier::Float
 	);
@@ -134,31 +134,31 @@ public:
 		return ThisMemberOffset;
 	}
 
-	uint32 AddNestedStruct(
+	RENDERCORE_API uint32 AddNestedStruct(
 		const TCHAR* Name,
 		const FShaderParametersMetadata* StructMetadata,
 		EShaderPrecisionModifier::Type Precision = EShaderPrecisionModifier::Float
 		);
 
-	void AddBufferSRV(
+	RENDERCORE_API void AddBufferSRV(
 		const TCHAR* Name,
 		const TCHAR* ShaderType,
 		EShaderPrecisionModifier::Type Precision = EShaderPrecisionModifier::Float
 		);
 
-	void AddBufferUAV(
+	RENDERCORE_API void AddBufferUAV(
 		const TCHAR* Name,
 		const TCHAR* ShaderType,
 		EShaderPrecisionModifier::Type Precision = EShaderPrecisionModifier::Float
 		);
 
-	void AddRDGBufferSRV(
+	RENDERCORE_API void AddRDGBufferSRV(
 		const TCHAR* Name,
 		const TCHAR* ShaderType,
 		EShaderPrecisionModifier::Type Precision = EShaderPrecisionModifier::Float
 		);
 
-	void AddRDGBufferUAV(
+	RENDERCORE_API void AddRDGBufferUAV(
 		const TCHAR* Name,
 		const TCHAR* ShaderType,
 		EShaderPrecisionModifier::Type Precision = EShaderPrecisionModifier::Float
@@ -171,12 +171,12 @@ public:
 
 	uint32 GetNextMemberOffset() const { return NextMemberOffset; }
 
-	FShaderParametersMetadata* Build(
+	RENDERCORE_API FShaderParametersMetadata* Build(
 		FShaderParametersMetadata::EUseCase UseCase,
 		const TCHAR* ShaderParameterName
 		);
 
-	FShaderParametersMetadata* Build(
+	RENDERCORE_API FShaderParametersMetadata* Build(
 		FShaderParametersMetadata::EUseCase InUseCase,
 		EUniformBufferBindingFlags InBindingFlags,
 		const TCHAR* InLayoutName,
@@ -193,4 +193,4 @@ public:
 private:
 	TArray<FShaderParametersMetadata::FMember> Members;
 	uint32 NextMemberOffset = 0;
-};
\ No newline at end of file
+};
diff --git a/Engine/Source/Runtime/RenderCore/Public/ShaderParameterParser.h b/Engine/Source/Runtime/RenderCore/Public/ShaderParameterParser.h
--- a/Engine/Source/Runtime/RenderCore/Public/ShaderParameterParser.h
+++ b/Engine/Source/Runtime/RenderCore/Public/ShaderParameterParser.h
@@ -21,7 +21,7 @@ enum class EBindlessConversionType : uint8
 };
 
 /** Validates and moves all the shader loose data parameter defined in the root scope of the shader into the root uniform buffer. */
-class RENDERCORE_API FShaderParameterParser
+class FShaderParameterParser
 {
 public:
 	struct FParsedShaderParameter
@@ -64,40 +64,40 @@ public:
 		friend class FShaderParameterParser;
 	};
 
-	FShaderParameterParser();
+	RENDERCORE_API FShaderParameterParser();
 
 	UE_DEPRECATED(5.3, "Use FShaderParameterParser constructor which accepts FShaderCompilerFlags")
-	FShaderParameterParser(const TCHAR* InConstantBufferType);
+	RENDERCORE_API FShaderParameterParser(const TCHAR* InConstantBufferType);
 	UE_DEPRECATED(5.3, "Use FShaderParameterParser constructor which accepts FShaderCompilerFlags")
-	FShaderParameterParser(const TCHAR* InConstantBufferType, TConstArrayView<const TCHAR*> InExtraSRVTypes, TConstArrayView<const TCHAR*> InExtraUAVTypes);
+	RENDERCORE_API FShaderParameterParser(const TCHAR* InConstantBufferType, TConstArrayView<const TCHAR*> InExtraSRVTypes, TConstArrayView<const TCHAR*> InExtraUAVTypes);
 
-	FShaderParameterParser(
+	RENDERCORE_API FShaderParameterParser(
 		FShaderCompilerFlags CompilerFlags,
 		const TCHAR* InConstantBufferType = nullptr,
 		TConstArrayView<const TCHAR*> InExtraSRVTypes = {},
 		TConstArrayView<const TCHAR*> InExtraUAVTypes = {});
 
-	virtual ~FShaderParameterParser();
+	RENDERCORE_API virtual ~FShaderParameterParser();
 
 	FShaderParameterParser& operator=(FShaderParameterParser&&) = default;
 
 	static constexpr const TCHAR* kBindlessResourcePrefix = TEXT("BindlessResource_");
 	static constexpr const TCHAR* kBindlessSamplerPrefix = TEXT("BindlessSampler_");
 
-	static EShaderParameterType ParseParameterType(FStringView InType, TConstArrayView<const TCHAR*> InExtraSRVTypes, TConstArrayView<const TCHAR*> InExtraUAVTypes);
-	static EShaderParameterType ParseAndRemoveBindlessParameterPrefix(FStringView& InName);
-	static EShaderParameterType ParseAndRemoveBindlessParameterPrefix(FString& InName);
-	static bool RemoveBindlessParameterPrefix(FString& InName);
+	static RENDERCORE_API EShaderParameterType ParseParameterType(FStringView InType, TConstArrayView<const TCHAR*> InExtraSRVTypes, TConstArrayView<const TCHAR*> InExtraUAVTypes);
+	static RENDERCORE_API EShaderParameterType ParseAndRemoveBindlessParameterPrefix(FStringView& InName);
+	static RENDERCORE_API EShaderParameterType ParseAndRemoveBindlessParameterPrefix(FString& InName);
+	static RENDERCORE_API bool RemoveBindlessParameterPrefix(FString& InName);
 
 	/** Parses the preprocessed shader code and applies the necessary modifications to it. */
 	UE_DEPRECATED(5.3, "Use ParseAndModify overload accepting array of FShaderCompilerError instead of passing FShaderCompilerOutput")
-	bool ParseAndModify(
+	RENDERCORE_API bool ParseAndModify(
 		const FShaderCompilerInput& CompilerInput,
 		FShaderCompilerOutput& CompilerOutput,
 		FString& PreprocessedShaderSource
 	);
 
-	bool ParseAndModify(
+	RENDERCORE_API bool ParseAndModify(
 		const FShaderCompilerInput& CompilerInput,
 		TArray<FShaderCompilerError>& OutErrors,
 		FString& PreprocessedShaderSource
@@ -123,7 +123,7 @@ public:
 	}
 
 	/** Validates the shader parameter in code is compatible with the shader parameter structure. */
-	void ValidateShaderParameterType(
+	RENDERCORE_API void ValidateShaderParameterType(
 		const FShaderCompilerInput& CompilerInput,
 		const FString& ShaderBindingName,
 		int32 ReflectionOffset,
@@ -142,7 +142,7 @@ public:
 	}
 
 	/** Validates shader parameter map is compatible with the shader parameter structure. */
-	void ValidateShaderParameterTypes(
+	RENDERCORE_API void ValidateShaderParameterTypes(
 		const FShaderCompilerInput& CompilerInput,
 		bool bPlatformSupportsPrecisionModifier,
 		FShaderCompilerOutput& CompilerOutput) const;
@@ -164,33 +164,33 @@ public:
 
 protected:
 	/** Parses the preprocessed shader code */
-	bool ParseParameters(
+	RENDERCORE_API bool ParseParameters(
 		const FShaderParametersMetadata* RootParametersStructure,
 		TArray<FShaderCompilerError>& OutErrors
 	);
 
-	void RemoveMovingParametersFromSource(
+	RENDERCORE_API void RemoveMovingParametersFromSource(
 		FString& PreprocessedShaderSource
 	);
 
 	/** Converts parsed parameters into their bindless forms. */
-	void ApplyBindlessModifications(
+	RENDERCORE_API void ApplyBindlessModifications(
 		FString& PreprocessedShaderSource
 	);
 
 	/** Moves parsed parameters into the root constant buffer. */
-	bool MoveShaderParametersToRootConstantBuffer(
+	RENDERCORE_API bool MoveShaderParametersToRootConstantBuffer(
 		const FShaderParametersMetadata* RootParametersStructure,
 		FString& PreprocessedShaderSource
 	);
 
-	void ExtractFileAndLine(int32 PragamLineoffset, int32 LineOffset, FString& OutFile, FString& OutLine) const;
+	RENDERCORE_API void ExtractFileAndLine(int32 PragamLineoffset, int32 LineOffset, FString& OutFile, FString& OutLine) const;
 
 	/**
 	* Generates shader source code to declare a bindless resource or sampler (for automatic bindless conversion).
 	* May be overriden to allow custom implementations for different platforms.
 	*/
-	virtual FString GenerateBindlessParameterDeclaration(const FParsedShaderParameter& ParsedParameter) const;
+	RENDERCORE_API virtual FString GenerateBindlessParameterDeclaration(const FParsedShaderParameter& ParsedParameter) const;
 
 	const TCHAR* ConstantBufferType = nullptr;
 
diff --git a/Engine/Source/Runtime/RenderCore/Public/ShaderPipelineCache.h b/Engine/Source/Runtime/RenderCore/Public/ShaderPipelineCache.h
--- a/Engine/Source/Runtime/RenderCore/Public/ShaderPipelineCache.h
+++ b/Engine/Source/Runtime/RenderCore/Public/ShaderPipelineCache.h
@@ -75,7 +75,7 @@ struct FRasterizerStateInitializerRHI;
  * - FShaderCodeLibrary must be enabled via Project Settings > Packaging > "Share Material Shader Code".
  * - Enabling "Native Shader Libraries" is optional, but strongly preferred for Metal.
  */
-class RENDERCORE_API FShaderPipelineCache : public FTickableObjectRenderThread
+class FShaderPipelineCache : public FTickableObjectRenderThread
 {
 	// the FShaderPipelineCacheTask class represents a pipeline precompile of a single PSO file cache.
 	friend class FShaderPipelineCacheTask;
@@ -84,12 +84,12 @@ public:
 	 * Initializes the shader pipeline cache for the desired platform, called by the engine.
 	 * The shader platform is used only to load the initial pipeline cache and can be changed by closing & reopening the cache if necessary.
 	 */
-	static void Initialize(EShaderPlatform Platform);
+	static RENDERCORE_API void Initialize(EShaderPlatform Platform);
 	/** Terminates the shader pipeline cache, called by the engine. */
-	static void Shutdown();
+	static RENDERCORE_API void Shutdown();
 
 	/** Pauses precompilation. */
-	static void PauseBatching();
+	static RENDERCORE_API void PauseBatching();
 	
 	enum class BatchMode
 	{
@@ -99,50 +99,50 @@ public:
 	};
 	
 	/** Sets the precompilation batching mode. */
-	static void SetBatchMode(BatchMode Mode);
+	static RENDERCORE_API void SetBatchMode(BatchMode Mode);
 	
 	/** Resumes precompilation batching. */
-	static void ResumeBatching();
+	static RENDERCORE_API void ResumeBatching();
 	
 	/** true if pipeline cache(s) are precompiling. */
-	static bool IsPrecompiling();
+	static RENDERCORE_API bool IsPrecompiling();
 
 	/** Returns the number of pipelines waiting for precompilation of the current PSOFC task. 
 		if there are multiple PSOFCs pending this will return a minimum value of 1. It may increase as subsequent caches are processed.*/
-	static uint32 NumPrecompilesRemaining();
+	static RENDERCORE_API uint32 NumPrecompilesRemaining();
     
 	/** Opens the shader pipeline cache file with either the LastOpened setting if available, or the project name otherwise */
-	static bool OpenPipelineFileCache(EShaderPlatform Platform);
+	static RENDERCORE_API bool OpenPipelineFileCache(EShaderPlatform Platform);
 
  	/** Opens the user pipeline cache file with either the LastOpened setting if available, or the project name otherwise */
- 	static bool OpenUserPipelineFileCache(EShaderPlatform Platform);
+ 	static RENDERCORE_API bool OpenUserPipelineFileCache(EShaderPlatform Platform);
 	
 	/** Saves the current shader pipeline cache to disk using one of the defined save modes, Fast uses an incremental approach whereas Slow will consolidate all data into the file. */
-	static bool SavePipelineFileCache(FPipelineFileCacheManager::SaveMode Mode);
+	static RENDERCORE_API bool SavePipelineFileCache(FPipelineFileCacheManager::SaveMode Mode);
 	
 	/** Closes the existing pipeline cache, allowing it to be reopened with a different file and/or shader platform. Will implicitly invoke a Fast Save. */
-	static void CloseUserPipelineFileCache();
+	static RENDERCORE_API void CloseUserPipelineFileCache();
 
-	static int32 GetGameVersionForPSOFileCache();
+	static RENDERCORE_API int32 GetGameVersionForPSOFileCache();
 	
 	/**
 	 * Set the current PSO Game Usage Mask and comparsion function .  Returns true if this mask is different from the old mask or false if not or the cache system is disabled or if the mask feature is disabled
 	 * Any new PSO's found will be logged with this value or existing PSO should have their masks updated.  See FPipelineFileCacheManager for more details.
 	 */
-	static bool SetGameUsageMaskWithComparison(uint64 Mask, FPSOMaskComparisonFn InComparisonFnPtr);
+	static RENDERCORE_API bool SetGameUsageMaskWithComparison(uint64 Mask, FPSOMaskComparisonFn InComparisonFnPtr);
 
-    static bool IsBatchingPaused();
+    static RENDERCORE_API bool IsBatchingPaused();
     
-	FShaderPipelineCache(EShaderPlatform Platform);
-	virtual ~FShaderPipelineCache();
+	RENDERCORE_API FShaderPipelineCache(EShaderPlatform Platform);
+	RENDERCORE_API virtual ~FShaderPipelineCache();
 
-	bool IsTickable() const;
+	RENDERCORE_API bool IsTickable() const;
 
-	void Tick( float DeltaTime );
+	RENDERCORE_API void Tick( float DeltaTime );
 
-	bool NeedsRenderingResumedForRenderingThreadTick() const;
+	RENDERCORE_API bool NeedsRenderingResumedForRenderingThreadTick() const;
 	
-	TStatId GetStatId() const;
+	RENDERCORE_API TStatId GetStatId() const;
 	
 	enum ELibraryState
 	{
@@ -152,7 +152,7 @@ public:
 	};
 	
     /** Called by FShaderCodeLibrary to notify us that the shader code library state changed and shader availability will need to be re-evaluated */
-    static void ShaderLibraryStateChanged(ELibraryState State, EShaderPlatform Platform, FString const& Name, int32 ComponentID);
+    static RENDERCORE_API void ShaderLibraryStateChanged(ELibraryState State, EShaderPlatform Platform, FString const& Name, int32 ComponentID);
 
 	/** Allows handlers of FShaderCacheOpenedDelegate to send state to begin and complete pre-compilation delegates */
 	class FShaderCachePrecompileContext
@@ -221,23 +221,23 @@ public:
 private:
 	static inline FShaderPipelineCache* Get() { return ShaderPipelineCache;	}
 
-	bool OpenPipelineFileCacheInternal(bool bUserCache, const FString& PSOCacheKey, const FString& CacheName, EShaderPlatform Platform);
+	RENDERCORE_API bool OpenPipelineFileCacheInternal(bool bUserCache, const FString& PSOCacheKey, const FString& CacheName, EShaderPlatform Platform);
 
-	void BeginNextPrecompileCacheTask();
-	void BeginNextPrecompileCacheTaskInternal();
-	bool ReadyForAutoSave() const;
+	RENDERCORE_API void BeginNextPrecompileCacheTask();
+	RENDERCORE_API void BeginNextPrecompileCacheTaskInternal();
+	RENDERCORE_API bool ReadyForAutoSave() const;
 
-	void ClearPendingPrecompileTaskQueue();
-	void EnqueuePendingPrecompileTask(const FString& PSOCacheKey);
+	RENDERCORE_API void ClearPendingPrecompileTaskQueue();
+	RENDERCORE_API void EnqueuePendingPrecompileTask(const FString& PSOCacheKey);
 
-	FRHIBlendState* GetOrCreateBlendState(const FBlendStateInitializerRHI& Initializer);
-	FRHIRasterizerState* GetOrCreateRasterizerState(const FRasterizerStateInitializerRHI& Initializer);
-	FRHIDepthStencilState* GetOrCreateDepthStencilState(const FDepthStencilStateInitializerRHI& Initializer);
+	RENDERCORE_API FRHIBlendState* GetOrCreateBlendState(const FBlendStateInitializerRHI& Initializer);
+	RENDERCORE_API FRHIRasterizerState* GetOrCreateRasterizerState(const FRasterizerStateInitializerRHI& Initializer);
+	RENDERCORE_API FRHIDepthStencilState* GetOrCreateDepthStencilState(const FDepthStencilStateInitializerRHI& Initializer);
 	
 private:
-	static FShaderPipelineCache* ShaderPipelineCache;
-	static uint32 BatchSize;
-	static float BatchTime;
+	static RENDERCORE_API FShaderPipelineCache* ShaderPipelineCache;
+	static RENDERCORE_API uint32 BatchSize;
+	static RENDERCORE_API float BatchTime;
 	bool bPaused;
     int32 PausedCount;
 	double LastAutoSaveTime = 0.0;
@@ -248,7 +248,7 @@ private:
 	/** Development tool that can describe PSOs in non-Shipping builds. */
 	UE::ShaderUtils::FShaderStableKeyDebugInfoReader ShaderHashToStableKey;
 
-	static FString UserCacheTaskKey;
+	static RENDERCORE_API FString UserCacheTaskKey;
 
 	struct FPipelineCachePSOContext
 	{
@@ -277,11 +277,11 @@ private:
 	mutable FCriticalSection Mutex;
 	FGraphEventRef NextPrecompileCacheTask;
 
-	static FShaderCachePreOpenDelegate OnCachePreOpen;
-    static FShaderCacheOpenedDelegate OnCachedOpened;
-    static FShaderCacheClosedDelegate OnCachedClosed;
-	static FShaderPrecompilationBeginDelegate OnPrecompilationBegin;
-	static FShaderPrecompilationCompleteDelegate OnPrecompilationComplete;
+	static RENDERCORE_API FShaderCachePreOpenDelegate OnCachePreOpen;
+    static RENDERCORE_API FShaderCacheOpenedDelegate OnCachedOpened;
+    static RENDERCORE_API FShaderCacheClosedDelegate OnCachedClosed;
+	static RENDERCORE_API FShaderPrecompilationBeginDelegate OnPrecompilationBegin;
+	static RENDERCORE_API FShaderPrecompilationCompleteDelegate OnPrecompilationComplete;
 
 	TMap<FBlendStateInitializerRHI, FRHIBlendState*> BlendStateCache;
 	TMap<FRasterizerStateInitializerRHI, FRHIRasterizerState*> RasterizerStateCache;
diff --git a/Engine/Source/Runtime/RenderCore/Public/ShaderPlatformCachedIniValue.h b/Engine/Source/Runtime/RenderCore/Public/ShaderPlatformCachedIniValue.h
--- a/Engine/Source/Runtime/RenderCore/Public/ShaderPlatformCachedIniValue.h
+++ b/Engine/Source/Runtime/RenderCore/Public/ShaderPlatformCachedIniValue.h
@@ -10,7 +10,7 @@
 RENDERCORE_API EShaderPlatform GetEditorShaderPlatform(EShaderPlatform ShaderPlatform);
 
 template<typename Type>
-struct RENDERCORE_API FShaderPlatformCachedIniValue
+struct FShaderPlatformCachedIniValue
 {
 	FShaderPlatformCachedIniValue(const TCHAR* InCVarName)
 		: CVarName(InCVarName)
diff --git a/Engine/Source/Runtime/RenderCore/Public/TickableObjectRenderThread.h b/Engine/Source/Runtime/RenderCore/Public/TickableObjectRenderThread.h
--- a/Engine/Source/Runtime/RenderCore/Public/TickableObjectRenderThread.h
+++ b/Engine/Source/Runtime/RenderCore/Public/TickableObjectRenderThread.h
@@ -13,7 +13,7 @@
  * This class provides common registration for render thread tickable objects. It is an
  * abstract base class requiring you to implement the Tick() method.
  */
-class RENDERCORE_API FTickableObjectRenderThread
+class FTickableObjectRenderThread
 {
 public:
 
@@ -35,8 +35,8 @@ public:
 		}
 	};
 
-	static FRenderingThreadTickableObjectsArray RenderingThreadTickableObjects;
-	static FRenderingThreadTickableObjectsArray RenderingThreadHighFrequencyTickableObjects;
+	static RENDERCORE_API FRenderingThreadTickableObjectsArray RenderingThreadTickableObjects;
+	static RENDERCORE_API FRenderingThreadTickableObjectsArray RenderingThreadHighFrequencyTickableObjects;
 
 	/**
 	 * Registers this instance with the static array of tickable objects.	
diff --git a/Engine/Source/Runtime/RenderCore/Public/UnifiedBuffer.h b/Engine/Source/Runtime/RenderCore/Public/UnifiedBuffer.h
--- a/Engine/Source/Runtime/RenderCore/Public/UnifiedBuffer.h
+++ b/Engine/Source/Runtime/RenderCore/Public/UnifiedBuffer.h
@@ -113,7 +113,7 @@ public:
 	RENDERCORE_API void Init( uint32 NumElements, uint32 InNumBytesPerElement, bool bInFloat4Buffer, const TCHAR* DebugName );
 
 	template<typename ResourceType>
-	RENDERCORE_API void ResourceUploadTo(FRHICommandList& RHICmdList, const ResourceType& DstBuffer, bool bFlush = false);
+	void ResourceUploadTo(FRHICommandList& RHICmdList, const ResourceType& DstBuffer, bool bFlush = false);
 	
 	UE_DEPRECATED(5.3, "Scattered uploads to FTextureRWBuffer will no longer supported in a future release. Use a supported destination buffer type instead.")
 	inline void ResourceUploadTo(FRHICommandList& RHICmdList, const FTextureRWBuffer& DstBuffer, bool bFlush = false)
@@ -326,12 +326,12 @@ protected:
 	friend FRDGAsyncScatterUploadBuffer;
 };
 
-class RENDERCORE_API FRDGScatterUploader
+class FRDGScatterUploader
 	: public FRDGScatterUploadBase
 {
 public:
-	void Lock(FRHICommandListBase& RHICmdList);
-	void Unlock(FRHICommandListBase& RHICmdList);
+	RENDERCORE_API void Lock(FRHICommandListBase& RHICmdList);
+	RENDERCORE_API void Unlock(FRHICommandListBase& RHICmdList);
 
 	FRDGViewableResource* GetDstResource() const
 	{
@@ -380,24 +380,24 @@ inline void UnlockIfValid(FRHICommandListBase& RHICmdList, FRDGScatterUploader*
 	}
 }
 
-class RENDERCORE_API FRDGAsyncScatterUploadBuffer
+class FRDGAsyncScatterUploadBuffer
 {
 public:
-	FRDGScatterUploader* Begin(FRDGBuilder& GraphBuilder, FRDGViewableResource* DstResource, uint32 NumElements, uint32 NumBytesPerElement, const TCHAR* Name);
-	FRDGScatterUploader* BeginPreSized(FRDGBuilder& GraphBuilder, FRDGViewableResource* DstResource, uint32 NumElements, uint32 NumBytesPerElement, const TCHAR* Name);
+	RENDERCORE_API FRDGScatterUploader* Begin(FRDGBuilder& GraphBuilder, FRDGViewableResource* DstResource, uint32 NumElements, uint32 NumBytesPerElement, const TCHAR* Name);
+	RENDERCORE_API FRDGScatterUploader* BeginPreSized(FRDGBuilder& GraphBuilder, FRDGViewableResource* DstResource, uint32 NumElements, uint32 NumBytesPerElement, const TCHAR* Name);
 
-	void End(FRDGBuilder& GraphBuilder, FRDGScatterUploader* Uploader);
+	RENDERCORE_API void End(FRDGBuilder& GraphBuilder, FRDGScatterUploader* Uploader);
 
-	void Release();
+	RENDERCORE_API void Release();
 
-	uint32 GetNumBytes() const;
+	RENDERCORE_API uint32 GetNumBytes() const;
 
 private:
 	TRefCountPtr<FRDGPooledBuffer> ScatterBuffer;
 	TRefCountPtr<FRDGPooledBuffer> UploadBuffer;
 };
 
-class RENDERCORE_API FRDGScatterUploadBuffer
+class FRDGScatterUploadBuffer
 	: public FRDGScatterUploadBase
 {
 public:
@@ -413,14 +413,14 @@ public:
 	/**
 	 * Init with presized num scatters, expecting each to be set at a later point. Requires the user to keep track of the offsets to use.
 	 */
-	void InitPreSized(FRDGBuilder& GraphBuilder, uint32 NumElements, uint32 InNumBytesPerElement, bool bInFloat4Buffer, const TCHAR* DebugName);
+	RENDERCORE_API void InitPreSized(FRDGBuilder& GraphBuilder, uint32 NumElements, uint32 InNumBytesPerElement, bool bInFloat4Buffer, const TCHAR* DebugName);
 
 	/**
 	 * Init with pre-existing destination index data, performs a bulk-copy.
 	 */
-	void Init(FRDGBuilder& GraphBuilder, TArrayView<const uint32> ElementScatterOffsets, uint32 InNumBytesPerElement, bool bInFloat4Buffer, const TCHAR* DebugName);
+	RENDERCORE_API void Init(FRDGBuilder& GraphBuilder, TArrayView<const uint32> ElementScatterOffsets, uint32 InNumBytesPerElement, bool bInFloat4Buffer, const TCHAR* DebugName);
 
-	void Init(FRDGBuilder& GraphBuilder, uint32 NumElements, uint32 NumBytesPerElement, bool bInFloat4Buffer, const TCHAR* Name);
+	RENDERCORE_API void Init(FRDGBuilder& GraphBuilder, uint32 NumElements, uint32 NumBytesPerElement, bool bInFloat4Buffer, const TCHAR* Name);
 
 	inline void ResourceUploadTo(FRDGBuilder& GraphBuilder, FRDGBuffer* DstResource)
 	{
@@ -433,16 +433,16 @@ public:
 		ResourceUploadToInternal(GraphBuilder, DstResource);
 	}
 
-	void Release();
+	RENDERCORE_API void Release();
 
-	uint32 GetNumBytes() const;
+	RENDERCORE_API uint32 GetNumBytes() const;
 
 private:
-	void ResourceUploadToInternal(FRDGBuilder& GraphBuilder, FRDGViewableResource* DstResource);
-	void Reset();
+	RENDERCORE_API void ResourceUploadToInternal(FRDGBuilder& GraphBuilder, FRDGViewableResource* DstResource);
+	RENDERCORE_API void Reset();
 
 	TRefCountPtr<FRDGPooledBuffer> ScatterBuffer;
 	TRefCountPtr<FRDGPooledBuffer> UploadBuffer;
 
 	bool bFloat4Buffer = false;
-};
\ No newline at end of file
+};
diff --git a/Engine/Source/Runtime/RenderCore/Public/VertexFactory.h b/Engine/Source/Runtime/RenderCore/Public/VertexFactory.h
--- a/Engine/Source/Runtime/RenderCore/Public/VertexFactory.h
+++ b/Engine/Source/Runtime/RenderCore/Public/VertexFactory.h
@@ -134,9 +134,9 @@ ENUM_CLASS_FLAGS(EVertexFactoryFlags);
 /**
  * An interface to the parameter bindings for the vertex factory used by a shader.
  */
-class RENDERCORE_API FVertexFactoryShaderParameters
+class FVertexFactoryShaderParameters
 {
-	DECLARE_TYPE_LAYOUT(FVertexFactoryShaderParameters, NonVirtual);
+	DECLARE_EXPORTED_TYPE_LAYOUT(FVertexFactoryShaderParameters, RENDERCORE_API, NonVirtual);
 public:
 	void Bind(const class FShaderParameterMap& ParameterMap) {}
 
@@ -587,7 +587,7 @@ public:
 /**
  * Encapsulates a vertex data source which can be linked into a vertex shader.
  */
-class RENDERCORE_API FVertexFactory : public FRenderResource
+class FVertexFactory : public FRenderResource
 {
 public:
 	FVertexFactory(ERHIFeatureLevel::Type InFeatureLevel) 
@@ -601,9 +601,9 @@ public:
 
 	virtual FVertexFactoryType* GetType() const { return nullptr; }
 
-	void GetStreams(ERHIFeatureLevel::Type InFeatureLevel, EVertexInputStreamType VertexStreamType, FVertexInputStreamArray& OutVertexStreams) const;
+	RENDERCORE_API void GetStreams(ERHIFeatureLevel::Type InFeatureLevel, EVertexInputStreamType VertexStreamType, FVertexInputStreamArray& OutVertexStreams) const;
 
-	void OffsetInstanceStreams(uint32 InstanceOffset, EVertexInputStreamType VertexStreamType, FVertexInputStreamArray& VertexStreams) const;
+	RENDERCORE_API void OffsetInstanceStreams(uint32 InstanceOffset, EVertexInputStreamType VertexStreamType, FVertexInputStreamArray& VertexStreams) const;
 
 	/**
 	* Can be overridden by FVertexFactory subclasses to modify their compile environment just before compilation occurs.
@@ -621,7 +621,7 @@ public:
 	static void GetPSOPrecacheVertexFetchElements(EVertexInputStreamType VertexInputStreamType, FVertexDeclarationElementList& Elements) { checkNoEntry(); }
 
 	// FRenderResource interface.
-	virtual void ReleaseRHI();
+	RENDERCORE_API virtual void ReleaseRHI();
 
 	// Accessors.
 	FVertexDeclarationRHIRef& GetDeclaration() { return Declaration; }
@@ -681,7 +681,7 @@ protected:
 		PrimitiveIdStreamIndex[TranslatePrimitiveIdStreamIndex(InFeatureLevel, InputStreamType)] = static_cast<int8>(StreamIndex);
 	}
 	
-	bool AddPrimitiveIdStreamElement(EVertexInputStreamType InputStreamType, FVertexDeclarationElementList& Elements, uint8 AttributeIndex, uint8 AttributeIndex_Mobile);
+	RENDERCORE_API bool AddPrimitiveIdStreamElement(EVertexInputStreamType InputStreamType, FVertexDeclarationElementList& Elements, uint8 AttributeIndex, uint8 AttributeIndex_Mobile);
 
 	/**
 	 * Creates a vertex element for a vertex stream components.  Adds a unique stream index for the vertex buffer used by the component.
@@ -701,7 +701,7 @@ protected:
 	 * @param AttributeIndex - The attribute index to which the stream component is bound.
 	 * @return The vertex element which corresponds to Component.
 	 */
-	FVertexElement AccessStreamComponent(const FVertexStreamComponent& Component, uint8 AttributeIndex, EVertexInputStreamType InputStreamType);
+	RENDERCORE_API FVertexElement AccessStreamComponent(const FVertexStreamComponent& Component, uint8 AttributeIndex, EVertexInputStreamType InputStreamType);
 	
 	/**
 	 * Creates a vertex element for a vertex stream components.  Adds a unique stream index for the vertex buffer used by the component.
@@ -726,7 +726,7 @@ protected:
 	 * Initializes the vertex declaration.
 	 * @param Elements - The elements of the vertex declaration.
 	 */
-	void InitDeclaration(const FVertexDeclarationElementList& Elements, EVertexInputStreamType StreamType = EVertexInputStreamType::Default);
+	RENDERCORE_API void InitDeclaration(const FVertexDeclarationElementList& Elements, EVertexInputStreamType StreamType = EVertexInputStreamType::Default);
 
 	/**
 	 * Information needed to set a vertex stream.
diff --git a/Engine/Source/Runtime/RenderCore/Public/VisualizeTexture.h b/Engine/Source/Runtime/RenderCore/Public/VisualizeTexture.h
--- a/Engine/Source/Runtime/RenderCore/Public/VisualizeTexture.h
+++ b/Engine/Source/Runtime/RenderCore/Public/VisualizeTexture.h
@@ -28,32 +28,32 @@ class FWildcardString;
 
 #define SUPPORTS_VISUALIZE_TEXTURE (WITH_ENGINE && (!UE_BUILD_SHIPPING || WITH_EDITOR))
 
-class RENDERCORE_API FVisualizeTexture : public FRenderResource
+class FVisualizeTexture : public FRenderResource
 {
 public:
 	FVisualizeTexture() = default;
 
-	void ParseCommands(const TCHAR* Cmd, FOutputDevice &Ar);
+	RENDERCORE_API void ParseCommands(const TCHAR* Cmd, FOutputDevice &Ar);
 
-	void DebugLogOnCrash();
+	RENDERCORE_API void DebugLogOnCrash();
 
-	void GetTextureInfos_GameThread(TArray<FString>& Infos) const;
+	RENDERCORE_API void GetTextureInfos_GameThread(TArray<FString>& Infos) const;
 
 	/** Creates a new checkpoint (e.g. "SceneDepth@N") for the pooled render target. A null parameter is a no-op. */
 #if SUPPORTS_VISUALIZE_TEXTURE
-	void SetCheckPoint(FRDGBuilder& GraphBuilder, IPooledRenderTarget* PooledRenderTarget);
-	void SetCheckPoint(FRHICommandListImmediate& RHICmdList, IPooledRenderTarget* PooledRenderTarget);
+	RENDERCORE_API void SetCheckPoint(FRDGBuilder& GraphBuilder, IPooledRenderTarget* PooledRenderTarget);
+	RENDERCORE_API void SetCheckPoint(FRHICommandListImmediate& RHICmdList, IPooledRenderTarget* PooledRenderTarget);
 #else
 	inline void SetCheckPoint(FRDGBuilder& GraphBuilder, IPooledRenderTarget* PooledRenderTarget) {}
 	inline void SetCheckPoint(FRHICommandListImmediate& RHICmdList, IPooledRenderTarget* PooledRenderTarget) {}
 #endif
 
-	static FRDGTextureRef AddVisualizeTexturePass(
+	static RENDERCORE_API FRDGTextureRef AddVisualizeTexturePass(
 		FRDGBuilder& GraphBuilder,
 		class FGlobalShaderMap* ShaderMap,
 		const FRDGTextureRef InputTexture);
 
-	static FRDGTextureRef AddVisualizeTextureAlphaPass(
+	static RENDERCORE_API FRDGTextureRef AddVisualizeTextureAlphaPass(
 		FRDGBuilder& GraphBuilder,
 		class FGlobalShaderMap* ShaderMap,
 		const FRDGTextureRef InputTexture);
@@ -112,12 +112,12 @@ private:
 	};
 
 #if SUPPORTS_VISUALIZE_TEXTURE
-	static void DisplayHelp(FOutputDevice &Ar);
-	void DisplayPoolResourceListToLog(ESortBy SortBy);
-	void DisplayResourceListToLog(const TOptional<FWildcardString>& Wildcard);
+	static RENDERCORE_API void DisplayHelp(FOutputDevice &Ar);
+	RENDERCORE_API void DisplayPoolResourceListToLog(ESortBy SortBy);
+	RENDERCORE_API void DisplayResourceListToLog(const TOptional<FWildcardString>& Wildcard);
 
 	/** Determine whether a texture should be captured for debugging purposes and return the capture id if needed. */
-	TOptional<uint32> ShouldCapture(const TCHAR* DebugName, uint32 MipIndex);
+	RENDERCORE_API TOptional<uint32> ShouldCapture(const TCHAR* DebugName, uint32 MipIndex);
 
 	struct FConfig
 	{
@@ -136,7 +136,7 @@ private:
 	};
 
 	/** Adds a pass to visualize a texture. */
-	static FRDGTextureRef AddVisualizeTexturePass(
+	static RENDERCORE_API FRDGTextureRef AddVisualizeTexturePass(
 		FRDGBuilder& GraphBuilder,
 		class FGlobalShaderMap* ShaderMap,
 		const FRDGTextureRef InputTexture,
@@ -145,13 +145,13 @@ private:
 		uint32 CaptureId);
 
 	/** Create a pass capturing a texture. */
-	void CreateContentCapturePass(FRDGBuilder& GraphBuilder, FRDGTextureRef Texture, uint32 CaptureId);
+	RENDERCORE_API void CreateContentCapturePass(FRDGBuilder& GraphBuilder, FRDGTextureRef Texture, uint32 CaptureId);
 
-	void ReleaseRHI() override;
+	RENDERCORE_API void ReleaseRHI() override;
 
-	void Visualize(const FString& InName, TOptional<uint32> InVersion = {});
+	RENDERCORE_API void Visualize(const FString& InName, TOptional<uint32> InVersion = {});
 
-	uint32 GetVersionCount(const TCHAR* InName) const;
+	RENDERCORE_API uint32 GetVersionCount(const TCHAR* InName) const;
 
 	FConfig Config;
 
