Initialize global texture resources before other global resources.

diff --git a/Engine/Source/Runtime/RenderCore/Private/CommonRenderResources.cpp b/Engine/Source/Runtime/RenderCore/Private/CommonRenderResources.cpp
--- a/Engine/Source/Runtime/RenderCore/Private/CommonRenderResources.cpp
+++ b/Engine/Source/Runtime/RenderCore/Private/CommonRenderResources.cpp
@@ -9,11 +9,11 @@
 #include "StereoRenderUtils.h"
 
 
-TGlobalResource<FFilterVertexDeclaration> GFilterVertexDeclaration;
-TGlobalResource<FEmptyVertexDeclaration> GEmptyVertexDeclaration;
+TGlobalResource<FFilterVertexDeclaration, FRenderResource::EInitPhase::Pre> GFilterVertexDeclaration;
+TGlobalResource<FEmptyVertexDeclaration, FRenderResource::EInitPhase::Pre> GEmptyVertexDeclaration;
 
-TGlobalResource<FScreenRectangleVertexBuffer> GScreenRectangleVertexBuffer;
-TGlobalResource<FScreenRectangleIndexBuffer> GScreenRectangleIndexBuffer;
+TGlobalResource<FScreenRectangleVertexBuffer, FRenderResource::EInitPhase::Pre> GScreenRectangleVertexBuffer;
+TGlobalResource<FScreenRectangleIndexBuffer, FRenderResource::EInitPhase::Pre> GScreenRectangleIndexBuffer;
 
 IMPLEMENT_GLOBAL_SHADER(FScreenVertexShaderVS, "/Engine/Private/Tools/FullscreenVertexShader.usf", "MainVS", SF_Vertex);
 IMPLEMENT_GLOBAL_SHADER(FInstancedScreenVertexShaderVS, "/Engine/Private/Tools/FullscreenVertexShader.usf", "MainVS", SF_Vertex);
diff --git a/Engine/Source/Runtime/RenderCore/Private/GlobalRenderResources.cpp b/Engine/Source/Runtime/RenderCore/Private/GlobalRenderResources.cpp
--- a/Engine/Source/Runtime/RenderCore/Private/GlobalRenderResources.cpp
+++ b/Engine/Source/Runtime/RenderCore/Private/GlobalRenderResources.cpp
@@ -121,15 +121,15 @@ public:
 	}
 };
 
-FTextureWithSRV* GWhiteTextureWithSRV = new TGlobalResource<FColoredTexture<255, 255, 255, 255> >;
-FTextureWithSRV* GBlackTextureWithSRV = new TGlobalResource<FBlackTextureWithSRV>();
-FTextureWithSRV* GTransparentBlackTextureWithSRV = new TGlobalResource<FColoredTexture<0, 0, 0, 0> >;
+FTextureWithSRV* GWhiteTextureWithSRV = new TGlobalResource<FColoredTexture<255, 255, 255, 255>, FRenderResource::EInitPhase::Pre>;
+FTextureWithSRV* GBlackTextureWithSRV = new TGlobalResource<FBlackTextureWithSRV, FRenderResource::EInitPhase::Pre>();
+FTextureWithSRV* GTransparentBlackTextureWithSRV = new TGlobalResource<FColoredTexture<0, 0, 0, 0>, FRenderResource::EInitPhase::Pre>;
 FTexture* GWhiteTexture = GWhiteTextureWithSRV;
 FTexture* GBlackTexture = GBlackTextureWithSRV;
 FTexture* GTransparentBlackTexture = GTransparentBlackTextureWithSRV;
 
-FVertexBufferWithSRV* GEmptyVertexBufferWithUAV = new TGlobalResource<FEmptyVertexBuffer>;
-FVertexBufferWithSRV* GEmptyStructuredBufferWithUAV = new TGlobalResource<FEmptyStructuredBuffer>;
+FVertexBufferWithSRV* GEmptyVertexBufferWithUAV = new TGlobalResource<FEmptyVertexBuffer, FRenderResource::EInitPhase::Pre>;
+FVertexBufferWithSRV* GEmptyStructuredBufferWithUAV = new TGlobalResource<FEmptyStructuredBuffer, FRenderResource::EInitPhase::Pre>;
 
 class FWhiteVertexBuffer : public FVertexBufferWithSRV
 {
@@ -150,7 +150,7 @@ public:
 	}
 };
 
-FVertexBufferWithSRV* GWhiteVertexBufferWithSRV = new TGlobalResource<FWhiteVertexBuffer>;
+FVertexBufferWithSRV* GWhiteVertexBufferWithSRV = new TGlobalResource<FWhiteVertexBuffer, FRenderResource::EInitPhase::Pre>;
 
 class FWhiteVertexBufferWithRDG : public FBufferWithRDG
 {
@@ -169,7 +169,7 @@ public:
 	}
 };
 
-FBufferWithRDG* GWhiteVertexBufferWithRDG = new TGlobalResource<FWhiteVertexBufferWithRDG>();
+FBufferWithRDG* GWhiteVertexBufferWithRDG = new TGlobalResource<FWhiteVertexBufferWithRDG, FRenderResource::EInitPhase::Pre>();
 
 /**
  * A class representing a 1x1x1 black volume texture.
@@ -216,11 +216,11 @@ public:
 };
 
 /** Global black volume texture resource. */
-FTexture* GBlackVolumeTexture = new TGlobalResource<FBlackVolumeTexture<PF_B8G8R8A8, 0>>();
-FTexture* GBlackAlpha1VolumeTexture = new TGlobalResource<FBlackVolumeTexture<PF_B8G8R8A8, 255>>();
+FTexture* GBlackVolumeTexture = new TGlobalResource<FBlackVolumeTexture<PF_B8G8R8A8, 0>, FRenderResource::EInitPhase::Pre>();
+FTexture* GBlackAlpha1VolumeTexture = new TGlobalResource<FBlackVolumeTexture<PF_B8G8R8A8, 255>, FRenderResource::EInitPhase::Pre>();
 
 /** Global black volume texture resource. */
-FTexture* GBlackUintVolumeTexture = new TGlobalResource<FBlackVolumeTexture<PF_R8G8B8A8_UINT, 0>>();
+FTexture* GBlackUintVolumeTexture = new TGlobalResource<FBlackVolumeTexture<PF_R8G8B8A8_UINT, 0>, FRenderResource::EInitPhase::Pre>();
 
 class FBlackArrayTexture : public FTexture
 {
@@ -247,7 +247,7 @@ public:
 	virtual uint32 GetSizeY() const override { return 1; }
 };
 
-FTexture* GBlackArrayTexture = new TGlobalResource<FBlackArrayTexture>;
+FTexture* GBlackArrayTexture = new TGlobalResource<FBlackArrayTexture, FRenderResource::EInitPhase::Pre>;
 
 //
 // FMipColorTexture implementation
@@ -400,13 +400,13 @@ private:
 };
 
 /** A white cube texture. */
-FTexture* GWhiteTextureCube = new TGlobalResource<FSolidColorTextureCube>(FColor::White);
+FTexture* GWhiteTextureCube = new TGlobalResource<FSolidColorTextureCube, FRenderResource::EInitPhase::Pre>(FColor::White);
 
 /** A black cube texture. */
-FTexture* GBlackTextureCube = new TGlobalResource<FSolidColorTextureCube>(FColor::Black);
+FTexture* GBlackTextureCube = new TGlobalResource<FSolidColorTextureCube, FRenderResource::EInitPhase::Pre>(FColor::Black);
 
 /** A black cube texture. */
-FTexture* GBlackTextureDepthCube = new TGlobalResource<FSolidColorTextureCube>(PF_ShadowDepth);
+FTexture* GBlackTextureDepthCube = new TGlobalResource<FSolidColorTextureCube, FRenderResource::EInitPhase::Pre>(PF_ShadowDepth);
 
 class FBlackCubeArrayTexture : public FTexture
 {
@@ -444,7 +444,7 @@ public:
 	virtual uint32 GetSizeX() const override { return 1; }
 	virtual uint32 GetSizeY() const override { return 1; }
 };
-FTexture* GBlackCubeArrayTexture = new TGlobalResource<FBlackCubeArrayTexture>;
+FTexture* GBlackCubeArrayTexture = new TGlobalResource<FBlackCubeArrayTexture, FRenderResource::EInitPhase::Pre>;
 
 /**
  * A UINT 1x1 texture.
@@ -520,7 +520,7 @@ protected:
 	}
 };
 
-FTexture* GBlackUintTexture = new TGlobalResource< FUintTexture<PF_R32G32B32A32_UINT> >;
+FTexture* GBlackUintTexture = new TGlobalResource<FUintTexture<PF_R32G32B32A32_UINT>, FRenderResource::EInitPhase::Pre>;
 
 
 //////////////////////////////////////////////////////////////////////////////////////////////////////////////////////
@@ -551,7 +551,7 @@ void FNullColorVertexBuffer::ReleaseRHI()
 }
 
 /** The global null color vertex buffer, which is set with a stride of 0 on meshes without a color component. */
-TGlobalResource<FNullColorVertexBuffer> GNullColorVertexBuffer;
+TGlobalResource<FNullColorVertexBuffer, FRenderResource::EInitPhase::Pre> GNullColorVertexBuffer;
 
 //////////////////////////////////////////////////////////////////////////////////////////////////////////////////////
 // FNullVertexBuffer
@@ -578,7 +578,7 @@ void FNullVertexBuffer::ReleaseRHI()
 }
 
 /** The global null vertex buffer, which is set with a stride of 0 on meshes */
-TGlobalResource<FNullVertexBuffer> GNullVertexBuffer;
+TGlobalResource<FNullVertexBuffer, FRenderResource::EInitPhase::Pre> GNullVertexBuffer;
 
 //////////////////////////////////////////////////////////////////////////////////////////////////////////////////////
 // FScreenSpaceVertexBuffer
@@ -600,7 +600,7 @@ void FScreenSpaceVertexBuffer::InitRHI()
 	RHIUnlockBuffer(VertexBufferRHI);
 }
 
-TGlobalResource<FScreenSpaceVertexBuffer> GScreenSpaceVertexBuffer;
+TGlobalResource<FScreenSpaceVertexBuffer, FRenderResource::EInitPhase::Pre> GScreenSpaceVertexBuffer;
 
 //////////////////////////////////////////////////////////////////////////////////////////////////////////////////////
 // FTileVertexDeclaration
@@ -621,7 +621,7 @@ void FTileVertexDeclaration::ReleaseRHI()
 	VertexDeclarationRHI.SafeRelease();
 }
 
-TGlobalResource<FTileVertexDeclaration> GTileVertexDeclaration;
+TGlobalResource<FTileVertexDeclaration, FRenderResource::EInitPhase::Pre> GTileVertexDeclaration;
 
 //////////////////////////////////////////////////////////////////////////////////////////////////////////////////////
 // FCubeIndexBuffer
@@ -636,7 +636,7 @@ void FCubeIndexBuffer::InitRHI()
 	RHIUnlockBuffer(IndexBufferRHI);
 }
 
-TGlobalResource<FCubeIndexBuffer> GCubeIndexBuffer;
+TGlobalResource<FCubeIndexBuffer, FRenderResource::EInitPhase::Pre> GCubeIndexBuffer;
 
 //////////////////////////////////////////////////////////////////////////////////////////////////////////////////////
 // FTwoTrianglesIndexBuffer
@@ -652,7 +652,7 @@ void FTwoTrianglesIndexBuffer::InitRHI()
 	RHIUnlockBuffer(IndexBufferRHI);
 }
 
-TGlobalResource<FTwoTrianglesIndexBuffer> GTwoTrianglesIndexBuffer;
+TGlobalResource<FTwoTrianglesIndexBuffer, FRenderResource::EInitPhase::Pre> GTwoTrianglesIndexBuffer;
 
 /*------------------------------------------------------------------------------
 	FGlobalDynamicVertexBuffer implementation.
@@ -847,7 +847,7 @@ private:
 	}
 };
 
-static TGlobalResource<FDynamicVertexBufferPool> GDynamicVertexBufferPool;
+static TGlobalResource<FDynamicVertexBufferPool, FRenderResource::EInitPhase::Pre> GDynamicVertexBufferPool;
 
 FGlobalDynamicVertexBuffer::FAllocation FGlobalDynamicVertexBuffer::Allocate(uint32 SizeInBytes)
 {
diff --git a/Engine/Source/Runtime/RenderCore/Private/RenderResource.cpp b/Engine/Source/Runtime/RenderCore/Private/RenderResource.cpp
--- a/Engine/Source/Runtime/RenderCore/Private/RenderResource.cpp
+++ b/Engine/Source/Runtime/RenderCore/Private/RenderResource.cpp
@@ -27,12 +27,24 @@ FAutoConsoleVariableRef CVarFreeStructuresOnRHIBufferCreation(
 class FRenderResourceList
 {
 public:
+	template <FRenderResource::EInitPhase>
 	static FRenderResourceList& Get()
 	{
 		static FRenderResourceList Instance;
 		return Instance;
 	}
 
+	static FRenderResourceList& Get(FRenderResource::EInitPhase InitPhase)
+	{
+		switch (InitPhase)
+		{
+		case FRenderResource::EInitPhase::Pre:
+			return Get<FRenderResource::EInitPhase::Pre>();
+		default:
+			return Get<FRenderResource::EInitPhase::Default>();
+		}
+	}
+
 	int32 Allocate(FRenderResource* Resource)
 	{
 		Mutex.Lock();
@@ -105,28 +117,29 @@ private:
 	TArray<FRenderResource*> ResourceList;
 };
 
-void FRenderResource::CoalesceResourceList()
-{
-}
-
 void FRenderResource::ReleaseRHIForAllResources()
 {
-	FRenderResourceList& ResourceList = FRenderResourceList::Get();
-	ResourceList.ForEachReverse([](FRenderResource* Resource) { check(Resource->IsInitialized()); Resource->ReleaseRHI(); });
+	const auto Lambda = [](FRenderResource* Resource) { check(Resource->IsInitialized()); Resource->ReleaseRHI(); };
+	FRenderResourceList::Get<FRenderResource::EInitPhase::Default>().ForEachReverse(Lambda);
+	FRenderResourceList::Get<FRenderResource::EInitPhase::Pre>().ForEachReverse(Lambda);
 }
 
 /** Initialize all resources initialized before the RHI was initialized */
 void FRenderResource::InitPreRHIResources()
 {
-	FRenderResourceList& ResourceList = FRenderResourceList::Get();
-
 	FRHICommandListImmediate& RHICmdList = FRHICommandListExecutor::GetImmediateCommandList();
 
+	FRenderResourceList& PreResourceList = FRenderResourceList::Get<FRenderResource::EInitPhase::Pre>();
+	FRenderResourceList& DefaultResourceList = FRenderResourceList::Get<FRenderResource::EInitPhase::Default>();
+
 	// Notify all initialized FRenderResources that there's a valid RHI device to create their RHI resources for now.
-	ResourceList.ForEach([&](FRenderResource* Resource) { Resource->InitRHI(RHICmdList); });
+	const auto Lambda = [&](FRenderResource* Resource) { Resource->InitRHI(RHICmdList); };
+	PreResourceList.ForEach(Lambda);
+	DefaultResourceList.ForEach(Lambda);
 
 #if !PLATFORM_NEEDS_RHIRESOURCELIST
-	ResourceList.Clear();
+	PreResourceList.Clear();
+	DefaultResourceList.Clear();
 #endif
 }
 
@@ -135,7 +148,7 @@ void FRenderResource::ChangeFeatureLevel(ERHIFeatureLevel::Type NewFeatureLevel)
 	ENQUEUE_RENDER_COMMAND(FRenderResourceChangeFeatureLevel)(
 		[NewFeatureLevel](FRHICommandList& RHICmdList)
 	{
-		FRenderResourceList::Get().ForEach([NewFeatureLevel](FRenderResource* Resource)
+		const auto Lambda = [NewFeatureLevel](FRenderResource* Resource)
 		{
 			// Only resources configured for a specific feature level need to be updated
 			if (Resource->HasValidFeatureLevel() && (Resource->FeatureLevel != NewFeatureLevel))
@@ -144,7 +157,10 @@ void FRenderResource::ChangeFeatureLevel(ERHIFeatureLevel::Type NewFeatureLevel)
 				Resource->FeatureLevel = NewFeatureLevel;
 				Resource->InitRHI();
 			}
-		});
+		};
+
+		FRenderResourceList::Get<FRenderResource::EInitPhase::Pre>().ForEach(Lambda);
+		FRenderResourceList::Get<FRenderResource::EInitPhase::Default>().ForEach(Lambda);
 	});
 }
 
@@ -158,7 +174,7 @@ void FRenderResource::InitResource()
 		if (PLATFORM_NEEDS_RHIRESOURCELIST || !GIsRHIInitialized)
 		{
 			LLM_SCOPE(ELLMTag::SceneRender);
-			LocalListIndex = FRenderResourceList::Get().Allocate(this);
+			LocalListIndex = FRenderResourceList::Get(InitPhase).Allocate(this);
 		}
 		else
 		{
@@ -187,7 +203,7 @@ void FRenderResource::InitResource(FRHICommandList& RHICmdList)
 		if (PLATFORM_NEEDS_RHIRESOURCELIST || !GIsRHIInitialized)
 		{
 			LLM_SCOPE(ELLMTag::SceneRender);
-			LocalListIndex = FRenderResourceList::Get().Allocate(this);
+			LocalListIndex = FRenderResourceList::Get(InitPhase).Allocate(this);
 		}
 		else
 		{
@@ -218,7 +234,7 @@ void FRenderResource::ReleaseResource()
 			}
 
 #if PLATFORM_NEEDS_RHIRESOURCELIST
-			FRenderResourceList::Get().Deallocate(ListIndex);
+			FRenderResourceList::Get(InitPhase).Deallocate(ListIndex);
 #endif
 			ListIndex = INDEX_NONE;
 		}
diff --git a/Engine/Source/Runtime/RenderCore/Public/CommonRenderResources.h b/Engine/Source/Runtime/RenderCore/Public/CommonRenderResources.h
--- a/Engine/Source/Runtime/RenderCore/Public/CommonRenderResources.h
+++ b/Engine/Source/Runtime/RenderCore/Public/CommonRenderResources.h
@@ -56,7 +56,7 @@ public:
 	}
 };
 
-extern RENDERCORE_API TGlobalResource<FFilterVertexDeclaration> GFilterVertexDeclaration;
+extern RENDERCORE_API TGlobalResource<FFilterVertexDeclaration, FRenderResource::EInitPhase::Pre> GFilterVertexDeclaration;
 
 /** The empty vertex declaration resource type. */
 class FEmptyVertexDeclaration : public FRenderResource
@@ -79,11 +79,11 @@ public:
 	}
 };
 
-extern RENDERCORE_API TGlobalResource<FEmptyVertexDeclaration> GEmptyVertexDeclaration;
+extern RENDERCORE_API TGlobalResource<FEmptyVertexDeclaration, FRenderResource::EInitPhase::Pre> GEmptyVertexDeclaration;
 
 /**
- * Static vertex and index buffer used for 2D screen rectangles.
- */
+* Static vertex and index buffer used for 2D screen rectangles.
+*/
 class FScreenRectangleVertexBuffer : public FVertexBuffer
 {
 public:
@@ -91,7 +91,7 @@ public:
 	void InitRHI() override;
 };
 
-extern RENDERCORE_API TGlobalResource<FScreenRectangleVertexBuffer> GScreenRectangleVertexBuffer;
+extern RENDERCORE_API TGlobalResource<FScreenRectangleVertexBuffer, FRenderResource::EInitPhase::Pre> GScreenRectangleVertexBuffer;
 
 
 class FScreenRectangleIndexBuffer : public FIndexBuffer
@@ -101,7 +101,7 @@ public:
 	void InitRHI() override;
 };
 
-extern RENDERCORE_API TGlobalResource<FScreenRectangleIndexBuffer> GScreenRectangleIndexBuffer;
+extern RENDERCORE_API TGlobalResource<FScreenRectangleIndexBuffer, FRenderResource::EInitPhase::Pre> GScreenRectangleIndexBuffer;
 
 
 /** Vertex shader to draw a screen quad that works on all platforms. Does not have any shader parameters.
diff --git a/Engine/Source/Runtime/RenderCore/Public/GlobalRenderResources.h b/Engine/Source/Runtime/RenderCore/Public/GlobalRenderResources.h
--- a/Engine/Source/Runtime/RenderCore/Public/GlobalRenderResources.h
+++ b/Engine/Source/Runtime/RenderCore/Public/GlobalRenderResources.h
@@ -87,7 +87,7 @@ public:
 };
 
 /** The global null color vertex buffer, which is set with a stride of 0 on meshes without a color component. */
-extern RENDERCORE_API TGlobalResource<FNullColorVertexBuffer> GNullColorVertexBuffer;
+extern RENDERCORE_API TGlobalResource<FNullColorVertexBuffer, FRenderResource::EInitPhase::Pre> GNullColorVertexBuffer;
 
 /**
 * A vertex buffer with a single zero float3 component.
@@ -105,7 +105,7 @@ public:
 };
 
 /** The global null vertex buffer, which is set with a stride of 0 on meshes */
-extern RENDERCORE_API TGlobalResource<FNullVertexBuffer> GNullVertexBuffer;
+extern RENDERCORE_API TGlobalResource<FNullVertexBuffer, FRenderResource::EInitPhase::Pre> GNullVertexBuffer;
 
 class FScreenSpaceVertexBuffer : public FVertexBuffer
 {
@@ -116,7 +116,7 @@ public:
 	RENDERCORE_API virtual void InitRHI() override;
 };
 
-extern RENDERCORE_API TGlobalResource<FScreenSpaceVertexBuffer> GScreenSpaceVertexBuffer;
+extern RENDERCORE_API TGlobalResource<FScreenSpaceVertexBuffer, FRenderResource::EInitPhase::Pre> GScreenSpaceVertexBuffer;
 
 class FTileVertexDeclaration : public FRenderResource
 {
@@ -130,7 +130,7 @@ public:
 	FVertexDeclarationRHIRef VertexDeclarationRHI;
 };
 
-extern RENDERCORE_API TGlobalResource<FTileVertexDeclaration> GTileVertexDeclaration;
+extern RENDERCORE_API TGlobalResource<FTileVertexDeclaration, FRenderResource::EInitPhase::Pre> GTileVertexDeclaration;
 
 class FCubeIndexBuffer : public FIndexBuffer
 {
@@ -140,7 +140,7 @@ public:
 	*/
 	virtual void InitRHI() override;
 };
-extern RENDERCORE_API TGlobalResource<FCubeIndexBuffer> GCubeIndexBuffer;
+extern RENDERCORE_API TGlobalResource<FCubeIndexBuffer, FRenderResource::EInitPhase::Pre> GCubeIndexBuffer;
 
 class FTwoTrianglesIndexBuffer : public FIndexBuffer
 {
@@ -150,7 +150,7 @@ public:
 	*/
 	virtual void InitRHI() override;
 };
-extern RENDERCORE_API TGlobalResource<FTwoTrianglesIndexBuffer> GTwoTrianglesIndexBuffer;
+extern RENDERCORE_API TGlobalResource<FTwoTrianglesIndexBuffer, FRenderResource::EInitPhase::Pre> GTwoTrianglesIndexBuffer;
 
 
 //////////////////////////////////////////////////////////////////////////////////////////////////////////////////////
diff --git a/Engine/Source/Runtime/RenderCore/Public/RenderResource.h b/Engine/Source/Runtime/RenderCore/Public/RenderResource.h
--- a/Engine/Source/Runtime/RenderCore/Public/RenderResource.h
+++ b/Engine/Source/Runtime/RenderCore/Public/RenderResource.h
@@ -58,13 +58,19 @@ enum class ERayTracingMode : uint8
 
 /**
  * A rendering resource which is owned by the rendering thread.
- * NOTE - Adding new virtual methods to this class may require stubs added to FViewport/FDummyViewport, otherwise certain modules may have link errors
  */
 class FRenderResource
 {
 public:
 	////////////////////////////////////////////////////////////////////////////////////
-	// The following methods may not be called while asynchronously initializing / releasing render resources.
+
+	/** Controls initialization order of render resources. Early engine resources utilize the 'Pre' phase to avoid static init ordering issues. */
+	enum class EInitPhase : uint8
+	{
+		Pre,
+		Default,
+		MAX
+	};
 
 	/** Release all render resources that are currently initialized. */
 	static RENDERCORE_API void ReleaseRHIForAllResources();
@@ -72,9 +78,6 @@ public:
 	/** Initialize all resources initialized before the RHI was initialized. */
 	static RENDERCORE_API void InitPreRHIResources();
 
-	/** Call periodically to coalesce the render resource list. */
-	static RENDERCORE_API void CoalesceResourceList();
-
 	/** Reinitializes render resources at a new feature level. */
 	static RENDERCORE_API void ChangeFeatureLevel(ERHIFeatureLevel::Type NewFeatureLevel);
 
@@ -149,6 +152,7 @@ public:
 	FORCEINLINE bool IsInitialized() const { return ListIndex != INDEX_NONE; }
 
 	int32 GetListIndex() const { return ListIndex; }
+	EInitPhase GetInitPhase() const { return InitPhase; }
 
 	/** SetOwnerName should be called before BeginInitResource for the owner name to be successfully tracked. */
 	RENDERCORE_API void SetOwnerName(const FName& InOwnerName);
@@ -185,6 +189,13 @@ protected:
 		return Buffer;
 	}
 
+	void SetInitPhase(EInitPhase InInitPhase)
+	{
+		check(InInitPhase != EInitPhase::MAX);
+		check(!IsInitialized());
+		InitPhase = InInitPhase;
+	}
+
 private:
 	static RENDERCORE_API bool ShouldFreeResourceObject(void* ResourceObject, FResourceArrayInterface* ResourceArray);
 	static RENDERCORE_API FBufferRHIRef CreateRHIBufferInternal(
@@ -203,6 +214,7 @@ private:
 
 	int32 ListIndex;
 	TEnumAsByte<ERHIFeatureLevel::Type> FeatureLevel;
+	EInitPhase InitPhase = EInitPhase::Default;
 
 public:
 	ERenderResourceState ResourceState = ERenderResourceState::Default;
@@ -506,7 +518,7 @@ public:
 };
 
 /** Used to declare a render resource that is initialized/released by static initialization/destruction. */
-template<class ResourceType>
+template<class ResourceType, FRenderResource::EInitPhase InInitPhase = FRenderResource::EInitPhase::Default>
 class TGlobalResource : public ResourceType
 {
 public:
@@ -537,6 +549,8 @@ private:
 	 */
 	void InitGlobalResource()
 	{
+		ResourceType::SetInitPhase(InInitPhase);
+
 		if (IsInRenderingThread())
 		{
 			// If the resource is constructed in the rendering thread, directly initialize it.
diff --git a/Engine/Source/Runtime/Renderer/Private/DepthRendering.cpp b/Engine/Source/Runtime/Renderer/Private/DepthRendering.cpp
--- a/Engine/Source/Runtime/Renderer/Private/DepthRendering.cpp
+++ b/Engine/Source/Runtime/Renderer/Private/DepthRendering.cpp
@@ -709,7 +709,7 @@ void FDeferredShadingSceneRenderer::RenderPrePassHMD(FRDGBuilder& GraphBuilder,
 				ERDGPassFlags::Raster,
 				[this, &View, HMDDevice](FRHICommandList& RHICmdList)
 			{
-				extern TGlobalResource<FFilterVertexDeclaration> GFilterVertexDeclaration;
+				extern TGlobalResource<FFilterVertexDeclaration, FRenderResource::EInitPhase::Pre> GFilterVertexDeclaration;
 
 				TShaderMapRef<TOneColorVS<true>> VertexShader(GetGlobalShaderMap(GMaxRHIFeatureLevel));
 
diff --git a/Engine/Source/Runtime/Renderer/Private/RendererModule.h b/Engine/Source/Runtime/Renderer/Private/RendererModule.h
--- a/Engine/Source/Runtime/Renderer/Private/RendererModule.h
+++ b/Engine/Source/Runtime/Renderer/Private/RendererModule.h
@@ -24,7 +24,7 @@ struct FMeshBatch;
 struct FSynthBenchmarkResults;
 struct FSceneTextures;
 
-template<class ResourceType> class TGlobalResource;
+template<class ResourceType, FRenderResource::EInitPhase> class TGlobalResource;
 
 DECLARE_LOG_CATEGORY_EXTERN(LogRenderer, Log, All);
 
diff --git a/Engine/Source/Runtime/Renderer/Private/SceneRendering.cpp b/Engine/Source/Runtime/Renderer/Private/SceneRendering.cpp
--- a/Engine/Source/Runtime/Renderer/Private/SceneRendering.cpp
+++ b/Engine/Source/Runtime/Renderer/Private/SceneRendering.cpp
@@ -4201,7 +4201,6 @@ static void FinishCleanUp(FRHICommandListImmediate& RHICmdList)
 	FGlobalDynamicVertexBuffer::GarbageCollect();
 	GPrimitiveIdVertexBufferPool.DiscardAll();
 	FGraphicsMinimalPipelineStateId::ResetLocalPipelineIdTableSize();
-	FRenderResource::CoalesceResourceList();
 }
 
 static void DeleteSceneRenderers(const TArray<FSceneRenderer*>& SceneRenderers, FParallelMeshDrawCommandPass::EWaitThread WaitThread)
diff --git a/Engine/Source/Runtime/Renderer/Private/ShadowDepthRendering.cpp b/Engine/Source/Runtime/Renderer/Private/ShadowDepthRendering.cpp
--- a/Engine/Source/Runtime/Renderer/Private/ShadowDepthRendering.cpp
+++ b/Engine/Source/Runtime/Renderer/Private/ShadowDepthRendering.cpp
@@ -849,7 +849,7 @@ void FProjectedShadowInfo::CopyCachedShadowMap(
 		GraphicsPSOInit.DepthStencilState = TStaticDepthStencilState<true, CF_Always>::GetRHI();
 		GraphicsPSOInit.PrimitiveType = PT_TriangleList;
 
-		extern TGlobalResource<FFilterVertexDeclaration> GFilterVertexDeclaration;
+		extern TGlobalResource<FFilterVertexDeclaration, FRenderResource::EInitPhase::Pre> GFilterVertexDeclaration;
 		GraphicsPSOInit.BoundShaderState.VertexDeclarationRHI = GFilterVertexDeclaration.VertexDeclarationRHI;
 
 		if (CacheMode == SDCM_MovablePrimitivesOnly)
