Fixes VisualizeTexture not working with RDG_TEXTURE_ACCESS

diff --git a/Engine/Source/Runtime/RenderCore/Private/RenderGraphBuilder.cpp b/Engine/Source/Runtime/RenderCore/Private/RenderGraphBuilder.cpp
--- a/Engine/Source/Runtime/RenderCore/Private/RenderGraphBuilder.cpp
+++ b/Engine/Source/Runtime/RenderCore/Private/RenderGraphBuilder.cpp
@@ -3864,9 +3864,23 @@ void FRDGBuilder::VisualizePassOutputs(const FRDGPass* Pass)
 		{
 			if (FRDGTextureAccess TextureAccess = Parameter.GetAsTextureAccess())
 			{
-				if (TextureAccess.GetAccess() == ERHIAccess::UAVCompute ||
-					TextureAccess.GetAccess() == ERHIAccess::UAVGraphics ||
-					TextureAccess.GetAccess() == ERHIAccess::RTV)
+				if (IsWritableAccess(TextureAccess.GetAccess()))
+				{
+					if (TOptional<uint32> CaptureId = GVisualizeTexture.ShouldCapture(TextureAccess->Name, /* MipIndex = */ 0))
+					{
+						GVisualizeTexture.CreateContentCapturePass(*this, TextureAccess.GetTexture(), *CaptureId);
+					}
+				}
+			}
+		}
+		break;
+		case UBMT_RDG_TEXTURE_ACCESS_ARRAY:
+		{
+			const FRDGTextureAccessArray& TextureAccessArray = Parameter.GetAsTextureAccessArray();
+
+			for (FRDGTextureAccess TextureAccess : TextureAccessArray)
+			{
+				if (IsWritableAccess(TextureAccess.GetAccess()))
 				{
 					if (TOptional<uint32> CaptureId = GVisualizeTexture.ShouldCapture(TextureAccess->Name, /* MipIndex = */ 0))
 					{
