[shaders]
- Add option to independent preprocessing IShaderFormat API to allow a "secondary" preprocess & compile if a RequiresSecondaryCompile predicate implemented in the shader format returns true; this predicate takes as input the compile input, environment,  and the output of the first execution of the preprocess function. This is intended to handle the usage of UESHADERMETADATA pragmas (set in the preprocess output) by some backends to indicate a secondary compilation is required to target the variant hardware (setting additional #defines and potentially different compiler arguments). Other backends have an additional requirement that any stereo shaders also require separate compilation for each target, which can be handled by querying the FShaderCompilerEnvironment in this same predicate.
- Move the combining of outputs in the secondary compile case into the core shader system so it doesn't need to be copypasta'd in any shader format backends which require the aforementioned secondary compilation, exposing an API function which returns the particular "key" used to identify these combined shader outputs in the RHI.
- Rework the internal compilation/preprocess functions so a minimal amount of code needs to be wrapped in __try/__except (allowing objects with destructors to be created in the higher level functions).

diff --git a/Engine/Source/Developer/TargetPlatform/Public/Interfaces/IShaderFormat.h b/Engine/Source/Developer/TargetPlatform/Public/Interfaces/IShaderFormat.h
--- a/Engine/Source/Developer/TargetPlatform/Public/Interfaces/IShaderFormat.h
+++ b/Engine/Source/Developer/TargetPlatform/Public/Interfaces/IShaderFormat.h
@@ -8,6 +8,10 @@
 #define SHADERFORMAT_MODULE_WILDCARD TEXT("*ShaderFormat*")
 
 class FSerializedShaderArchive;
+class FShaderPreprocessOutput;
+struct FShaderCompilerEnvironment;
+struct FShaderCompilerInput;
+struct FShaderCompilerOutput;
 
 /**
  * IShaderFormat, shader pre-compilation abstraction
@@ -24,7 +28,7 @@ public:
 	 * @param Output The output from shader compiler.
 	 * @param WorkingDirectory The working directory.
 	 */
-	virtual void CompileShader(FName Format, const struct FShaderCompilerInput& Input, struct FShaderCompilerOutput& Output, const FString& WorkingDirectory) const {};
+	virtual void CompileShader(FName Format, const FShaderCompilerInput& Input, FShaderCompilerOutput& Output, const FString& WorkingDirectory) const {};
 
 	/**
 	 * Gets the current version of the specified shader format.
@@ -32,14 +36,14 @@ public:
 	 * @param Format The format to get the version for.
 	 * @return Version number.
 	 */
-	virtual uint32 GetVersion( FName Format ) const = 0;
+	virtual uint32 GetVersion(FName Format) const = 0;
 
 	/**
 	 * Gets the list of supported formats.
 	 *
 	 * @param OutFormats Will hold the list of formats.
 	 */
-	virtual void GetSupportedFormats( TArray<FName>& OutFormats ) const = 0;
+	virtual void GetSupportedFormats(TArray<FName>& OutFormats) const = 0;
 	
 
 	/**
@@ -57,7 +61,7 @@ public:
 	 * @param bNative Whether the final shader library uses a native format which may determine how the shader is stripped.
 	 * @return True if the format has successfully stripped the extraneous data from shaders, otherwise false
 	 */
-    virtual bool StripShaderCode( TArray<uint8>& Code, FString const& DebugOutputDir, bool const bNative ) const { return false; }
+    virtual bool StripShaderCode(TArray<uint8>& Code, FString const& DebugOutputDir, bool const bNative) const { return false; }
     
 	/**
 	* Whether this shader format supports a format-specific archive for precompiled shader code.
@@ -77,7 +81,7 @@ public:
 	 * @param Optional pointer to a TArray that on success will be filled with a list of the written file paths.
      * @returns true if the archive was created
      */
-    virtual bool CreateShaderArchive( FString const& LibraryName,
+    virtual bool CreateShaderArchive(FString const& LibraryName,
 		FName ShaderFormatAndShaderPlatformName,
 		const FString& WorkingDirectory,
 		const FString& OutputDir,
@@ -126,28 +130,87 @@ public:
 	 * Can the shader compiler use the HLSLcc library when compiling shaders
 	 * @returns True if the shader compiler can use the HLSLcc library when compiling shaders, otherwise false.
 	 */
-	virtual bool UsesHLSLcc(const struct FShaderCompilerInput& Input) const { return false; }
+	virtual bool UsesHLSLcc(const FShaderCompilerInput& Input) const { return false; }
 
 	/**
 	 * Executes all shader preprocessing steps; only called if SupportsIndependentPreprocessing returns true
 	 */
-	virtual bool PreprocessShader(const struct FShaderCompilerInput& Input, const struct FShaderCompilerEnvironment& Environment, class FShaderPreprocessOutput& PreprocessOutput) const { return false; };
+	virtual bool PreprocessShader(const FShaderCompilerInput& Input, const FShaderCompilerEnvironment& Environment, FShaderPreprocessOutput& PreprocessOutput) const { return false; };
 
 	/**
 	 * Compile the specified preprocessed shader; only called if SupportsIndependentPreprocessing returns true
 	 */
-	virtual void CompilePreprocessedShader(const struct FShaderCompilerInput& Input, const class FShaderPreprocessOutput& PreprocessOutput, struct FShaderCompilerOutput& Output, const FString& WorkingDirectory) const {}
+	virtual void CompilePreprocessedShader(
+		const FShaderCompilerInput& Input, 
+		const FShaderPreprocessOutput& PreprocessOutput, 
+		FShaderCompilerOutput& Output, 
+		const FString& WorkingDirectory) const {}
+
+	/**
+	 * Compile the specified preprocessed shaders; only called if SupportsIndependentPreprocessing returns true, and the call to
+	 * RequiresSecondaryCompile given the first preprocess output also returns true. The shader system will pack these outputs together
+	 * in the following format: 
+	 * [int32 key][uint32 primary length][uint32 secondary length][full primary shader code][full secondary shader code]
+	 * where "key" is the return value of the GetPackedShaderKey function (this should also be implemented by any backends which
+	 * require secondary compilation and is used by the RHI to differentiate packed shader code from single shaders)
+	 */
+	virtual void CompilePreprocessedShader(
+		const FShaderCompilerInput& Input, 
+		const FShaderPreprocessOutput& PrimaryPreprocessOutput, 
+		const FShaderPreprocessOutput& SecondaryPreprocessOutput, 
+		FShaderCompilerOutput& PrimaryOutput, 
+		FShaderCompilerOutput& SecondaryOutput,
+		const FString& WorkingDirectory) const {}
+
+	/**
+	 * Predicate which should return true if a second preprocess & compilation is required given the initial preprocess output.
+	 * Only called if SupportsIndependentPreprocessing returns true.
+	 * This is generally be determined by analyzing the directives on the given preprocess output (which are set by the presence
+	 * of UESHADERMETADATA directives in the original source), though other conditions are possible and up to the implementation.
+	 *  In the event this returns true, PreprocessShader will be called an additional time, with the bIsSecondary field on the 
+	 * FShaderPreprocessOutput set to true, and the overload of CompilePreprocessedShader taking two FShaderPreprocessOutput and 
+	 * two FShaderCompilerOutput objects will be called instead of the single object variants. Note that this is done instead of 
+	 * calling the single object variant twice so shader formats can statically distinguish primary compilation in a dual-compile 
+	 * case from the single-compile case.
+	 */
+	virtual bool RequiresSecondaryCompile(
+		const FShaderCompilerInput& Input,
+		const FShaderCompilerEnvironment& Environment,
+		const FShaderPreprocessOutput& PreprocessOutput) const { return false; }
+
+	/**
+	 * Virtual function that can be implemented by backends which require secondary compilation to specify a unique "key" identifier
+	 * for the packing scheme (for validation in the runtime at load). see RequiresSecondaryCompile and the overloads of 
+	 * CompilePreprocessedShader for additional details.
+	 */
+	virtual int32 GetPackedShaderKey() const { return 0; }
 
 	/* 
-	 * Output debug info for a single compile job; only called if SupportsIndependentPreprocessing returns true.
+	 * Implement to output debug info for a single compile job; only called if SupportsIndependentPreprocessing returns true.
 	 * This will be called for all jobs (including those found in the job cache) but only if debug info is enabled for the job.
 	 * Note that any debug info output in CompilePreprocessedShader will only be done for the job that actually executes the
 	 * compile step, as such any debug outputs that are desirable for all jobs should be written by this function.
 	 * A BaseShaderFormat implementation is provided in ShaderCompilerCommon which dumps preprocessed and stripped USFs along
 	 * with the hash of the shader code in an OutputHash.txt file; if no additional custom debug output is required the shader 
-	 * format can inherit from BaseShaderFormat instead of IShaderFormat.
+	 * format can inherit from FBaseShaderFormat instead of IShaderFormat.
+	 */
+	virtual void OutputDebugData(
+		const FShaderCompilerInput& Input,
+		const FShaderPreprocessOutput& PreprocessOutput,
+		const FShaderCompilerOutput& Output) const {};
+
+	/*
+	 * Implement to output debug info for the case where a secondary preprocessed shader was created. As with the dual-output version
+	 * of CompilePreprocessedShader this is provided so shader formats can statically distinguish the single and dual output cases.
+	 * Note that BaseShaderFormat does not provide an implementation of this so it needs to be explicitly implemented by backends 
+	 * which require secondary compilation, even if inheriting from FBaseShaderFormat.
 	 */
-	virtual void OutputDebugData(const struct FShaderCompilerInput& Input, const class FShaderPreprocessOutput& PreprocessOutput, const struct FShaderCompilerOutput& Output) const {};
+	virtual void OutputDebugData(
+		const FShaderCompilerInput& Input,
+		const FShaderPreprocessOutput& PreprocessOutput,
+		const FShaderPreprocessOutput& SecondaryPreprocessOutput,
+		const FShaderCompilerOutput& Output,
+		const FShaderCompilerOutput& SecondaryOutput) const {};
 
 	/**
 	 * Return true if preprocessing for this format can be executed independent of compilation (i.e. the format has an
diff --git a/Engine/Source/Runtime/RenderCore/Private/ShaderCore.cpp b/Engine/Source/Runtime/RenderCore/Private/ShaderCore.cpp
--- a/Engine/Source/Runtime/RenderCore/Private/ShaderCore.cpp
+++ b/Engine/Source/Runtime/RenderCore/Private/ShaderCore.cpp
@@ -975,9 +975,52 @@ int HandleShaderCompileException(Windows::LPEXCEPTION_POINTERS Info, FString& Ou
 class FInternalShaderCompilerFunctions
 {
 public:
+
+	static bool InvokePreprocess(
+		const IShaderFormat* Backend,
+		const FShaderCompilerInput& Input,
+		const FShaderCompilerEnvironment& Environment,
+		FShaderPreprocessOutput& Output,
+		FString& OutExceptionCallstack,
+		FString& OutExceptionMsg)
+	{
+#if PLATFORM_WINDOWS
+		__try
+#endif
+		{
+			return Backend->PreprocessShader(Input, Environment, Output);
+		}
+#if PLATFORM_WINDOWS
+		__except (HandleShaderCompileException(GetExceptionInformation(), OutExceptionMsg, OutExceptionCallstack))
+		{
+		}
+#endif
+		return false;
+	}
+
+	static void PreprocessShaderInternal(
+		const IShaderFormat* Backend, 
+		FShaderCompileJob& Job, 
+		const FShaderCompilerEnvironment& Environment,
+		FString& OutExceptionCallstack,
+		FString& OutExceptionMsg)
+	{
+		Job.PreprocessOutput.bSucceeded = InvokePreprocess(Backend, Job.Input, Environment, Job.PreprocessOutput, OutExceptionCallstack, OutExceptionMsg);
+		if (Job.PreprocessOutput.bSucceeded && Backend->RequiresSecondaryCompile(Job.Input, Environment, Job.PreprocessOutput))
+		{
+			Job.SecondaryPreprocessOutput = MakeUnique<FShaderPreprocessOutput>();
+			Job.SecondaryPreprocessOutput->bIsSecondary = true;
+			Job.PreprocessOutput.bSucceeded &= InvokePreprocess(Backend, Job.Input, Environment, *Job.SecondaryPreprocessOutput, OutExceptionCallstack, OutExceptionMsg);
+		}
+	}
+
 	static bool PreprocessShaderInternal(const IShaderFormat* Backend, FShaderCompileJob& Job)
 	{
 		const double StartPreprocessTime = FPlatformTime::Seconds();
+
+		FString ExceptionCallstack;
+		FString ExceptionMsg;
+
 		if (IsValidRef(Job.Input.SharedEnvironment))
 		{
 			// only create new environment & merge if necessary, save some allocs
@@ -985,11 +1028,20 @@ public:
 			// and affect what is passed to the workers)
 			FShaderCompilerEnvironment MergedEnvironment = Job.Input.Environment;
 			MergedEnvironment.Merge(*Job.Input.SharedEnvironment);
-			Job.PreprocessOutput.bSucceeded = Backend->PreprocessShader(Job.Input, MergedEnvironment, Job.PreprocessOutput);
+			PreprocessShaderInternal(Backend, Job, MergedEnvironment, ExceptionCallstack, ExceptionMsg);
 		}
 		else
 		{
-			Job.PreprocessOutput.bSucceeded = Backend->PreprocessShader(Job.Input, Job.Input.Environment, Job.PreprocessOutput);
+			PreprocessShaderInternal(Backend, Job, Job.Input.Environment, ExceptionCallstack, ExceptionMsg);
+		}
+
+		if (!Job.PreprocessOutput.bSucceeded && (!ExceptionMsg.IsEmpty() || !ExceptionCallstack.IsEmpty()))
+		{
+			FString StrippedErrorMessage = FString::Printf(
+				TEXT("Exception encountered in platform compiler: %s\nException Callstack:\n%s"),
+				*ExceptionMsg,
+				*ExceptionCallstack);
+			Job.PreprocessOutput.LogError(MoveTemp(StrippedErrorMessage));
 		}
 
 		if (Job.PreprocessOutput.bSucceeded && Job.Input.bCachePreprocessed)
@@ -1004,65 +1056,119 @@ public:
 		return Job.PreprocessOutput.bSucceeded;
 	}
 
-#if PLATFORM_WINDOWS
-	static void CompileShaderInternal(const IShaderFormat* Compiler, FShaderCompileJob& Job, const FString& WorkingDirectory, FString& OutExceptionMsg, FString& OutExceptionCallstack, int32* CompileCount)
-#else
-	static void CompileShaderInternal(const IShaderFormat* Compiler, FShaderCompileJob& Job, const FString& WorkingDirectory, int32* CompileCount)
-#endif
+	static void CombineOutputs(const IShaderFormat* Compiler, FShaderCompileJob& Job)
+	{
+		// Pack shader code results together
+		// [int32 key][uint32 primary length][uint32 secondary length][full primary shader code][full secondary shader code]
+		TArray<uint8> CombinedSource;
+
+		int32 PackedShaderKey = Compiler->GetPackedShaderKey();
+		CombinedSource.Append(reinterpret_cast<const uint8*>(&PackedShaderKey), sizeof(PackedShaderKey));
+
+		const uint32 PrimaryLength = Job.Output.ShaderCode.GetReadAccess().Num();
+		CombinedSource.Append(reinterpret_cast<const uint8*>(&PrimaryLength), sizeof(PrimaryLength));
+
+		const uint32 SecondaryLength = Job.SecondaryOutput->ShaderCode.GetReadAccess().Num();
+		CombinedSource.Append(reinterpret_cast<const uint8*>(&SecondaryLength), sizeof(SecondaryLength));
+
+		CombinedSource.Append(Job.Output.ShaderCode.GetReadAccess());
+		CombinedSource.Append(Job.SecondaryOutput->ShaderCode.GetReadAccess());
+
+		// Replace Output shader code with the combined result
+		Job.Output.ShaderCode = {};
+		TArray<uint8>& FinalShaderCode = Job.Output.ShaderCode.GetWriteAccess();
+		FinalShaderCode.Append(CombinedSource);
+		Job.Output.ShaderCode.FinalizeShaderCode();
+	}
+
+	static void InvokeCompile(const IShaderFormat* Compiler, FShaderCompileJob& Job, const FString& WorkingDirectory, FString& OutExceptionMsg, FString& OutExceptionCallstack)
 	{
 #if PLATFORM_WINDOWS
 		__try
 #endif
 		{
-			double TimeStart = FPlatformTime::Seconds();
-
-			if (Compiler->SupportsIndependentPreprocessing())
+			if (Job.SecondaryOutput.IsValid())
 			{
-				if (!Job.Input.bCachePreprocessed)
-				{
-					PreprocessShaderInternal(Compiler, Job);
-				}
-				if (Job.PreprocessOutput.bSucceeded)
-				{
-					Compiler->CompilePreprocessedShader(Job.Input, Job.PreprocessOutput, Job.Output, WorkingDirectory);
-				}
-				else
-				{
-					Job.Output.bSucceeded = false;
-				}
-				Job.Output.Errors.Append(Job.PreprocessOutput.Errors);
+				check(Job.SecondaryPreprocessOutput.IsValid());
+				Compiler->CompilePreprocessedShader(Job.Input, Job.PreprocessOutput, *Job.SecondaryPreprocessOutput, Job.Output, *Job.SecondaryOutput, WorkingDirectory);
+			}
+			else if (Compiler->SupportsIndependentPreprocessing())
+			{
+				Compiler->CompilePreprocessedShader(Job.Input, Job.PreprocessOutput, Job.Output, WorkingDirectory);
 			}
 			else
 			{
 				Compiler->CompileShader(Job.Input.ShaderFormat, Job.Input, Job.Output, WorkingDirectory);
 			}
 
-			if (Job.Output.bSucceeded)
+		}
+#if PLATFORM_WINDOWS
+		__except(HandleShaderCompileException(GetExceptionInformation(), OutExceptionMsg, OutExceptionCallstack))
+		{
+			Job.Output.bSucceeded = false;
+		}
+#endif
+	}
+
+	static void CompileShaderInternal(const IShaderFormat* Compiler, FShaderCompileJob& Job, const FString& WorkingDirectory, FString& OutExceptionMsg, FString& OutExceptionCallstack, int32* CompileCount)
+	{
+		double TimeStart = FPlatformTime::Seconds();
+		if (Compiler->SupportsIndependentPreprocessing())
+		{
+			if (!Job.Input.bCachePreprocessed)
+			{
+				PreprocessShaderInternal(Compiler, Job);
+			}
+
+			Job.Output.Errors.Append(Job.PreprocessOutput.Errors);
+
+			if (Job.PreprocessOutput.bSucceeded)
 			{
-				Job.Output.GenerateOutputHash();
-				if (Job.Input.CompressionFormat != NAME_None)
+				if (Job.SecondaryPreprocessOutput.IsValid())
 				{
-					Job.Output.CompressOutput(Job.Input.CompressionFormat, Job.Input.OodleCompressor, Job.Input.OodleLevel);
+					Job.SecondaryOutput = MakeUnique<FShaderCompilerOutput>();
+				}
+				InvokeCompile(Compiler, Job, WorkingDirectory, OutExceptionMsg, OutExceptionCallstack);
+				if (Job.SecondaryOutput.IsValid())
+				{
+					Job.Output.bSucceeded = Job.Output.bSucceeded && Job.SecondaryOutput->bSucceeded;
+					if (Job.Output.bSucceeded)
+					{
+						Job.SecondaryOutput->GenerateOutputHash();
+					}
+					CombineOutputs(Compiler, Job);
 				}
 			}
-			Job.Output.CompileTime = FPlatformTime::Seconds() - TimeStart;
-
-			if (Compiler->UsesHLSLcc(Job.Input))
+			else
 			{
-				Job.Output.bUsedHLSLccCompiler = true;
+				Job.Output.bSucceeded = false;
 			}
+		}
+		else
+		{
+			InvokeCompile(Compiler, Job, WorkingDirectory, OutExceptionMsg, OutExceptionCallstack);
+		}
+
+		if (Job.Output.bSucceeded)
+		{
+			Job.Output.GenerateOutputHash();
 
-			if (CompileCount)
+			if (Job.Input.CompressionFormat != NAME_None)
 			{
-				++(*CompileCount);
+				Job.Output.CompressOutput(Job.Input.CompressionFormat, Job.Input.OodleCompressor, Job.Input.OodleLevel);
 			}
 		}
-#if PLATFORM_WINDOWS
-		__except (HandleShaderCompileException(GetExceptionInformation(), OutExceptionMsg, OutExceptionCallstack))
+		Job.Output.CompileTime = FPlatformTime::Seconds() - TimeStart;
+
+		if (Compiler->UsesHLSLcc(Job.Input))
 		{
-			Job.Output.bSucceeded = false;
+			Job.Output.bUsedHLSLccCompiler = true;
+		}
+
+		if (CompileCount)
+		{
+			++(*CompileCount);
 		}
-#endif
 	}
 };
 
@@ -1121,8 +1227,6 @@ void CompileShader(const TArray<const IShaderFormat*>& ShaderFormats, FShaderCom
 	{
 		Job.Input.Environment.Merge(*Job.Input.SharedEnvironment);
 	}
-
-#if PLATFORM_WINDOWS
 	FString ExceptionMsg;
 	FString ExceptionCallstack;
 	FInternalShaderCompilerFunctions::CompileShaderInternal(Compiler, Job, WorkingDirectory, ExceptionMsg, ExceptionCallstack, CompileCount);
@@ -1135,9 +1239,6 @@ void CompileShader(const TArray<const IShaderFormat*>& ShaderFormats, FShaderCom
 			*ExceptionCallstack);
 		Job.Output.Errors.Add(Error);
 	}
-#else
-	FInternalShaderCompilerFunctions::CompileShaderInternal(Compiler, Job, WorkingDirectory, CompileCount);
-#endif
 
 	Job.bSucceeded = Job.Output.bSucceeded;
 	if (Job.Input.DumpDebugInfoEnabled())
@@ -1148,7 +1249,14 @@ void CompileShader(const TArray<const IShaderFormat*>& ShaderFormats, FShaderCom
 			// don't serialize preprocess output back to the cooker from SCW (if enabled this will occur in the job OnComplete callback)
 			if (!Job.Input.bCachePreprocessed)
 			{
-				Compiler->OutputDebugData(Job.Input, Job.PreprocessOutput, Job.Output);
+				if (Job.SecondaryPreprocessOutput.IsValid() && Job.SecondaryOutput.IsValid())
+				{
+					Compiler->OutputDebugData(Job.Input, Job.PreprocessOutput, *Job.SecondaryPreprocessOutput, Job.Output, *Job.SecondaryOutput);
+				}
+				else
+				{
+					Compiler->OutputDebugData(Job.Input, Job.PreprocessOutput, Job.Output);
+				}
 			}
 		}
 		else
@@ -2832,6 +2940,10 @@ FShaderCommonCompileJob::FInputHash FShaderCompileJob::GetInputHash()
 
 		// const_cast due to serialization API requiring non-const. better than not having const correctness in the API.
 		Hasher << const_cast<FString&>(PreprocessOutput.GetSource());
+		if (SecondaryPreprocessOutput.IsValid())
+		{
+			Hasher << const_cast<FString&>(SecondaryPreprocessOutput->GetSource());
+		}
 		InputHash = Hasher.Finalize();
 	}
 	else
@@ -3012,7 +3124,14 @@ void FShaderCompileJob::OnComplete()
 		// or matched another in-flight job's hash and so could share its results
 		if (Input.bCachePreprocessed && Input.DumpDebugInfoEnabled())
 		{
-			ShaderFormat->OutputDebugData(Input, PreprocessOutput, Output);
+			if (SecondaryPreprocessOutput.IsValid() && SecondaryOutput.IsValid())
+			{
+				ShaderFormat->OutputDebugData(Input, PreprocessOutput, *SecondaryPreprocessOutput, Output, *SecondaryOutput);
+			}
+			else
+			{
+				ShaderFormat->OutputDebugData(Input, PreprocessOutput, Output);
+			}
 		}
 	}
 }
@@ -3024,6 +3143,21 @@ void FShaderCompileJob::SerializeWorkerOutput(FArchive& Ar)
 	// or (b) the "extract shader source" setting is enabled (i.e. something upstream explicitly wants the final source passed to the compiler).
 	Output.bSerializeModifiedSource = (Input.bCachePreprocessed && Input.DumpDebugInfoEnabled()) || Input.ExtraSettings.bExtractShaderSource;
 
+
+	Ar << Output;
+
+	bool bSecondaryOutput = SecondaryOutput.IsValid();
+	Ar << bSecondaryOutput;
+
+	if (bSecondaryOutput)
+	{
+		if (Ar.IsLoading())
+		{
+			SecondaryOutput = MakeUnique<FShaderCompilerOutput>();
+		}
+		Ar << *SecondaryOutput;
+	}
+
 	// edge case for backends which have implemented independent preprocessing API when the preprocessed cache is not enabled.
 	// if no modifications have occurred as part of the compile step, we still need a copy of the source back in the cooker
 	// if bExtractShaderSource is set, so explicitly serialize just that portion of the preprocess output struct here.
@@ -3032,7 +3166,6 @@ void FShaderCompileJob::SerializeWorkerOutput(FArchive& Ar)
 		Ar << PreprocessOutput.EditSource();
 	}
 
-	Ar << Output;
 	bool bSucceededTemp = (bool)bSucceeded;
 	Ar << bSucceededTemp;
 	bSucceeded = bSucceededTemp;
@@ -3045,6 +3178,18 @@ void FShaderCompileJob::SerializeWorkerInput(FArchive& Ar)
 	if (Input.bCachePreprocessed)
 	{
 		Ar << PreprocessOutput;
+
+		bool bSecondaryPreprocessOutput = SecondaryPreprocessOutput.IsValid();
+		Ar << bSecondaryPreprocessOutput;
+
+		if (bSecondaryPreprocessOutput)
+		{
+			if (Ar.IsLoading())
+			{
+				SecondaryPreprocessOutput = MakeUnique<FShaderPreprocessOutput>();
+			}
+			Ar << *SecondaryPreprocessOutput;
+		}
 	}
 }
 
diff --git a/Engine/Source/Runtime/RenderCore/Private/ShaderPreprocessTypes.cpp b/Engine/Source/Runtime/RenderCore/Private/ShaderPreprocessTypes.cpp
--- a/Engine/Source/Runtime/RenderCore/Private/ShaderPreprocessTypes.cpp
+++ b/Engine/Source/Runtime/RenderCore/Private/ShaderPreprocessTypes.cpp
@@ -290,6 +290,7 @@ void FShaderPreprocessOutput::RemapErrors(FShaderCompilerOutput& Output) const
 FArchive& operator<<(FArchive& Ar, FShaderPreprocessOutput& PreprocessOutput)
 {
 	Ar << PreprocessOutput.bSucceeded;
+	Ar << PreprocessOutput.bIsSecondary;
 	Ar << PreprocessOutput.Errors;
 	Ar << PreprocessOutput.PragmaDirectives;
 	Ar << PreprocessOutput.PreprocessedSource;
diff --git a/Engine/Source/Runtime/RenderCore/Public/ShaderCompilerCore.h b/Engine/Source/Runtime/RenderCore/Public/ShaderCompilerCore.h
--- a/Engine/Source/Runtime/RenderCore/Public/ShaderCompilerCore.h
+++ b/Engine/Source/Runtime/RenderCore/Public/ShaderCompilerCore.h
@@ -27,9 +27,9 @@ class FShaderCompileJob;
 class FShaderPipelineCompileJob;
 
 // this is for the protocol, not the data, bump if FShaderCompilerInput or ProcessInputFromArchive changes.
-inline const int32 ShaderCompileWorkerInputVersion = 20;
+inline const int32 ShaderCompileWorkerInputVersion = 21;
 // this is for the protocol, not the data, bump if FShaderCompilerOutput or WriteToOutputArchive changes.
-inline const int32 ShaderCompileWorkerOutputVersion = 14;
+inline const int32 ShaderCompileWorkerOutputVersion = 15;
 // this is for the protocol, not the data.
 inline const int32 ShaderCompileWorkerSingleJobHeader = 'S';
 // this is for the protocol, not the data.
diff --git a/Engine/Source/Runtime/RenderCore/Public/ShaderCompilerJobTypes.h b/Engine/Source/Runtime/RenderCore/Public/ShaderCompilerJobTypes.h
--- a/Engine/Source/Runtime/RenderCore/Public/ShaderCompilerJobTypes.h
+++ b/Engine/Source/Runtime/RenderCore/Public/ShaderCompilerJobTypes.h
@@ -201,7 +201,9 @@ public:
 	/** Input for the shader compile */
 	FShaderCompilerInput Input;
 	FShaderPreprocessOutput PreprocessOutput;
+	TUniquePtr<FShaderPreprocessOutput> SecondaryPreprocessOutput{};
 	FShaderCompilerOutput Output;
+	TUniquePtr<FShaderCompilerOutput> SecondaryOutput{};
 
 	// List of pipelines that are sharing this job.
 	TMap<const FVertexFactoryType*, TArray<const FShaderPipelineType*>> SharingPipelines;
diff --git a/Engine/Source/Runtime/RenderCore/Public/ShaderPreprocessTypes.h b/Engine/Source/Runtime/RenderCore/Public/ShaderPreprocessTypes.h
--- a/Engine/Source/Runtime/RenderCore/Public/ShaderPreprocessTypes.h
+++ b/Engine/Source/Runtime/RenderCore/Public/ShaderPreprocessTypes.h
@@ -105,6 +105,11 @@ public:
 		}
 	}
 
+	inline bool IsSecondary() const
+	{
+		return bIsSecondary;
+	}
+
 	inline void LogError(FString&& Message)
 	{
 		FShaderCompilerError& CompilerError = Errors.AddDefaulted_GetRef();;
@@ -176,6 +181,7 @@ private:
 
 	double ElapsedTime = 0.0;
 	bool bSucceeded = false;
+	bool bIsSecondary = false;
 
 	FShaderParameterParser ParameterParser;
 };
