Nix RDG pass lambdas on server.

- TRDGLambdaPass<>::Execute() is responsible for 0.8% (couple MBs in abs size) of generated (x86-64) code, top 10 symbols contributing to the binary size.

diff --git a/Engine/Source/Runtime/RenderCore/Public/RenderGraphBuilder.inl b/Engine/Source/Runtime/RenderCore/Public/RenderGraphBuilder.inl
--- a/Engine/Source/Runtime/RenderCore/Public/RenderGraphBuilder.inl
+++ b/Engine/Source/Runtime/RenderCore/Public/RenderGraphBuilder.inl
@@ -217,10 +217,15 @@ FORCEINLINE FRDGSubresourceState* FRDGBuilder::AllocSubresource(const FRDGSubres
 template <typename ParameterStructType>
 TRDGUniformBufferRef<ParameterStructType> FRDGBuilder::CreateUniformBuffer(const ParameterStructType* ParameterStruct)
 {
+#if !USE_NULL_RHI
 	IF_RDG_ENABLE_DEBUG(UserValidation.ValidateCreateUniformBuffer(ParameterStruct, ParameterStructType::FTypeInfo::GetStructMetadata()));
 	auto* UniformBuffer = UniformBuffers.Allocate<TRDGUniformBuffer<ParameterStructType>>(Allocator, ParameterStruct, ParameterStructType::FTypeInfo::GetStructMetadata()->GetShaderVariableName());
 	IF_RDG_ENABLE_DEBUG(UserValidation.ValidateCreateUniformBuffer(UniformBuffer));
 	return UniformBuffer;
+#else
+	checkNoEntry();
+	return nullptr;
+#endif // !USE_NULL_RHI
 }
 
 template <typename ExecuteLambdaType>
@@ -229,6 +234,7 @@ FRDGPassRef FRDGBuilder::AddPass(
 	ERDGPassFlags Flags,
 	ExecuteLambdaType&& ExecuteLambda)
 {
+#if !USE_NULL_RHI
 	using LambdaPassType = TRDGEmptyLambdaPass<ExecuteLambdaType>;
 
 	IF_RDG_ENABLE_DEBUG(UserValidation.ValidateAddPass(Name, Flags));
@@ -240,8 +246,13 @@ FRDGPassRef FRDGBuilder::AddPass(
 	LambdaPassType* Pass = Passes.Allocate<LambdaPassType>(Allocator, MoveTemp(Name), Flags, MoveTemp(ExecuteLambda));
 	SetupEmptyPass(Pass);
 	return Pass;
+#else
+	checkNoEntry();
+	return nullptr;
+#endif // !USE_NULL_RHI
 }
 
+#if !USE_NULL_RHI
 template <typename ParameterStructType, typename ExecuteLambdaType>
 FRDGPassRef FRDGBuilder::AddPassInternal(
 	FRDGEventName&& Name,
@@ -268,6 +279,7 @@ FRDGPassRef FRDGBuilder::AddPassInternal(
 	SetupParameterPass(Pass);
 	return Pass;
 }
+#endif // !USE_NULL_RHI
 
 template <typename ExecuteLambdaType>
 FRDGPassRef FRDGBuilder::AddPass(
@@ -277,7 +289,12 @@ FRDGPassRef FRDGBuilder::AddPass(
 	ERDGPassFlags Flags,
 	ExecuteLambdaType&& ExecuteLambda)
 {
+#if !USE_NULL_RHI
 	return AddPassInternal(Forward<FRDGEventName>(Name), ParametersMetadata, ParameterStruct, Flags, Forward<ExecuteLambdaType>(ExecuteLambda));
+#else
+	checkNoEntry();
+	return nullptr;
+#endif // !USE_NULL_RHI
 }
 
 template <typename ParameterStructType, typename ExecuteLambdaType>
@@ -287,7 +304,12 @@ FRDGPassRef FRDGBuilder::AddPass(
 	ERDGPassFlags Flags,
 	ExecuteLambdaType&& ExecuteLambda)
 {
+#if !USE_NULL_RHI
 	return AddPassInternal(Forward<FRDGEventName>(Name), ParameterStructType::FTypeInfo::GetStructMetadata(), ParameterStruct, Flags, Forward<ExecuteLambdaType>(ExecuteLambda));
+#else
+	checkNoEntry();
+	return nullptr;
+#endif // !USE_NULL_RHI
 }
 
 inline void FRDGBuilder::SetPassWorkload(FRDGPass* Pass, uint32 Workload)
diff --git a/Engine/Source/Runtime/RenderCore/Public/RenderGraphPass.h b/Engine/Source/Runtime/RenderCore/Public/RenderGraphPass.h
--- a/Engine/Source/Runtime/RenderCore/Public/RenderGraphPass.h
+++ b/Engine/Source/Runtime/RenderCore/Public/RenderGraphPass.h
@@ -610,9 +610,13 @@ private:
 
 	void Execute(FRHIComputeCommandList& RHICmdList) override
 	{
+#if !USE_NULL_RHI
 		QUICK_SCOPE_CYCLE_COUNTER(STAT_FRDGPass_Execute);
 		RHICmdList.SetStaticUniformBuffers(ParameterStruct.GetStaticUniformBuffers());
 		ExecuteLambdaFunc<TRDGPass>(static_cast<TRHICommandList&>(RHICmdList));
+#else
+		checkNoEntry();
+#endif // !USE_NULL_RHI
 	}
 
 	ExecuteLambdaType ExecuteLambda;
