adding "return false" to except block in preprocessor invocation to satisfy msvc static analysis

diff --git a/Engine/Source/Runtime/RenderCore/Private/ShaderCore.cpp b/Engine/Source/Runtime/RenderCore/Private/ShaderCore.cpp
--- a/Engine/Source/Runtime/RenderCore/Private/ShaderCore.cpp
+++ b/Engine/Source/Runtime/RenderCore/Private/ShaderCore.cpp
@@ -993,6 +993,7 @@ public:
 #if PLATFORM_WINDOWS
 		__except (HandleShaderCompileException(GetExceptionInformation(), OutExceptionMsg, OutExceptionCallstack))
 		{
+			return false;
 		}
 #endif
 		return false;
