[shaders] Fix for "duplicate definitions" errors when a material shader compile job  is executed in the cook process. This was caused by some (rather bizzarre) behaviour in FShaderCompileEnvironment::Merge which causes in-memory source files to be appended to each other if the same key is found in both the source and target environment, along with an erroneous/redundant merge step occurring in the compile path if a backend has implemented the independent preprocessing API but the preprocessed cache is disabled.

diff --git a/Engine/Source/Runtime/RenderCore/Private/ShaderCore.cpp b/Engine/Source/Runtime/RenderCore/Private/ShaderCore.cpp
--- a/Engine/Source/Runtime/RenderCore/Private/ShaderCore.cpp
+++ b/Engine/Source/Runtime/RenderCore/Private/ShaderCore.cpp
@@ -1022,7 +1022,9 @@ public:
 		FString ExceptionCallstack;
 		FString ExceptionMsg;
 
-		if (IsValidRef(Job.Input.SharedEnvironment))
+		// if preprocessed cache is disabled, this function is executed as part of the compile process
+		// which will have merged the environment inline prior to calling this, so we can skip merging here 
+		if (Job.Input.bCachePreprocessed && IsValidRef(Job.Input.SharedEnvironment))
 		{
 			// only create new environment & merge if necessary, save some allocs
 			// (need a copy here as we don't want to merge the environment in place like we do in the compile path
