Fix global usage of all cvars related to shader pipelines:

r.ShaderPipelines
r.Material.ExcludeNonPipelinedShaders

Both cvars are currently disable on WindowsEngine.ini which make them also disable when we cook for consoles. ExcludeNonPipelinedShaderTypes is currently used only at cooking time when we compile shaders which is currently returning false, since the value is disabled on windows. Making those cvars use FShaderPlatformCachedIniValue, will allow to have specific platform value when cooking for a specific target. Moving r.Material.ExcludeNonPipelinedShaders to the defaultengine.ini file to keep it disable for the moment and not affecting shader key string. Will enable it progressively on console later.

r.Shaders.RemoveUnusedInterpolators

Make that cvar platform agnostic, since it will be used in on other consoles eventually. Platforms supporting the stripping of unused interpolators will check the same flag (CFLAG_ForceRemoveUnusedInterpolators) to enable it. Since windows can run opengl, and stripping is not supported on that shader platform, add extra logic to exclude the flag.

-Make the RHISupportsShaderPipelines function access a real DDPI attribute and not just test if we are not on mobile. Will need this to be able to have stripping unused shader interpolators on Switch.

diff --git a/Engine/Config/Android/DataDrivenPlatformInfo.ini b/Engine/Config/Android/DataDrivenPlatformInfo.ini
--- a/Engine/Config/Android/DataDrivenPlatformInfo.ini
+++ b/Engine/Config/Android/DataDrivenPlatformInfo.ini
@@ -64,6 +64,7 @@ bSupportsMobileMultiView=true
 bNeedsOfflineCompiler=true
 bSupportsManualVertexFetch=false
 bSupportsClipDistance = false
+bSupportsShaderPipelines = false
 FriendlyName=LOCTEXT("FriendlyShaderPlatformName_Android", "Android OpenGL Mobile")
 
 [ShaderPlatform VULKAN_ES3_1_ANDROID]
@@ -79,6 +80,7 @@ bNeedsOfflineCompiler=true
 bSupportsIndexBufferUAVs=true
 bSupportsManualVertexFetch=false
 bSupportsClipDistance = false
+bSupportsShaderPipelines = false
 FriendlyName=LOCTEXT("FriendlyShaderPlatformName_AndroidVulkan", "Android Vulkan Mobile")
 
 [ShaderPlatform VULKAN_SM5_ANDROID]
diff --git a/Engine/Config/IOS/DataDrivenPlatformInfo.ini b/Engine/Config/IOS/DataDrivenPlatformInfo.ini
--- a/Engine/Config/IOS/DataDrivenPlatformInfo.ini
+++ b/Engine/Config/IOS/DataDrivenPlatformInfo.ini
@@ -54,6 +54,7 @@ bTargetsTiledGPU=true
 bSupports4ComponentUAVReadWrite=true
 bSupportsManualVertexFetch=false
 bSupportsVolumeTextureAtomics=false
+bSupportsShaderPipelines = false
 FriendlyName=LOCTEXT("FriendlyShaderPlatformName_IOSMobile", "IOS Metal Mobile")
 
 [ShaderPlatform METAL_MRT]
diff --git a/Engine/Config/Mac/DataDrivenPlatformInfo.ini b/Engine/Config/Mac/DataDrivenPlatformInfo.ini
--- a/Engine/Config/Mac/DataDrivenPlatformInfo.ini
+++ b/Engine/Config/Mac/DataDrivenPlatformInfo.ini
@@ -85,4 +85,5 @@ bIsMobile=true
 bSupports4ComponentUAVReadWrite=true
 bSupportsManualVertexFetch=false
 bSupportsVolumeTextureAtomics=false
-FriendlyName=LOCTEXT("FriendlyShaderPlatformName_MacMobile", "Mac Metal Mobile")
+bSupportsShaderPipelines = false
+FriendlyName=LOCTEXT("FriendlyShaderPlatformName_MacMobile", "Mac Metal Mobile")
\ No newline at end of file
diff --git a/Engine/Config/TVOS/DataDrivenPlatformInfo.ini b/Engine/Config/TVOS/DataDrivenPlatformInfo.ini
--- a/Engine/Config/TVOS/DataDrivenPlatformInfo.ini
+++ b/Engine/Config/TVOS/DataDrivenPlatformInfo.ini
@@ -36,4 +36,5 @@ bIsMobile=true
 bTargetsTiledGPU=true
 bSupports4ComponentUAVReadWrite=true
 bSupportsManualVertexFetch=false
+bSupportsShaderPipelines = false
 FriendlyName=LOCTEXT("FriendlyShaderPlatformName_TVOS_Mobile", "TVOS Metal Mobile")
\ No newline at end of file
diff --git a/Engine/Config/VulkanPC/DataDrivenPlatformInfo.ini b/Engine/Config/VulkanPC/DataDrivenPlatformInfo.ini
--- a/Engine/Config/VulkanPC/DataDrivenPlatformInfo.ini
+++ b/Engine/Config/VulkanPC/DataDrivenPlatformInfo.ini
@@ -264,9 +264,9 @@ bSupportsManualVertexFetch=false
 ; The support really is conditional, the runtime code needs to also test GRHISupportsArrayIndexFromAnyShader
 bSupportsVertexShaderLayer=true
 bSupportsSceneDataCompressedTransforms=true
+bSupportsShaderPipelines = false
 FriendlyName=LOCTEXT("FriendlyShaderPlatformName_VulkanPC_Mobile", "PC Vulkan Mobile")
 
-
 [PreviewPlatform VULKAN_SM5]
 PlatformName=VULKAN_SM5
 ShaderFormat=SF_VULKAN_SM5
diff --git a/Engine/Config/Windows/DataDrivenPlatformInfo.ini b/Engine/Config/Windows/DataDrivenPlatformInfo.ini
--- a/Engine/Config/Windows/DataDrivenPlatformInfo.ini
+++ b/Engine/Config/Windows/DataDrivenPlatformInfo.ini
@@ -157,6 +157,7 @@ bSupportsVertexShaderLayer=true
 bSupportsManualVertexFetch=false
 bSupportsSceneDataCompressedTransforms=true
 bSupportsDebugViewShaders=true
+bSupportsShaderPipelines = false
 FriendlyName=LOCTEXT("FriendlyShaderPlatformName_D3D_Mobile", "PC D3D Mobile")
 SupportsMultiViewport=RuntimeDependent
 
@@ -175,4 +176,5 @@ ShaderFormat=GLSL_150_ES31
 bIsMobile=true
 bSupportsManualVertexFetch=false
 bSupportsClipDistance = false
-FriendlyName=LOCTEXT("FriendlyShaderPlatformName_PCGL_Mobile", "PC OpenGL Mobile")
+bSupportsShaderPipelines = false
+FriendlyName=LOCTEXT("FriendlyShaderPlatformName_PCGL_Mobile", "PC OpenGL Mobile")
\ No newline at end of file
diff --git a/Engine/Config/Windows/WindowsEngine.ini b/Engine/Config/Windows/WindowsEngine.ini
--- a/Engine/Config/Windows/WindowsEngine.ini
+++ b/Engine/Config/Windows/WindowsEngine.ini
@@ -23,3 +23,6 @@ r.Shadow.DetectVertexShaderLayerAtRuntime=1
 
 [PlatformCrypto]
 PlatformRequiresDataCrypto=True
+
+[/Script/Engine.RendererSettings]
+r.Shaders.RemoveUnusedInterpolators=1
\ No newline at end of file
diff --git a/Engine/Source/Runtime/Engine/Private/Materials/MaterialShared.cpp b/Engine/Source/Runtime/Engine/Private/Materials/MaterialShared.cpp
--- a/Engine/Source/Runtime/Engine/Private/Materials/MaterialShared.cpp
+++ b/Engine/Source/Runtime/Engine/Private/Materials/MaterialShared.cpp
@@ -3388,7 +3388,7 @@ bool FMaterial::TryGetShaders(const FMaterialShaderTypes& InTypes, const FVertex
 {
 	// TRACE_CPUPROFILER_EVENT_SCOPE(FMaterial::TryGetShaders); // <-- disabled by default due to verbosity (hundreds of calls per frame)
 
-	static const auto* CVarShaderPipelines = IConsoleManager::Get().FindTConsoleVariableDataInt(TEXT("r.ShaderPipelines"));
+
 	const bool bIsInGameThread = (IsInGameThread() || IsInParallelGameThread());
 	const FMaterialShaderMap* ShaderMap = bIsInGameThread ? GameThreadShaderMap : RenderingThreadShaderMap;
 	const bool bShaderMapComplete = bIsInGameThread ? IsGameThreadShaderMapComplete() : IsRenderingThreadShaderMapComplete();
@@ -3432,9 +3432,7 @@ bool FMaterial::TryGetShaders(const FMaterialShaderTypes& InTypes, const FVertex
 		return false;
 	};
 
-	if (InTypes.PipelineType &&
-		RHISupportsShaderPipelines(ShaderPlatform) &&
-		CVarShaderPipelines && CVarShaderPipelines->GetValueOnAnyThread(bIsInGameThread) != 0)
+	if (InTypes.PipelineType && RHISupportsShaderPipelines(ShaderPlatform) && UseShaderPipelines(ShaderPlatform))
 	{
 		FShaderPipeline* Pipeline = ShaderMapContent ? ShaderMapContent->GetShaderPipeline(InTypes.PipelineType) : nullptr;
 		if (Pipeline)
diff --git a/Engine/Source/Runtime/Engine/Private/ShaderCompiler/ShaderCompiler.cpp b/Engine/Source/Runtime/Engine/Private/ShaderCompiler/ShaderCompiler.cpp
--- a/Engine/Source/Runtime/Engine/Private/ShaderCompiler/ShaderCompiler.cpp
+++ b/Engine/Source/Runtime/Engine/Private/ShaderCompiler/ShaderCompiler.cpp
@@ -1139,16 +1139,6 @@ static TAutoConsoleVariable<int32> CVarShaderFlowControl(
 	TEXT("\t2: Avoid - Attempt to unroll and flatten flow-control.\n"),
 	ECVF_ReadOnly);
 
-static TAutoConsoleVariable<int32> CVarD3DRemoveUnusedInterpolators(
-	TEXT("r.D3D.RemoveUnusedInterpolators"),
-	1,
-	TEXT("Enables removing unused interpolators mode when compiling pipelines for D3D.\n")
-	TEXT(" -1: Do not actually remove, but make the app think it did (for debugging)\n")
-	TEXT(" 0: Disable (default)\n")
-	TEXT(" 1: Enable removing unused"),
-	ECVF_ReadOnly
-	);
-
 static TAutoConsoleVariable<int32> CVarD3DCheckedForTypedUAVs(
 	TEXT("r.D3D.CheckedForTypedUAVs"),
 	1,
@@ -6369,12 +6359,14 @@ void GlobalBeginCompileShader(
 		Input.Environment.CompilerFlags.Add(CFLAG_WarningsAsErrors);
 	}
 
+	if (UseRemoveUnsedInterpolators((EShaderPlatform)Target.Platform) && !IsOpenGLPlatform((EShaderPlatform)Target.Platform))
+	{
+		Input.Environment.CompilerFlags.Add(CFLAG_ForceRemoveUnusedInterpolators);
+	}
+	
 	if (IsD3DPlatform((EShaderPlatform)Target.Platform) && IsPCPlatform((EShaderPlatform)Target.Platform))
 	{
-		if (CVarD3DRemoveUnusedInterpolators.GetValueOnAnyThread() != 0)
-		{
-			Input.Environment.CompilerFlags.Add(CFLAG_ForceRemoveUnusedInterpolators);
-		}
+
 
 		if (CVarD3DCheckedForTypedUAVs.GetValueOnAnyThread() == 0)
 		{
diff --git a/Engine/Source/Runtime/RHI/Private/DataDrivenShaderPlatformInfo.cpp b/Engine/Source/Runtime/RHI/Private/DataDrivenShaderPlatformInfo.cpp
--- a/Engine/Source/Runtime/RHI/Private/DataDrivenShaderPlatformInfo.cpp
+++ b/Engine/Source/Runtime/RHI/Private/DataDrivenShaderPlatformInfo.cpp
@@ -145,6 +145,7 @@ void FGenericDataDrivenShaderPlatformInfo::SetDefaultValues()
 	bSupportsManualVertexFetch = true;
 	bSupportsVolumeTextureAtomics = true;
 	bSupportsClipDistance = true;
+	bSupportsShaderPipelines = true;
 }
 
 void FGenericDataDrivenShaderPlatformInfo::ParseDataDrivenShaderInfo(const FConfigSection& Section, uint32 Index)
@@ -282,6 +283,7 @@ void FGenericDataDrivenShaderPlatformInfo::ParseDataDrivenShaderInfo(const FConf
 	GET_SECTION_BOOL_HELPER(bSupportsSwapchainUAVs);
 	GET_SECTION_BOOL_HELPER(bSupportsClipDistance);
 	GET_SECTION_BOOL_HELPER(bSupportsNNEShaders);
+	GET_SECTION_BOOL_HELPER(bSupportsShaderPipelines);
 #undef GET_SECTION_BOOL_HELPER
 #undef GET_SECTION_INT_HELPER
 #undef GET_SECTION_SUPPORT_HELPER
diff --git a/Engine/Source/Runtime/RHI/Public/DataDrivenShaderPlatformInfo.h b/Engine/Source/Runtime/RHI/Public/DataDrivenShaderPlatformInfo.h
--- a/Engine/Source/Runtime/RHI/Public/DataDrivenShaderPlatformInfo.h
+++ b/Engine/Source/Runtime/RHI/Public/DataDrivenShaderPlatformInfo.h
@@ -111,7 +111,7 @@ class FGenericDataDrivenShaderPlatformInfo
 	uint32 bSupportsSwapchainUAVs : 1;
 	uint32 bSupportsClipDistance : 1;
 	uint32 bSupportsNNEShaders: 1;
-
+	uint32 bSupportsShaderPipelines : 1;
 	// NOTE: When adding fields, you must also add to ParseDataDrivenShaderInfo!
 	uint32 bContainsValidPlatformInfo : 1;
 
@@ -687,6 +687,11 @@ public:
 		return Infos[Platform].bSupportsVolumeTextureAtomics;
 	}
 
+	static FORCEINLINE_DEBUGGABLE const bool GetSupportsPipelineShaders(const FStaticShaderPlatform Platform)
+	{
+		return Infos[Platform].bSupportsShaderPipelines;
+	}
+
 	static FORCEINLINE_DEBUGGABLE const bool GetSupportsROV(const FStaticShaderPlatform Platform)
 	{
 		return Infos[Platform].bSupportsROV;
@@ -920,7 +925,7 @@ inline bool RHISupportsNativeShaderLibraries(const FStaticShaderPlatform Platfor
 
 inline bool RHISupportsShaderPipelines(const FStaticShaderPlatform Platform)
 {
-	return !IsMobilePlatform(Platform);
+	return FDataDrivenShaderPlatformInfo::GetSupportsPipelineShaders(Platform);
 }
 
 inline bool RHISupportsDualSourceBlending(const FStaticShaderPlatform Platform)
diff --git a/Engine/Source/Runtime/RenderCore/Private/RenderUtils.cpp b/Engine/Source/Runtime/RenderCore/Private/RenderUtils.cpp
--- a/Engine/Source/Runtime/RenderCore/Private/RenderUtils.cpp
+++ b/Engine/Source/Runtime/RenderCore/Private/RenderUtils.cpp
@@ -1195,17 +1195,29 @@ RENDERCORE_API bool UseNaniteLandscapeMesh(EShaderPlatform ShaderPlatform)
 	return DoesPlatformSupportNanite(ShaderPlatform);
 }
 
+RENDERCORE_API bool UseShaderPipelines(EShaderPlatform ShaderPlatform)
+{
+	static FShaderPlatformCachedIniValue<int32> PerPlatformCVar(TEXT("r.ShaderPipelines"));
+	bool bUseShaderPipelines = PerPlatformCVar.Get(ShaderPlatform) != 0;
+	return bUseShaderPipelines;
+}
+
+RENDERCORE_API bool UseRemoveUnsedInterpolators(EShaderPlatform ShaderPlatform)
+{
+	static FShaderPlatformCachedIniValue<int32> PerPlatformCVar(TEXT("r.Shaders.RemoveUnusedInterpolators"));
+	bool bRemoveUnusedInterpolators = PerPlatformCVar.Get(ShaderPlatform) != 0;
+	return bRemoveUnusedInterpolators;
+}
+
 RENDERCORE_API bool ExcludeNonPipelinedShaderTypes(EShaderPlatform ShaderPlatform)
 {
 	if (RHISupportsShaderPipelines(ShaderPlatform))
 	{
-		static const TConsoleVariableData<int32>* CVarShaderPipelines = IConsoleManager::Get().FindTConsoleVariableDataInt(TEXT("r.ShaderPipelines"));
-		bool bShaderPipelinesAreEnabled = CVarShaderPipelines && CVarShaderPipelines->GetValueOnAnyThread(IsInGameThread()) != 0;
+		bool bShaderPipelinesAreEnabled = UseShaderPipelines(ShaderPlatform);
 		if (bShaderPipelinesAreEnabled)
 		{
-			static const IConsoleVariable* CVarExcludeNonPipelinedShaders = IConsoleManager::Get().FindConsoleVariable(TEXT("r.Material.ExcludeNonPipelinedShaders"));
-			bool bExcludeNonPipelinedShaders = CVarExcludeNonPipelinedShaders && CVarExcludeNonPipelinedShaders->GetInt() != 0;
-
+			static FShaderPlatformCachedIniValue<int32> PerPlatformCVar(TEXT("r.Material.ExcludeNonPipelinedShaders"));
+			bool bExcludeNonPipelinedShaders = PerPlatformCVar.Get(ShaderPlatform) != 0;
 			return bExcludeNonPipelinedShaders;
 		}
 	}
diff --git a/Engine/Source/Runtime/RenderCore/Private/Shader.cpp b/Engine/Source/Runtime/RenderCore/Private/Shader.cpp
--- a/Engine/Source/Runtime/RenderCore/Private/Shader.cpp
+++ b/Engine/Source/Runtime/RenderCore/Private/Shader.cpp
@@ -81,6 +81,15 @@ static TAutoConsoleVariable<int32> CVarUsePipelines(
 	1,
 	TEXT("Enable using Shader pipelines."));
 
+static TAutoConsoleVariable<int32> CVarRemoveUnusedInterpolators(
+	TEXT("r.Shaders.RemoveUnusedInterpolators"),
+	0,
+	TEXT("Enables removing unused interpolators mode when compiling shader pipelines.\n")
+	TEXT(" 0: Disable (default)\n")
+	TEXT(" 1: Enable removing unused"),
+	ECVF_ReadOnly
+);
+
 static TAutoConsoleVariable<int32> CVarSkipShaderCompression(
 	TEXT("r.Shaders.SkipCompression"),
 	0,
@@ -1647,15 +1656,9 @@ void ShaderMapAppendKeyString(EShaderPlatform Platform, FString& KeyString)
 		KeyString += TEXT("_NoPPSM");
 	}
 	
-	if (IsD3DPlatform(Platform) && IsPCPlatform(Platform))
+	if (UseRemoveUnsedInterpolators(Platform) && !IsOpenGLPlatform(Platform))
 	{
-		{
-			static const auto CVar = IConsoleManager::Get().FindConsoleVariable(TEXT("r.D3D.RemoveUnusedInterpolators"));
-			if (CVar && CVar->GetInt() != 0)
-			{
-				KeyString += TEXT("_UnInt");
-			}
-		}
+		KeyString += TEXT("_UnInt");
 	}
 
 	if (IsMobilePlatform(Platform))
diff --git a/Engine/Source/Runtime/RenderCore/Public/RenderUtils.h b/Engine/Source/Runtime/RenderCore/Public/RenderUtils.h
--- a/Engine/Source/Runtime/RenderCore/Public/RenderUtils.h
+++ b/Engine/Source/Runtime/RenderCore/Public/RenderUtils.h
@@ -423,6 +423,16 @@ RENDERCORE_API bool UseNaniteLandscapeMesh(EShaderPlatform ShaderPlatform);
  */
 RENDERCORE_API bool ExcludeNonPipelinedShaderTypes(EShaderPlatform ShaderPlatform);
 
+/**
+ *  Checks if shader pipelines is enable for a specific platform.
+ */
+RENDERCORE_API bool UseShaderPipelines(EShaderPlatform ShaderPlatform);
+
+/**
+ *  Checks if we can strip unused interpolators for a specific platform.
+ */
+RENDERCORE_API bool UseRemoveUnsedInterpolators(EShaderPlatform ShaderPlatform);
+
 /**
  *   Checks if skin cache shaders are enabled for the platform (via r.SkinCache.CompileShaders)
  */
diff --git a/Engine/Source/Runtime/Renderer/Private/DepthRendering.cpp b/Engine/Source/Runtime/Renderer/Private/DepthRendering.cpp
--- a/Engine/Source/Runtime/Renderer/Private/DepthRendering.cpp
+++ b/Engine/Source/Runtime/Renderer/Private/DepthRendering.cpp
@@ -160,11 +160,7 @@ static bool IsDepthPassWaitForTasksEnabled()
 	return CVarRHICmdFlushRenderThreadTasksPrePass.GetValueOnRenderThread() > 0 || CVarRHICmdFlushRenderThreadTasks.GetValueOnRenderThread() > 0;
 }
 
-static FORCEINLINE bool UseShaderPipelines(ERHIFeatureLevel::Type InFeatureLevel)
-{
-	static const auto* CVar = IConsoleManager::Get().FindTConsoleVariableDataInt(TEXT("r.ShaderPipelines"));
-	return RHISupportsShaderPipelines(GShaderPlatformForFeatureLevel[InFeatureLevel]) && CVar && CVar->GetValueOnAnyThread() != 0;
-}
+
 
 template <bool bPositionOnly>
 bool GetDepthPassShaders(
