Removing deprecated shader parameter utils.

diff --git a/Engine/Source/Runtime/RenderCore/Public/ShaderParameterUtils.h b/Engine/Source/Runtime/RenderCore/Public/ShaderParameterUtils.h
--- a/Engine/Source/Runtime/RenderCore/Public/ShaderParameterUtils.h
+++ b/Engine/Source/Runtime/RenderCore/Public/ShaderParameterUtils.h
@@ -505,28 +505,6 @@ GUARD_SETSHADERVALUE(FBox3)
  * Sets the value of a shader surface parameter (e.g. to access MSAA samples).
  * Template'd on shader type (e.g. pixel shader or compute shader).
  */
-template<typename TRHIShader, typename TRHICmdList>
-UE_DEPRECATED(5.2, "SetTextureParameter with an index can't be supported anymore. Your code should be changed to use shader parameter structs to utilize resource arrays.")
-FORCEINLINE void SetTextureParameter(TRHICmdList& RHICmdList, TRHIShader* Shader, const FShaderResourceParameter& Parameter, FRHITexture* TextureRHI, uint32 ElementIndex)
-{
-	if (Parameter.IsBound() && ElementIndex < Parameter.GetNumResources())
-	{
-		FRHIBatchedShaderParameters& BatchedParameters = RHICmdList.GetScratchShaderParameters();
-#if PLATFORM_SUPPORTS_BINDLESS_RENDERING
-		if (Parameter.GetType() == EShaderParameterType::BindlessResourceIndex)
-		{
-			checkf(ElementIndex == 0, TEXT("Bindless resources don't support element offsets"));
-			BatchedParameters.SetBindlessTexture(Parameter.GetBaseIndex(), TextureRHI);
-		}
-		else
-#endif
-		{
-			BatchedParameters.SetShaderTexture(Parameter.GetBaseIndex() + ElementIndex, TextureRHI);
-		}
-		RHICmdList.SetBatchedShaderParameters(Shader, BatchedParameters);
-	}
-}
-
 template<typename TRHIShader, typename TRHICmdList>
 UE_DEPRECATED(5.3, "SetTextureParameter with FRHIBatchedShaderParameters should be used.")
 FORCEINLINE void SetTextureParameter(TRHICmdList& RHICmdList, TRHIShader* Shader, const FShaderResourceParameter& Parameter, FRHITexture* TextureRHI)
@@ -539,28 +517,6 @@ FORCEINLINE void SetTextureParameter(TRHICmdList& RHICmdList, TRHIShader* Shader
 /**
  * Sets the value of a shader sampler parameter. Template'd on shader type.
  */
-template<typename TRHIShader, typename TRHICmdList>
-UE_DEPRECATED(5.2, "SetSamplerParameter with an index can't be supported anymore. Your code should be changed to use shader parameter structs to utilize resource arrays.")
-FORCEINLINE void SetSamplerParameter(TRHICmdList& RHICmdList, TRHIShader* Shader, const FShaderResourceParameter& Parameter, FRHISamplerState* SamplerStateRHI, uint32 ElementIndex)
-{
-	if (Parameter.IsBound() && ElementIndex < Parameter.GetNumResources())
-	{
-		FRHIBatchedShaderParameters& BatchedParameters = RHICmdList.GetScratchShaderParameters();
-#if PLATFORM_SUPPORTS_BINDLESS_RENDERING
-		if (Parameter.GetType() == EShaderParameterType::BindlessSamplerIndex)
-		{
-			checkf(ElementIndex == 0, TEXT("Bindless resources don't support element offsets"));
-			BatchedParameters.SetBindlessSampler(Parameter.GetBaseIndex(), SamplerStateRHI);
-		}
-		else
-#endif
-		{
-			BatchedParameters.SetShaderSampler(Parameter.GetBaseIndex() + ElementIndex, SamplerStateRHI);
-		}
-		RHICmdList.SetBatchedShaderParameters(Shader, BatchedParameters);
-	}
-}
-
 template<typename TRHIShader, typename TRHICmdList>
 UE_DEPRECATED(5.3, "SetSamplerParameter with FRHIBatchedShaderParameters should be used.")
 FORCEINLINE void SetSamplerParameter(TRHICmdList& RHICmdList, TRHIShader* Shader, const FShaderResourceParameter& Parameter, FRHISamplerState* SamplerStateRHI)
@@ -573,32 +529,6 @@ FORCEINLINE void SetSamplerParameter(TRHICmdList& RHICmdList, TRHIShader* Shader
 /**
  * Sets the value of a shader texture parameter. Template'd on shader type.
  */
-template<typename TRHIShader, typename TRHICmdList>
-UE_DEPRECATED(5.2, "SetTextureParameter with an index can't be supported anymore. Your code should be changed to use shader parameter structs to utilize resource arrays.")
-FORCEINLINE void SetTextureParameter(
-	TRHICmdList& RHICmdList,
-	TRHIShader* Shader,
-	const FShaderResourceParameter& TextureParameter,
-	const FShaderResourceParameter& SamplerParameter,
-	FRHISamplerState* SamplerStateRHI,
-	FRHITexture* TextureRHI,
-	uint32 ElementIndex
-	)
-{
-PRAGMA_DISABLE_DEPRECATION_WARNINGS
-	SetTextureParameter(RHICmdList, Shader, TextureParameter, TextureRHI, ElementIndex);
-	
-	// @todo UE samplerstate Should we maybe pass in two separate values? SamplerElement and TextureElement? Or never allow an array of samplers? Unsure best
-	// if there is a matching sampler for this texture array index (ElementIndex), then set it. This will help with this case:
-	//			Texture2D LightMapTextures[NUM_LIGHTMAP_COEFFICIENTS];
-	//			SamplerState LightMapTexturesSampler;
-	// In this case, we only set LightMapTexturesSampler when ElementIndex is 0, we don't set the sampler state for all 4 textures
-	// This assumes that the all textures want to use the same sampler state
-
-	SetSamplerParameter(RHICmdList, Shader, SamplerParameter, SamplerStateRHI, ElementIndex);
-PRAGMA_ENABLE_DEPRECATION_WARNINGS
-}
-
 template<typename TRHIShader, typename TRHICmdList>
 FORCEINLINE void SetTextureParameter(TRHICmdList& RHICmdList, TRHIShader* Shader, const FShaderResourceParameter& TextureParameter, const FShaderResourceParameter& SamplerParameter, FRHISamplerState* SamplerStateRHI, FRHITexture* TextureRHI)
 {
@@ -610,28 +540,6 @@ FORCEINLINE void SetTextureParameter(TRHICmdList& RHICmdList, TRHIShader* Shader
 /**
  * Sets the value of a shader texture parameter.  Template'd on shader type
  */
-template<typename TRHIShader, typename TRHICmdList>
-UE_DEPRECATED(5.2, "SetTextureParameter with an index can't be supported anymore. Your code should be changed to use shader parameter structs to utilize resource arrays.")
-FORCEINLINE void SetTextureParameter(
-	TRHICmdList& RHICmdList,
-	TRHIShader* Shader,
-	const FShaderResourceParameter& TextureParameter,
-	const FShaderResourceParameter& SamplerParameter,
-	const FTexture* Texture,
-	uint32 ElementIndex
-)
-{
-	if (TextureParameter.IsBound())
-	{
-		Texture->LastRenderTime = FApp::GetCurrentTime();
-	}
-
-PRAGMA_DISABLE_DEPRECATION_WARNINGS
-	SetTextureParameter(RHICmdList, Shader, TextureParameter, Texture->TextureRHI, ElementIndex);
-	SetSamplerParameter(RHICmdList, Shader, SamplerParameter, Texture->SamplerStateRHI, ElementIndex);
-PRAGMA_ENABLE_DEPRECATION_WARNINGS
-}
-
 template<typename TRHIShader, typename TRHICmdList>
 UE_DEPRECATED(5.3, "SetTextureParameter with FRHIBatchedShaderParameters should be used.")
 FORCEINLINE void SetTextureParameter(TRHICmdList& RHICmdList, TRHIShader* Shader, const FShaderResourceParameter& TextureParameter, const FShaderResourceParameter& SamplerParameter, const FTexture* Texture)
