Changing FComputeShaderUtils to be a namespace instead of a struct. Deduplicating the majority of code in the FComputeShaderUtils Dispatch methods.

diff --git a/Engine/Source/Runtime/RenderCore/Public/RenderGraphUtils.h b/Engine/Source/Runtime/RenderCore/Public/RenderGraphUtils.h
--- a/Engine/Source/Runtime/RenderCore/Public/RenderGraphUtils.h
+++ b/Engine/Source/Runtime/RenderCore/Public/RenderGraphUtils.h
@@ -338,41 +338,41 @@ inline FRDGTextureMSAA CreateTextureMSAA(
 
 /** All utils for compute shaders.
  */
-struct FComputeShaderUtils
+namespace FComputeShaderUtils
 {
 	/** Ideal size of group size 8x8 to occupy at least an entire wave on GCN, two warp on Nvidia. */
 	static constexpr int32 kGolden2DGroupSize = 8;
 
 	/** Compute the number of groups to dispatch. */
-	static FIntVector GetGroupCount(const int32 ThreadCount, const int32 GroupSize)
+	inline FIntVector GetGroupCount(const int32 ThreadCount, const int32 GroupSize)
 	{
 		return FIntVector(
 			FMath::DivideAndRoundUp(ThreadCount, GroupSize),
 			1,
 			1);
 	}
-	static FIntVector GetGroupCount(const FIntPoint& ThreadCount, const FIntPoint& GroupSize)
+	inline FIntVector GetGroupCount(const FIntPoint& ThreadCount, const FIntPoint& GroupSize)
 	{
 		return FIntVector(
 			FMath::DivideAndRoundUp(ThreadCount.X, GroupSize.X),
 			FMath::DivideAndRoundUp(ThreadCount.Y, GroupSize.Y),
 			1);
 	}
-	static FIntVector GetGroupCount(const FIntPoint& ThreadCount, const int32 GroupSize)
+	inline FIntVector GetGroupCount(const FIntPoint& ThreadCount, const int32 GroupSize)
 	{
 		return FIntVector(
 			FMath::DivideAndRoundUp(ThreadCount.X, GroupSize),
 			FMath::DivideAndRoundUp(ThreadCount.Y, GroupSize),
 			1);
 	}
-	static FIntVector GetGroupCount(const FIntVector& ThreadCount, const FIntVector& GroupSize)
+	inline FIntVector GetGroupCount(const FIntVector& ThreadCount, const FIntVector& GroupSize)
 	{
 		return FIntVector(
 			FMath::DivideAndRoundUp(ThreadCount.X, GroupSize.X),
 			FMath::DivideAndRoundUp(ThreadCount.Y, GroupSize.Y),
 			FMath::DivideAndRoundUp(ThreadCount.Z, GroupSize.Z));
 	}
-	static FIntVector GetGroupCount(const FIntVector& ThreadCount, const int32 GroupSize)
+	inline FIntVector GetGroupCount(const FIntVector& ThreadCount, const int32 GroupSize)
 	{
 		return FIntVector(
 			FMath::DivideAndRoundUp(ThreadCount.X, GroupSize),
@@ -380,8 +380,6 @@ struct FComputeShaderUtils
 			FMath::DivideAndRoundUp(ThreadCount.Z, GroupSize));
 	}
 
-
-
 	/**
 	 * Constant stride used when wrapping too large 1D dispatches using GetGroupCountWrapped, selected as 128 appears to be the lowest common denominator 
 	 * for mobile (GLES 3.1). For PC (with ~64k groups / dimension) this yields ~8M groups (500M threads @ group size 64) before even wrapping into Z.
@@ -395,7 +393,7 @@ struct FComputeShaderUtils
 	 *  uint LinearGroupId = GroupId.X + (GroupId.Z * WrappedGroupStride + GroupId.Y) * WrappedGroupStride;
 	 * Note that you must use an early out because LinearGroupId may be larger than the ideal due to wrapping.
 	 */
-	static FIntVector GetGroupCountWrapped(const int32 TargetGroupCount)
+	inline FIntVector GetGroupCountWrapped(const int32 TargetGroupCount)
 	{
 		check(GRHIMaxDispatchThreadGroupsPerDimension.X >= WrappedGroupStride && GRHIMaxDispatchThreadGroupsPerDimension.Y >= WrappedGroupStride);
 
@@ -423,15 +421,72 @@ struct FComputeShaderUtils
 	 *  uint LinearGroupId = GroupId.X + (GroupId.Z * WrappedGroupStride + GroupId.Y) * WrappedGroupStride;
 	 * Note that you must use an early out because LinearGroupId may be larger than the ideal due to wrapping.
 	 */
-	static FIntVector GetGroupCountWrapped(const int32 ThreadCount, const int32 GroupSize)
+	inline FIntVector GetGroupCountWrapped(const int32 ThreadCount, const int32 GroupSize)
 	{
 		return GetGroupCountWrapped(FMath::DivideAndRoundUp(ThreadCount, GroupSize));
 	}
 
+	inline void ValidateGroupCount(const FIntVector& GroupCount)
+	{
+		ensure(GroupCount.X <= GRHIMaxDispatchThreadGroupsPerDimension.X);
+		ensure(GroupCount.Y <= GRHIMaxDispatchThreadGroupsPerDimension.Y);
+		ensure(GroupCount.Z <= GRHIMaxDispatchThreadGroupsPerDimension.Z);
+	}
+
+	inline void ValidateIndirectArgsBuffer(uint32 IndirectArgsBufferSize, uint32 IndirectArgOffset)
+	{
+		constexpr uint32 IndirectArgsSize = sizeof(FRHIDispatchIndirectParametersNoPadding);
+		checkf((IndirectArgOffset % 4) == 0, TEXT("IndirectArgOffset for compute shader indirect dispatch needs to be a multiple of 4."));
+		checkf(
+			(IndirectArgOffset + IndirectArgsSize) <= IndirectArgsBufferSize,
+			TEXT("Indirect parameters buffer for compute shader indirect dispatch at byte offset %d doesn't have enough room for one element."),
+			IndirectArgOffset);
+#if PLATFORM_DISPATCH_INDIRECT_ARGUMENT_BOUNDARY_SIZE != 0
+		checkf(IndirectArgOffset / PLATFORM_DISPATCH_INDIRECT_ARGUMENT_BOUNDARY_SIZE == (IndirectArgOffset + IndirectArgsSize - 1) / PLATFORM_DISPATCH_INDIRECT_ARGUMENT_BOUNDARY_SIZE, TEXT("Compute indirect dispatch arguments cannot cross %d byte boundary."), PLATFORM_DISPATCH_INDIRECT_ARGUMENT_BOUNDARY_SIZE);
+#endif // #if PLATFORM_DISPATCH_INDIRECT_ARGUMENT_BOUNDARY_SIZE != 0
+	}
+
+	inline void ValidateIndirectArgsBuffer(const FRDGBufferRef IndirectArgsBuffer, uint32 IndirectArgOffset)
+	{
+		checkf(EnumHasAnyFlags(IndirectArgsBuffer->Desc.Usage, EBufferUsageFlags::VertexBuffer), TEXT("The buffer %s needs to be a vertex buffer to be used as an indirect dispatch parameters"), IndirectArgsBuffer->Name);
+		checkf(EnumHasAnyFlags(IndirectArgsBuffer->Desc.Usage, EBufferUsageFlags::DrawIndirect), TEXT("The buffer %s for indirect dispatch parameters was not flagged with BUF_DrawIndirect"), IndirectArgsBuffer->Name);
+		ValidateIndirectArgsBuffer(IndirectArgsBuffer->GetSize(), IndirectArgOffset);
+	}
+
+	namespace Private
+	{
+		template<typename TShaderClass>
+		inline void PrepareDispatch(
+			FRHIComputeCommandList& RHICmdList,
+			const TShaderRef<TShaderClass>& ComputeShader,
+			const FShaderParametersMetadata* ParametersMetadata,
+			const typename TShaderClass::FParameters& Parameters)
+		{
+			FRHIComputeShader* ShaderRHI = ComputeShader.GetComputeShader();
+			SetComputePipelineState(RHICmdList, ShaderRHI);
+			SetShaderParameters(RHICmdList, ComputeShader, ShaderRHI, ParametersMetadata, Parameters);
+		}
+
+		template<typename TShaderClass>
+		inline void PrepareDispatch(
+			FRHIComputeCommandList& RHICmdList,
+			const TShaderRef<TShaderClass>& ComputeShader,
+			const typename TShaderClass::FParameters& Parameters)
+		{
+			const FShaderParametersMetadata* ParametersMetadata = TShaderClass::FParameters::FTypeInfo::GetStructMetadata();
+			PrepareDispatch(RHICmdList, ComputeShader, ParametersMetadata, Parameters);
+		}
+
+		template<typename TShaderClass>
+		inline void AfterDispatch(FRHIComputeCommandList& RHICmdList, const TShaderRef<TShaderClass>& ComputeShader)
+		{
+			UnsetShaderUAVs(RHICmdList, ComputeShader, ComputeShader.GetComputeShader());
+		}
+	}
 
 	/** Dispatch a compute shader to rhi command list with its parameters. */
 	template<typename TShaderClass>
-	static void Dispatch(
+	inline void Dispatch(
 		FRHIComputeCommandList& RHICmdList, 
 		const TShaderRef<TShaderClass>& ComputeShader, 
 		const FShaderParametersMetadata* ParametersMetadata,
@@ -439,15 +494,14 @@ struct FComputeShaderUtils
 		FIntVector GroupCount)
 	{
 		ValidateGroupCount(GroupCount);
-		FRHIComputeShader* ShaderRHI = ComputeShader.GetComputeShader();
-		SetComputePipelineState(RHICmdList, ShaderRHI);
-		SetShaderParameters(RHICmdList, ComputeShader, ShaderRHI, ParametersMetadata, Parameters);
+
+		Private::PrepareDispatch(RHICmdList, ComputeShader, ParametersMetadata, Parameters);
 		RHICmdList.DispatchComputeShader(GroupCount.X, GroupCount.Y, GroupCount.Z);
-		UnsetShaderUAVs(RHICmdList, ComputeShader, ShaderRHI);
+		Private::AfterDispatch(RHICmdList, ComputeShader);
 	}
 
 	template<typename TShaderClass>
-	static void Dispatch(
+	inline void Dispatch(
 		FRHIComputeCommandList& RHICmdList,
 		const TShaderRef<TShaderClass>& ComputeShader,
 		const typename TShaderClass::FParameters& Parameters,
@@ -459,7 +513,7 @@ struct FComputeShaderUtils
 	
 	/** Indirect dispatch a compute shader to rhi command list with its parameters. */
 	template<typename TShaderClass>
-	static void DispatchIndirect(
+	inline void DispatchIndirect(
 		FRHIComputeCommandList& RHICmdList,
 		const TShaderRef<TShaderClass>& ComputeShader,
 		const typename TShaderClass::FParameters& Parameters,
@@ -467,16 +521,15 @@ struct FComputeShaderUtils
 		uint32 IndirectArgOffset)
 	{
 		ValidateIndirectArgsBuffer(IndirectArgsBuffer->GetSize(), IndirectArgOffset);
-		FRHIComputeShader* ShaderRHI = ComputeShader.GetComputeShader();
-		SetComputePipelineState(RHICmdList, ShaderRHI);
-		SetShaderParameters(RHICmdList, ComputeShader, ShaderRHI, Parameters);
+
+		Private::PrepareDispatch(RHICmdList, ComputeShader, Parameters);
 		RHICmdList.DispatchIndirectComputeShader(IndirectArgsBuffer, IndirectArgOffset);
-		UnsetShaderUAVs(RHICmdList, ComputeShader, ShaderRHI);
+		Private::AfterDispatch(RHICmdList, ComputeShader);
 	}
 
 	/** Dispatch a compute shader to rhi command list with its parameters and indirect args. */
 	template<typename TShaderClass>
-	static FORCEINLINE_DEBUGGABLE void DispatchIndirect(
+	inline void DispatchIndirect(
 		FRHIComputeCommandList& RHICmdList,
 		const TShaderRef<TShaderClass>& ComputeShader,
 		const typename TShaderClass::FParameters& Parameters,
@@ -484,16 +537,15 @@ struct FComputeShaderUtils
 		uint32 IndirectArgOffset)
 	{
 		ValidateIndirectArgsBuffer(IndirectArgsBuffer, IndirectArgOffset);
-		FRHIComputeShader* ShaderRHI = ComputeShader.GetComputeShader();
-		SetComputePipelineState(RHICmdList, ShaderRHI);
-		SetShaderParameters(RHICmdList, ComputeShader, ShaderRHI, Parameters);
+
+		Private::PrepareDispatch(RHICmdList, ComputeShader, Parameters);
 		RHICmdList.DispatchIndirectComputeShader(IndirectArgsBuffer->GetIndirectRHICallBuffer(), IndirectArgOffset);
-		UnsetShaderUAVs(RHICmdList, ComputeShader, ShaderRHI);
+		Private::AfterDispatch(RHICmdList, ComputeShader);
 	}
 
 	/** Dispatch a compute shader to render graph builder with its parameters. */
 	template<typename TShaderClass>
-	static FRDGPassRef AddPass(
+	inline FRDGPassRef AddPass(
 		FRDGBuilder& GraphBuilder,
 		FRDGEventName&& PassName,
 		ERDGPassFlags PassFlags,
@@ -524,7 +576,7 @@ struct FComputeShaderUtils
 	 *  This allows adding a dispatch with unknown GroupCount but the value must be ready before the pass is executed.
 	 */
 	template<typename TShaderClass>
-	static FRDGPassRef AddPass(
+	inline FRDGPassRef AddPass(
 		FRDGBuilder& GraphBuilder,
 		FRDGEventName&& PassName,
 		ERDGPassFlags PassFlags,
@@ -556,7 +608,7 @@ struct FComputeShaderUtils
 	}
 
 	template<typename TShaderClass>
-	static FRDGPassRef AddPass(
+	inline FRDGPassRef AddPass(
 		FRDGBuilder& GraphBuilder,
 		FRDGEventName&& PassName,
 		ERDGPassFlags PassFlags,
@@ -569,7 +621,7 @@ struct FComputeShaderUtils
 	}
 
 	template <typename TShaderClass>
-	static FORCEINLINE FRDGPassRef AddPass(
+	inline FRDGPassRef AddPass(
 		FRDGBuilder& GraphBuilder,
 		FRDGEventName&& PassName,
 		const TShaderRef<TShaderClass>& ComputeShader,
@@ -581,7 +633,7 @@ struct FComputeShaderUtils
 	}
 
 	template <typename TShaderClass>
-	static FORCEINLINE FRDGPassRef AddPass(
+	inline FRDGPassRef AddPass(
 		FRDGBuilder& GraphBuilder,
 		FRDGEventName&& PassName,
 		const TShaderRef<TShaderClass>& ComputeShader,
@@ -594,7 +646,7 @@ struct FComputeShaderUtils
 
 	/** Dispatch a compute shader to render graph builder with its parameters. */
 	template<typename TShaderClass>
-	static FRDGPassRef AddPass(
+	inline FRDGPassRef AddPass(
 		FRDGBuilder& GraphBuilder,
 		FRDGEventName&& PassName,
 		ERDGPassFlags PassFlags,
@@ -629,7 +681,7 @@ struct FComputeShaderUtils
 	}
 
 	template<typename TShaderClass>
-	static FORCEINLINE FRDGPassRef AddPass(
+	inline FRDGPassRef AddPass(
 		FRDGBuilder& GraphBuilder,
 		FRDGEventName&& PassName,
 		const TShaderRef<TShaderClass>& ComputeShader,
@@ -641,42 +693,15 @@ struct FComputeShaderUtils
 		return AddPass(GraphBuilder, Forward<FRDGEventName>(PassName), ERDGPassFlags::Compute, ComputeShader, Parameters, IndirectArgsBuffer, IndirectArgsOffset, MoveTemp(DispatchLateParamCallback));
 	}
 
-	static RENDERCORE_API void ClearUAV(FRDGBuilder& GraphBuilder, FGlobalShaderMap* ShaderMap, FRDGBufferUAVRef UAV, uint32 ClearValue);
-	static RENDERCORE_API void ClearUAV(FRDGBuilder& GraphBuilder, FGlobalShaderMap* ShaderMap, FRDGBufferUAVRef UAV, FVector4f ClearValue);
-
-	static inline void ValidateGroupCount(const FIntVector& GroupCount)
-	{
-		ensure(GroupCount.X <= GRHIMaxDispatchThreadGroupsPerDimension.X);
-		ensure(GroupCount.Y <= GRHIMaxDispatchThreadGroupsPerDimension.Y);
-		ensure(GroupCount.Z <= GRHIMaxDispatchThreadGroupsPerDimension.Z);
-	}
-
-	static inline void ValidateIndirectArgsBuffer(uint32 IndirectArgsBufferSize, uint32 IndirectArgOffset)
-	{
-		constexpr uint32 IndirectArgsSize = sizeof(FRHIDispatchIndirectParametersNoPadding);
-		checkf((IndirectArgOffset % 4) == 0, TEXT("IndirectArgOffset for compute shader indirect dispatch needs to be a multiple of 4."));
-		checkf(
-			(IndirectArgOffset + IndirectArgsSize) <= IndirectArgsBufferSize,
-			TEXT("Indirect parameters buffer for compute shader indirect dispatch at byte offset %d doesn't have enough room for one element."),
-			IndirectArgOffset);
-#if PLATFORM_DISPATCH_INDIRECT_ARGUMENT_BOUNDARY_SIZE != 0
-			checkf(IndirectArgOffset / PLATFORM_DISPATCH_INDIRECT_ARGUMENT_BOUNDARY_SIZE == (IndirectArgOffset + IndirectArgsSize - 1) / PLATFORM_DISPATCH_INDIRECT_ARGUMENT_BOUNDARY_SIZE, TEXT("Compute indirect dispatch arguments cannot cross %d byte boundary."), PLATFORM_DISPATCH_INDIRECT_ARGUMENT_BOUNDARY_SIZE);
-#endif // #if PLATFORM_DISPATCH_INDIRECT_ARGUMENT_BOUNDARY_SIZE != 0
-	}
-
-	static inline void ValidateIndirectArgsBuffer(const FRDGBufferRef IndirectArgsBuffer, uint32 IndirectArgOffset)
-	{
-		checkf(EnumHasAnyFlags(IndirectArgsBuffer->Desc.Usage, EBufferUsageFlags::VertexBuffer), TEXT("The buffer %s needs to be a vertex buffer to be used as an indirect dispatch parameters"), IndirectArgsBuffer->Name);
-		checkf(EnumHasAnyFlags(IndirectArgsBuffer->Desc.Usage, EBufferUsageFlags::DrawIndirect), TEXT("The buffer %s for indirect dispatch parameters was not flagged with BUF_DrawIndirect"), IndirectArgsBuffer->Name);
-		ValidateIndirectArgsBuffer(IndirectArgsBuffer->GetSize(), IndirectArgOffset);
-	}
+	RENDERCORE_API void ClearUAV(FRDGBuilder& GraphBuilder, FGlobalShaderMap* ShaderMap, FRDGBufferUAVRef UAV, uint32 ClearValue);
+	RENDERCORE_API void ClearUAV(FRDGBuilder& GraphBuilder, FGlobalShaderMap* ShaderMap, FRDGBufferUAVRef UAV, FVector4f ClearValue);
 
 	/**
 	 * Create and set up an 1D indirect dispatch argument from some GPU-side integer in a buffer (InputCountBuffer).
 	 * 	Sets up a group count as (InputCountBuffer[InputCountOffset] * Multiplier + Divisor - 1U) / Divisor;
 	 *  Commonly use Divisor <=> number of threads per group.
 	 */
-	static RENDERCORE_API FRDGBufferRef AddIndirectArgsSetupCsPass1D(FRDGBuilder& GraphBuilder, ERHIFeatureLevel::Type FeatureLevel, FRDGBufferRef& InputCountBuffer, const TCHAR* OutputBufferName, uint32 Divisor, uint32 InputCountOffset = 0U, uint32 Multiplier = 1U);
+	RENDERCORE_API FRDGBufferRef AddIndirectArgsSetupCsPass1D(FRDGBuilder& GraphBuilder, ERHIFeatureLevel::Type FeatureLevel, FRDGBufferRef& InputCountBuffer, const TCHAR* OutputBufferName, uint32 Divisor, uint32 InputCountOffset = 0U, uint32 Multiplier = 1U);
 };
 
 /** Adds a render graph pass to copy a region from one texture to another. Uses RHICopyTexture under the hood.
