Modified RenderCaptureInterface::FScopedCapture to take an FRHICommandList, which no-ops if it's not the immediate command list.

diff --git a/Engine/Source/Runtime/RenderCore/Private/RenderCaptureInterface.cpp b/Engine/Source/Runtime/RenderCore/Private/RenderCaptureInterface.cpp
--- a/Engine/Source/Runtime/RenderCore/Private/RenderCaptureInterface.cpp
+++ b/Engine/Source/Runtime/RenderCore/Private/RenderCaptureInterface.cpp
@@ -29,17 +29,17 @@ namespace RenderCaptureInterface
 		}
 	}
 
-	FScopedCapture::FScopedCapture(bool bEnable, FRHICommandListImmediate* InRHICommandList, TCHAR const* InEventName, TCHAR const* InFileName)
-		: bCapture(bEnable && IRenderCaptureProvider::IsAvailable())
+	FScopedCapture::FScopedCapture(bool bEnable, FRHICommandList* InRHICommandList, TCHAR const* InEventName, TCHAR const* InFileName)
+		: bCapture(bEnable && IRenderCaptureProvider::IsAvailable() && InRHICommandList->IsImmediate())
 		, bEvent(InEventName != nullptr)
 		, RHICommandList(InRHICommandList)
 		, GraphBuilder(nullptr)
 	{
-		check(!GIsThreadedRendering || IsInRenderingThread());
-
 		if (bCapture)
 		{
-			IRenderCaptureProvider::Get().BeginCapture(RHICommandList, IRenderCaptureProvider::ECaptureFlags_Launch, FString(InFileName));
+			check(!GIsThreadedRendering || IsInRenderingThread());
+
+			IRenderCaptureProvider::Get().BeginCapture(&FRHICommandListImmediate::Get(*RHICommandList), IRenderCaptureProvider::ECaptureFlags_Launch, FString(InFileName));
 		
 			if (bEvent)
 			{
@@ -103,7 +103,7 @@ namespace RenderCaptureInterface
 					RHICommandList->PopEvent();
 				}
 
-				IRenderCaptureProvider::Get().EndCapture(RHICommandList);
+				IRenderCaptureProvider::Get().EndCapture(&FRHICommandListImmediate::Get(*RHICommandList));
 			}
 			else
 			{
diff --git a/Engine/Source/Runtime/RenderCore/Public/RenderCaptureInterface.h b/Engine/Source/Runtime/RenderCore/Public/RenderCaptureInterface.h
--- a/Engine/Source/Runtime/RenderCore/Public/RenderCaptureInterface.h
+++ b/Engine/Source/Runtime/RenderCore/Public/RenderCaptureInterface.h
@@ -6,7 +6,7 @@
 #include "HAL/Platform.h"
 
 class FRDGBuilder;
-class FRHICommandListImmediate;
+class FRHICommandList;
 
 /** Easy to use interface for IRenderCaptureProvider. */
 namespace RenderCaptureInterface
@@ -21,7 +21,7 @@ namespace RenderCaptureInterface
 		/** Use this constructor if not on rendering thread. Use bEnable to allow control over the capture frequency. */
 		RENDERCORE_API FScopedCapture(bool bEnable, TCHAR const* InEventName = nullptr, TCHAR const* InFileName = nullptr);
 		/** Use this constructor if on rendering thread. Use bEnable to allow control over the capture frequency. */
-		RENDERCORE_API FScopedCapture(bool bEnable, FRHICommandListImmediate* InRHICommandList, TCHAR const* InEventName = nullptr, TCHAR const* InFileName = nullptr);
+		RENDERCORE_API FScopedCapture(bool bEnable, FRHICommandList* InRHICommandList, TCHAR const* InEventName = nullptr, TCHAR const* InFileName = nullptr);
 		/** Use this constructor if using RenderGraph to schedule work. Use bEnable to allow control over the capture frequency. */
 		RENDERCORE_API FScopedCapture(bool bEnable, FRDGBuilder& InGraphBuilder, TCHAR const* InEventName = nullptr, TCHAR const* InFileName = nullptr);
 
@@ -30,7 +30,7 @@ namespace RenderCaptureInterface
 	private:
 		bool bCapture;
 		bool bEvent;
-		FRHICommandListImmediate* RHICommandList;
+		FRHICommandList* RHICommandList;
 		FRDGBuilder* GraphBuilder;
 	};
 }
