Moving FClearVertexBuffer::InitRHI to its cpp file to remove header dependencies.

diff --git a/Engine/Source/Runtime/RenderCore/Private/ClearQuad.cpp b/Engine/Source/Runtime/RenderCore/Private/ClearQuad.cpp
--- a/Engine/Source/Runtime/RenderCore/Private/ClearQuad.cpp
+++ b/Engine/Source/Runtime/RenderCore/Private/ClearQuad.cpp
@@ -9,6 +9,21 @@
 #include "RendererInterface.h"
 #include "Logging/LogMacros.h"
 
+void FClearVertexBuffer::InitRHI(FRHICommandListBase& RHICmdList)
+{
+	// create a static vertex buffer
+	FRHIResourceCreateInfo CreateInfo(TEXT("FClearVertexBuffer"));
+	VertexBufferRHI = RHICmdList.CreateVertexBuffer(sizeof(FVector4f) * 4, BUF_Static, CreateInfo);
+	void* VoidPtr = RHICmdList.LockBuffer(VertexBufferRHI, 0, sizeof(FVector4f) * 4, RLM_WriteOnly);
+	// Generate the vertices used
+	FVector4f* Vertices = reinterpret_cast<FVector4f*>(VoidPtr);
+	Vertices[0] = FVector4f(-1.0f, 1.0f, 0.0f, 1.0f);
+	Vertices[1] = FVector4f(1.0f, 1.0f, 0.0f, 1.0f);
+	Vertices[2] = FVector4f(-1.0f, -1.0f, 0.0f, 1.0f);
+	Vertices[3] = FVector4f(1.0f, -1.0f, 0.0f, 1.0f);
+	RHICmdList.UnlockBuffer(VertexBufferRHI);
+}
+
 TGlobalResource<FClearVertexBuffer> GClearVertexBuffer;
 
 DEFINE_LOG_CATEGORY_STATIC(LogClearQuad, Log, Log)
diff --git a/Engine/Source/Runtime/RenderCore/Public/ClearQuad.h b/Engine/Source/Runtime/RenderCore/Public/ClearQuad.h
--- a/Engine/Source/Runtime/RenderCore/Public/ClearQuad.h
+++ b/Engine/Source/Runtime/RenderCore/Public/ClearQuad.h
@@ -2,17 +2,12 @@
 
 #pragma once
 
-#include "CoreMinimal.h"
 #include "Math/Color.h"
 #include "Math/IntPoint.h"
 #include "Math/IntRect.h"
 #include "Math/UnrealMathSSE.h"
 #include "Math/Vector4.h"
-#include "RHI.h"
-#include "RHICommandList.h"
-#include "RHIDefinitions.h"
 #include "RenderResource.h"
-#include "RendererInterface.h"
 #include "Templates/Function.h"
 
 class FGraphicsPipelineStateInitializer;
@@ -27,20 +22,7 @@ public:
 	/**
 	* Initialize the RHI for this rendering resource
 	*/
-	virtual void InitRHI(FRHICommandListBase& RHICmdList) override
-	{
-		// create a static vertex buffer
-		FRHIResourceCreateInfo CreateInfo(TEXT("FClearVertexBuffer"));
-		VertexBufferRHI = RHICmdList.CreateVertexBuffer(sizeof(FVector4f) * 4, BUF_Static, CreateInfo);
-		void* VoidPtr = RHICmdList.LockBuffer(VertexBufferRHI, 0, sizeof(FVector4f) * 4, RLM_WriteOnly);
-		// Generate the vertices used
-		FVector4f* Vertices = reinterpret_cast<FVector4f*>(VoidPtr);
-		Vertices[0] = FVector4f(-1.0f, 1.0f, 0.0f, 1.0f);
-		Vertices[1] = FVector4f(1.0f, 1.0f, 0.0f, 1.0f);
-		Vertices[2] = FVector4f(-1.0f, -1.0f, 0.0f, 1.0f);
-		Vertices[3] = FVector4f(1.0f, -1.0f, 0.0f, 1.0f);
-		RHICmdList.UnlockBuffer(VertexBufferRHI);
-	}
+	virtual void InitRHI(FRHICommandListBase& RHICmdList) override;
 };
 extern RENDERCORE_API TGlobalResource<FClearVertexBuffer> GClearVertexBuffer;
 
@@ -76,3 +58,12 @@ inline void DrawClearQuad(FRHICommandList& RHICmdList, const FLinearColor& Color
 {
 	DrawClearQuadMRT(RHICmdList, true, 1, &Color, false, 0, false, 0, ClearQuadCallbacks);
 }
+
+
+#if UE_ENABLE_INCLUDE_ORDER_DEPRECATED_IN_5_4
+#include "CoreMinimal.h"
+#include "RHI.h"
+#include "RHICommandList.h"
+#include "RHIDefinitions.h"
+#include "RendererInterface.h"
+#endif
\ No newline at end of file
