Fixed incorrect UB declaration for a uniform buffer view

diff --git a/Engine/Shaders/Public/ShaderVersion.ush b/Engine/Shaders/Public/ShaderVersion.ush
--- a/Engine/Shaders/Public/ShaderVersion.ush
+++ b/Engine/Shaders/Public/ShaderVersion.ush
@@ -2,5 +2,5 @@
 // Each time the console command is executed it generates a new GUID. As this file is included
 // in Platform.ush (which should be included in any shader) it allows to invalidate the shader DDC.
 // 
-// GUID = 540d2e3c-937e-40b4-ab75-581325ebb06f
+// GUID = 631DFD83-05E2-4AC1-A7AC-47CBF4E629EB
 
diff --git a/Engine/Source/Runtime/RenderCore/Private/ShaderParameters.cpp b/Engine/Source/Runtime/RenderCore/Private/ShaderParameters.cpp
--- a/Engine/Source/Runtime/RenderCore/Private/ShaderParameters.cpp
+++ b/Engine/Source/Runtime/RenderCore/Private/ShaderParameters.cpp
@@ -180,7 +180,7 @@ static void CreateHLSLUniformBufferStructMembersDeclaration(
 		{
 			FString ParameterName = FString::Printf(TEXT("%s%s"),*NamePrefix,Member.GetName());
 			Decl.ConstantBufferMembers += FString::Printf(TEXT("UB_MEMBER_UNIFORM_BLOCK(%s);\r\n"), *ParameterName);
-			Decl.StructMembers += FString::Printf(TEXT("\t%s %s;\r\n"), Member.GetShaderType(), Member.GetName());
+			Decl.StructMembers += FString::Printf(TEXT("UB_MEMBER_UNIFORM_BLOCK(%s);\r\n"), Member.GetName());
 			Decl.Initializer += FString::Printf(TEXT("%s,"),*ParameterName);
 		}
 		else if (Member.GetBaseType() == UBMT_INCLUDED_STRUCT)
