[RenderCore]
* Fixed compile error happening in new compilers

diff --git a/Engine/Source/Runtime/RenderCore/Private/ShaderCore.cpp b/Engine/Source/Runtime/RenderCore/Private/ShaderCore.cpp
--- a/Engine/Source/Runtime/RenderCore/Private/ShaderCore.cpp
+++ b/Engine/Source/Runtime/RenderCore/Private/ShaderCore.cpp
@@ -138,6 +138,8 @@ FShaderCompilerDefinitions::FShaderCompilerDefinitions(bool bIncludeInitialDefin
 	}
 }
 
+FShaderCompilerDefinitions::FShaderCompilerDefinitions(const FShaderCompilerDefinitions&) = default;
+
 void FShaderCompilerDefinitions::InitializeInitialDefines(const FShaderCompilerDefinitions& InDefines)
 {
 	check(GInitialDefines == nullptr);
diff --git a/Engine/Source/Runtime/RenderCore/Public/ShaderCore.h b/Engine/Source/Runtime/RenderCore/Public/ShaderCore.h
--- a/Engine/Source/Runtime/RenderCore/Public/ShaderCore.h
+++ b/Engine/Source/Runtime/RenderCore/Public/ShaderCore.h
@@ -331,6 +331,8 @@ class FShaderCompilerDefinitions
 {
 public:
 	RENDERCORE_API FShaderCompilerDefinitions(bool bIncludeInitialDefines = false);
+	
+	RENDERCORE_API FShaderCompilerDefinitions(const FShaderCompilerDefinitions&);
 
 	/** Value types supported:   bool, int32, uint32, float, const TCHAR*, FString& */
 	template <typename ValueType>
