ShaderLibrary: Logging: Downgrade LogShaderLibrary stat messages to Log to make it easier to scan the display output of cooking.
#rnx
#b Jason.Nadro

diff --git a/Engine/Source/Runtime/RenderCore/Private/ShaderCodeLibrary.cpp b/Engine/Source/Runtime/RenderCore/Private/ShaderCodeLibrary.cpp
--- a/Engine/Source/Runtime/RenderCore/Private/ShaderCodeLibrary.cpp
+++ b/Engine/Source/Runtime/RenderCore/Private/ShaderCodeLibrary.cpp
@@ -2025,19 +2025,19 @@ struct FEditorShaderCodeArchive
 	{
 		bool bUseExtendedDebugInfo = UE::ShaderLibrary::Private::GProduceExtendedStats != 0;
 
-		UE_LOG(LogShaderLibrary, Display, TEXT(""));
-		UE_LOG(LogShaderLibrary, Display, TEXT("Shader Library '%s' (%s) Stats:"), *LibraryName, *FormatName.ToString());
-		UE_LOG(LogShaderLibrary, Display, TEXT("================="));
+		UE_LOG(LogShaderLibrary, Log, TEXT(""));
+		UE_LOG(LogShaderLibrary, Log, TEXT("Shader Library '%s' (%s) Stats:"), *LibraryName, *FormatName.ToString());
+		UE_LOG(LogShaderLibrary, Log, TEXT("================="));
 
 		FSerializedShaderArchive::FDebugStats Stats;
 		FSerializedShaderArchive::FExtendedDebugStats ExtendedStats;
 		SerializedShaders.CollectStatsAndDebugInfo(Stats, bUseExtendedDebugInfo ? &ExtendedStats : nullptr);
 
-		UE_LOG(LogShaderLibrary, Display, TEXT("Assets: %d, Unique Shadermaps: %d (%.2f%%)"), 
+		UE_LOG(LogShaderLibrary, Log, TEXT("Assets: %d, Unique Shadermaps: %d (%.2f%%)"),
 			Stats.NumAssets, Stats.NumShaderMaps, (Stats.NumAssets > 0) ? 100.0 * static_cast<double>(Stats.NumShaderMaps) / static_cast<double>(Stats.NumAssets) : 0.0);
-		UE_LOG(LogShaderLibrary, Display, TEXT("Total Shaders: %d, Unique Shaders: %d (%.2f%%)"), 
+		UE_LOG(LogShaderLibrary, Log, TEXT("Total Shaders: %d, Unique Shaders: %d (%.2f%%)"),
 			Stats.NumShaders, Stats.NumUniqueShaders, (Stats.NumShaders > 0) ? 100.0 * static_cast<double>(Stats.NumUniqueShaders) / static_cast<double>(Stats.NumShaders) : 0.0);
-		UE_LOG(LogShaderLibrary, Display, TEXT("Total Shader Size: %.2fmb, Unique Shaders Size: %.2fmb (%.2f%%)"), 
+		UE_LOG(LogShaderLibrary, Log, TEXT("Total Shader Size: %.2fmb, Unique Shaders Size: %.2fmb (%.2f%%)"),
 			FUnitConversion::Convert(static_cast<double>(Stats.ShadersSize), EUnit::Bytes, EUnit::Megabytes), 
 			FUnitConversion::Convert(static_cast<double>(Stats.ShadersUniqueSize), EUnit::Bytes, EUnit::Megabytes),
 			(Stats.ShadersSize > 0) ? 100.0 * static_cast<double>(Stats.ShadersUniqueSize) / static_cast<double>(Stats.ShadersSize) : 0.0
@@ -2045,23 +2045,23 @@ struct FEditorShaderCodeArchive
 
 		if (bUseExtendedDebugInfo)
 		{
-			UE_LOG(LogShaderLibrary, Display, TEXT("=== Extended info:"));
-			UE_LOG(LogShaderLibrary, Display, TEXT("Minimum number of shaders in shadermap: %d"), ExtendedStats.MinNumberOfShadersPerSM);
-			UE_LOG(LogShaderLibrary, Display, TEXT("Median number of shaders in shadermap: %d"), ExtendedStats.MedianNumberOfShadersPerSM);
-			UE_LOG(LogShaderLibrary, Display, TEXT("Maximum number of shaders in shadermap: %d"), ExtendedStats.MaxNumberofShadersPerSM);
+			UE_LOG(LogShaderLibrary, Log, TEXT("=== Extended info:"));
+			UE_LOG(LogShaderLibrary, Log, TEXT("Minimum number of shaders in shadermap: %d"), ExtendedStats.MinNumberOfShadersPerSM);
+			UE_LOG(LogShaderLibrary, Log, TEXT("Median number of shaders in shadermap: %d"), ExtendedStats.MedianNumberOfShadersPerSM);
+			UE_LOG(LogShaderLibrary, Log, TEXT("Maximum number of shaders in shadermap: %d"), ExtendedStats.MaxNumberofShadersPerSM);
 			if (ExtendedStats.TopShaderUsages.Num() > 0)
 			{
 				FString UsageString;
-				UE_LOG(LogShaderLibrary, Display, TEXT("Number of shadermaps referencing top %d most shared shaders:"), ExtendedStats.TopShaderUsages.Num());
+				UE_LOG(LogShaderLibrary, Log, TEXT("Number of shadermaps referencing top %d most shared shaders:"), ExtendedStats.TopShaderUsages.Num());
 				for (int IdxUsage = 0; IdxUsage < ExtendedStats.TopShaderUsages.Num() - 1; ++IdxUsage)
 				{
 					UsageString += FString::Printf(TEXT("%d, "), ExtendedStats.TopShaderUsages[IdxUsage]);
 				}
-				UE_LOG(LogShaderLibrary, Display, TEXT("    %s%d"), *UsageString, ExtendedStats.TopShaderUsages[ExtendedStats.TopShaderUsages.Num() - 1]);
+				UE_LOG(LogShaderLibrary, Log, TEXT("    %s%d"), *UsageString, ExtendedStats.TopShaderUsages[ExtendedStats.TopShaderUsages.Num() - 1]);
 			}
 			else
 			{
-				UE_LOG(LogShaderLibrary, Display, TEXT("No shader usage info is provided"));
+				UE_LOG(LogShaderLibrary, Log, TEXT("No shader usage info is provided"));
 			}
 
 			FString DebugLibFolder = GetShaderDebugFolder(FPaths::ProjectSavedDir() / TEXT("Shaders") / FormatName.ToString(), LibraryName, FormatName);
@@ -2072,7 +2072,7 @@ struct FEditorShaderCodeArchive
 				TUniquePtr<FArchive> DumpAr(IFileManager::Get().CreateFileWriter(*DumpFile));
 				FTCHARToUTF8 Converter(*ExtendedStats.TextualRepresentation);
 				DumpAr->Serialize(const_cast<UTF8CHAR*>(reinterpret_cast<const UTF8CHAR*>(Converter.Get())), Converter.Length());
-				UE_LOG(LogShaderLibrary, Display, TEXT("Textual dump saved to '%s'"), *DumpFile);
+				UE_LOG(LogShaderLibrary, Log, TEXT("Textual dump saved to '%s'"), *DumpFile);
 			}
 #if 0 // creating a graphviz graph - maybe one day we'll return to this
 			FString DebugGraphFolder = GetShaderDebugFolder(FPaths::ProjectSavedDir() / TEXT("Shaders") / FormatName.ToString(), LibraryName, FormatName);
@@ -2100,7 +2100,7 @@ struct FEditorShaderCodeArchive
 #endif//
 		}
 
-		UE_LOG(LogShaderLibrary, Display, TEXT("================="));
+		UE_LOG(LogShaderLibrary, Log, TEXT("================="));
 	}
 
 private:
