Mutate DDC string since platform symbol serialization data changed
Protect WriteSymbolData if the Filename is empty

#tests cooking CitySample

diff --git a/Engine/Source/Developer/ShaderCompilerCommon/Private/ShaderSymbolExport.cpp b/Engine/Source/Developer/ShaderCompilerCommon/Private/ShaderSymbolExport.cpp
--- a/Engine/Source/Developer/ShaderCompilerCommon/Private/ShaderSymbolExport.cpp
+++ b/Engine/Source/Developer/ShaderCompilerCommon/Private/ShaderSymbolExport.cpp
@@ -134,6 +134,12 @@ void FShaderSymbolExport::Initialize()
 
 void FShaderSymbolExport::WriteSymbolData(const FString& Filename, const FString& DebugData, TConstArrayView<uint8> Contents)
 {
+	// No writing is possible if the Filename is empty
+	if (Filename.IsEmpty())
+	{
+		return;
+	}
+
 	// Skip this symbol data if we've already exported it before.
 	bool bAlreadyInSet = false;
 	ExportedShaders.Add(Filename, &bAlreadyInSet);
diff --git a/Engine/Source/Runtime/RenderCore/Private/Shader.cpp b/Engine/Source/Runtime/RenderCore/Private/Shader.cpp
--- a/Engine/Source/Runtime/RenderCore/Private/Shader.cpp
+++ b/Engine/Source/Runtime/RenderCore/Private/Shader.cpp
@@ -1597,9 +1597,12 @@ void ShaderMapAppendKeyString(EShaderPlatform Platform, FString& KeyString)
 	{
 		// Extra data (names, etc)
 		KeyString += ShouldEnableExtraShaderData(ShaderFormatName) ? TEXT("_ExtraData") : TEXT("");
-		// Symbols and/or SymbolsInfo
-		KeyString += ShouldGenerateShaderSymbols(ShaderFormatName) ? TEXT("_Symbols") : TEXT("");
-		KeyString += ShouldGenerateShaderSymbolsInfo(ShaderFormatName) ? TEXT("_SymbolsInfo") : TEXT("");
+		// Symbols and/or SymbolsInfo and version if symbols serialization changes
+		const bool bSymbols = ShouldGenerateShaderSymbols(ShaderFormatName);
+		const bool bSymbolsinfo = ShouldGenerateShaderSymbolsInfo(ShaderFormatName);
+		KeyString += bSymbols ? TEXT("_Symbols") : TEXT("");
+		KeyString += bSymbolsinfo ? TEXT("_SymbolsInfo") : TEXT("");
+		KeyString += bSymbols || bSymbolsinfo ? TEXT("_V2") : TEXT("");
 		// Are symbols based on source or results
 		KeyString += ShouldAllowUniqueShaderSymbols(ShaderFormatName) ? TEXT("_FullDbg") : TEXT("");
 	}
