Shader Compiler:  Made FShaderCompilerDefinitions private, to reduce code publicly visible in ShaderCore.h.  For now, it's just marked deprecated for 5.4, but will be hidden more generally in 5.5.

#rnx

diff --git a/Engine/Source/Developer/Apple/MetalShaderFormat/Private/MetalDerivedData.cpp b/Engine/Source/Developer/Apple/MetalShaderFormat/Private/MetalDerivedData.cpp
--- a/Engine/Source/Developer/Apple/MetalShaderFormat/Private/MetalDerivedData.cpp
+++ b/Engine/Source/Developer/Apple/MetalShaderFormat/Private/MetalDerivedData.cpp
@@ -11,6 +11,7 @@
 #include "MetalBackend.h"
 #include "Serialization/MemoryReader.h"
 #include "HlslccHeaderWriter.h"
+#include "Runtime/RenderCore/Internal/ShaderCompilerDefinitions.h"
 
 #include "MetalShaderFormat.h"
 #include "SpirvReflectCommon.h"
@@ -962,6 +963,8 @@ bool DoCompileMetalShader(
 
 		CrossCompiler::FShaderConductorTarget TargetDesc;
 
+		PRAGMA_DISABLE_DEPRECATION_WARNINGS		// FShaderCompilerDefinitions will be made internal in the future, marked deprecated until then
+
 		if (Result)
 		{
 			if (bUsingInlineRayTracing)
@@ -969,18 +972,18 @@ bool DoCompileMetalShader(
 				// For inline raytracing, we only need the scene Instance Descriptors (to emulate a missing intrinsic).
 				uint32_t RayTracingInstanceIndexBuffer = FPlatformMath::CountTrailingZeros64(BufferIndices);
 				BufferIndices &= ~(1ull << (uint64)RayTracingInstanceIndexBuffer);
-				TargetDesc.CompileFlags.SetDefine(TEXT("raytracing_instance_descriptor_table_index"), RayTracingInstanceIndexBuffer);
+				TargetDesc.CompileFlags->SetDefine(TEXT("raytracing_instance_descriptor_table_index"), RayTracingInstanceIndexBuffer);
 				RTString = FString::Printf(TEXT("// @RayTracingInstanceIndexBuffer: %u\n\n"), RayTracingInstanceIndexBuffer);
 			}
 
 			SideTableIndex = FPlatformMath::CountTrailingZeros64(BufferIndices);
 			BufferIndices &= ~(1ull << (uint64)SideTableIndex);
 
-			TargetDesc.CompileFlags.SetDefine(TEXT("texel_buffer_texture_width"), 0);
-			TargetDesc.CompileFlags.SetDefine(TEXT("enforce_storge_buffer_bounds"), 1);
-			TargetDesc.CompileFlags.SetDefine(TEXT("buffer_size_buffer_index"), SideTableIndex);
-			TargetDesc.CompileFlags.SetDefine(TEXT("invariant_float_math"), Options.bEnableFMAPass ? 1 : 0);
-			TargetDesc.CompileFlags.SetDefine(TEXT("enable_decoration_binding"), 1);
+			TargetDesc.CompileFlags->SetDefine(TEXT("texel_buffer_texture_width"), 0);
+			TargetDesc.CompileFlags->SetDefine(TEXT("enforce_storge_buffer_bounds"), 1);
+			TargetDesc.CompileFlags->SetDefine(TEXT("buffer_size_buffer_index"), SideTableIndex);
+			TargetDesc.CompileFlags->SetDefine(TEXT("invariant_float_math"), Options.bEnableFMAPass ? 1 : 0);
+			TargetDesc.CompileFlags->SetDefine(TEXT("enable_decoration_binding"), 1);
 
             if(Input.Target.Platform == SP_METAL_SM6)
             {
@@ -1002,7 +1005,7 @@ bool DoCompileMetalShader(
                     }
                 }
                 
-                TargetDesc.CompileFlags.SetDefine(TEXT("flatten_2d_array"), bShouldFlatten2DArray ? 1 : 0);
+				TargetDesc.CompileFlags->SetDefine(TEXT("flatten_2d_array"), bShouldFlatten2DArray ? 1 : 0);
             }
 
 			switch (Semantics)
@@ -1012,16 +1015,16 @@ bool DoCompileMetalShader(
 				break;
 			case EMetalGPUSemanticsTBDRDesktop:
 				TargetDesc.Language = CrossCompiler::EShaderConductorLanguage::Metal_iOS;
-                TargetDesc.CompileFlags.SetDefine(TEXT("ios_support_base_vertex_instance"), !bSupportAppleA8);
-				TargetDesc.CompileFlags.SetDefine(TEXT("use_framebuffer_fetch_subpasses"), 1);
-				TargetDesc.CompileFlags.SetDefine(TEXT("emulate_cube_array"), 1);
+                TargetDesc.CompileFlags->SetDefine(TEXT("ios_support_base_vertex_instance"), !bSupportAppleA8);
+				TargetDesc.CompileFlags->SetDefine(TEXT("use_framebuffer_fetch_subpasses"), 1);
+				TargetDesc.CompileFlags->SetDefine(TEXT("emulate_cube_array"), 1);
 				break;
 			case EMetalGPUSemanticsMobile:
 			default:
 				TargetDesc.Language = CrossCompiler::EShaderConductorLanguage::Metal_iOS;
-                TargetDesc.CompileFlags.SetDefine(TEXT("ios_support_base_vertex_instance"), !bSupportAppleA8);
-                TargetDesc.CompileFlags.SetDefine(TEXT("use_framebuffer_fetch_subpasses"), 1);
-				TargetDesc.CompileFlags.SetDefine(TEXT("emulate_cube_array"), 1);
+                TargetDesc.CompileFlags->SetDefine(TEXT("ios_support_base_vertex_instance"), !bSupportAppleA8);
+                TargetDesc.CompileFlags->SetDefine(TEXT("use_framebuffer_fetch_subpasses"), 1);
+				TargetDesc.CompileFlags->SetDefine(TEXT("emulate_cube_array"), 1);
 				break;
 			}
 
@@ -1044,16 +1047,18 @@ bool DoCompileMetalShader(
 				{
 					// If a dimension for the subpass input attachment at binding slot 0 was determined,
 					// forward this dimension to SPIRV-Cross because SPIR-V doesn't support a dimension for OpTypeImage instruction with SubpassData
-					TargetDesc.CompileFlags.SetDefine(subpass_input_dimension_names[SubpassIndex], SubpassInputDim);
+					TargetDesc.CompileFlags->SetDefine(subpass_input_dimension_names[SubpassIndex], SubpassInputDim);
 				}
 			}
 			
 			if (IABTier >= 1)
 			{
-				TargetDesc.CompileFlags.SetDefine(TEXT("argument_buffers"), 1);
-				TargetDesc.CompileFlags.SetDefine(TEXT("argument_buffer_offset"), IABOffsetIndex);
+				TargetDesc.CompileFlags->SetDefine(TEXT("argument_buffers"), 1);
+				TargetDesc.CompileFlags->SetDefine(TEXT("argument_buffer_offset"), IABOffsetIndex);
 			}
-			TargetDesc.CompileFlags.SetDefine(TEXT("texture_buffer_native"), 1);
+			TargetDesc.CompileFlags->SetDefine(TEXT("texture_buffer_native"), 1);
+
+			PRAGMA_ENABLE_DEPRECATION_WARNINGS
 			
 			switch (VersionEnum)
 			{
diff --git a/Engine/Source/Developer/Apple/MetalShaderFormat/Private/MetalShaderCompiler.cpp b/Engine/Source/Developer/Apple/MetalShaderFormat/Private/MetalShaderCompiler.cpp
--- a/Engine/Source/Developer/Apple/MetalShaderFormat/Private/MetalShaderCompiler.cpp
+++ b/Engine/Source/Developer/Apple/MetalShaderFormat/Private/MetalShaderCompiler.cpp
@@ -13,6 +13,7 @@
 #include "Misc/Paths.h"
 #include "Misc/EngineVersion.h"
 #include "HAL/PlatformFileManager.h"
+#include "Runtime/RenderCore/Internal/ShaderCompilerDefinitions.h"
 
 #include "MetalShaderFormat.h"
 
@@ -765,6 +766,8 @@ void BuildMetalShaderOutput(
 
 void CompileShader_Metal(const FShaderCompilerInput& _Input,FShaderCompilerOutput& Output,const FString& WorkingDirectory)
 {
+	PRAGMA_DISABLE_DEPRECATION_WARNINGS		// FShaderCompilerDefinitions will be made internal in the future, marked deprecated until then
+
 	auto Input = _Input;
 	FString PreprocessedShader;
 	FShaderCompilerDefinitions AdditionalDefines;
@@ -988,6 +991,8 @@ void CompileShader_Metal(const FShaderCompilerInput& _Input,FShaderCompilerOutpu
 
 	AdditionalDefines.SetDefine(TEXT("COMPILER_SUPPORTS_DUAL_SOURCE_BLENDING_SLOT_DECORATION"), (uint32)1);
 
+	PRAGMA_ENABLE_DEPRECATION_WARNINGS
+
 	const double StartPreprocessTime = FPlatformTime::Seconds();
 
 	if (Input.bSkipPreprocessedCache)
diff --git a/Engine/Source/Developer/ShaderCompilerCommon/Private/ShaderConductorContext.cpp b/Engine/Source/Developer/ShaderCompilerCommon/Private/ShaderConductorContext.cpp
--- a/Engine/Source/Developer/ShaderCompilerCommon/Private/ShaderConductorContext.cpp
+++ b/Engine/Source/Developer/ShaderCompilerCommon/Private/ShaderConductorContext.cpp
@@ -2,6 +2,7 @@
 
 #include "ShaderConductorContext.h"
 #include "HAL/ExceptionHandling.h"
+#include "Runtime/RenderCore/Internal/ShaderCompilerDefinitions.h"
 
 #if PLATFORM_MAC || PLATFORM_WINDOWS || PLATFORM_LINUX
 THIRD_PARTY_INCLUDES_START
@@ -245,6 +246,8 @@ namespace CrossCompiler
 		}
 	}
 
+	PRAGMA_DISABLE_DEPRECATION_WARNINGS		// FShaderCompilerDefinitions will be made internal in the future, marked deprecated until then
+
 	// Converts a map of string pairs to a C-Style macro defines array
 	static void ConvertDefineMapToMacroDefines(const FShaderCompilerDefinitions& Definitions, TArray<TPair<TArray<ANSICHAR>, TArray<ANSICHAR>>>& OutPairs, TArray<ShaderConductor::MacroDefine>& OutPairRefs)
 	{
@@ -265,6 +268,7 @@ namespace CrossCompiler
 			OutPairRefs[Index].value = OutPairs[Index].Value.GetData();
 		}
 	}
+	PRAGMA_ENABLE_DEPRECATION_WARNINGS
 
 	static void ConvertScTargetDesc(FShaderConductorContext::FShaderConductorIntermediates& Intermediates, const FShaderConductorTarget& InTarget, ShaderConductor::Compiler::TargetDesc& OutTargetDesc)
 	{
@@ -290,7 +294,7 @@ namespace CrossCompiler
 		}
 
 		// Convert flags map into an array container
-		ConvertDefineMapToMacroDefines(InTarget.CompileFlags, Intermediates.Flags, Intermediates.FlagRefs);
+		ConvertDefineMapToMacroDefines(*InTarget.CompileFlags, Intermediates.Flags, Intermediates.FlagRefs);
 
 		OutTargetDesc.options = Intermediates.FlagRefs.GetData();
 		OutTargetDesc.numOptions = static_cast<uint32>(Intermediates.FlagRefs.Num());
@@ -471,6 +475,13 @@ namespace CrossCompiler
 		}
 	}
 
+	FShaderConductorTarget::FShaderConductorTarget()
+	{
+		PRAGMA_DISABLE_DEPRECATION_WARNINGS		// FShaderCompilerDefinitions will be made internal in the future, marked deprecated until then
+		CompileFlags = MakePimpl<FShaderCompilerDefinitions>();
+		PRAGMA_ENABLE_DEPRECATION_WARNINGS
+	}
+
 	FShaderConductorContext::FShaderConductorContext()
 		: Intermediates(new FShaderConductorIntermediates())
 	{
@@ -497,7 +508,9 @@ namespace CrossCompiler
 		return *this;
 	}
 
+	PRAGMA_DISABLE_DEPRECATION_WARNINGS		// FShaderCompilerDefinitions will be made internal in the future, marked deprecated until then
 	bool FShaderConductorContext::LoadSource(const FString& ShaderSource, const FString& Filename, const FString& EntryPoint, EShaderFrequency ShaderStage, const FShaderCompilerDefinitions* Definitions, const TArray<FString>* ExtraDxcArgs)
+	PRAGMA_ENABLE_DEPRECATION_WARNINGS
 	{
 		// Convert FString to ANSI string and store them as intermediates
 		ConvertFStringToAnsiString(ShaderSource, Intermediates->ShaderSource);
@@ -521,7 +534,9 @@ namespace CrossCompiler
 		return true;
 	}
 
+	PRAGMA_DISABLE_DEPRECATION_WARNINGS		// FShaderCompilerDefinitions will be made internal in the future, marked deprecated until then
 	bool FShaderConductorContext::LoadSource(const ANSICHAR* ShaderSource, const ANSICHAR* Filename, const ANSICHAR* EntryPoint, EShaderFrequency ShaderStage, const FShaderCompilerDefinitions* Definitions, const TArray<FString>* ExtraDxcArgs)
+	PRAGMA_ENABLE_DEPRECATION_WARNINGS
 	{
 		// Store ANSI strings as intermediates
 		CopyAnsiString(ShaderSource, Intermediates->ShaderSource);
diff --git a/Engine/Source/Developer/ShaderCompilerCommon/Public/ShaderConductorContext.h b/Engine/Source/Developer/ShaderCompilerCommon/Public/ShaderConductorContext.h
--- a/Engine/Source/Developer/ShaderCompilerCommon/Public/ShaderConductorContext.h
+++ b/Engine/Source/Developer/ShaderCompilerCommon/Public/ShaderConductorContext.h
@@ -13,6 +13,8 @@
 #include "ShaderCore.h"
 #include "Templates/Function.h"
 
+class FShaderCompilerDefinitions;
+
 // Cross compiler support/common functionality
 namespace CrossCompiler
 {
@@ -130,6 +132,8 @@ namespace CrossCompiler
 	/** Shader conductor output target descriptor. */
 	struct SHADERCOMPILERCOMMON_API FShaderConductorTarget
 	{
+		FShaderConductorTarget();
+
 		/** Target shader semantics, e.g. "macOS" or "iOS" for Metal GPU semantics. */
 		EShaderConductorLanguage Language = EShaderConductorLanguage::Glsl;
 
@@ -142,7 +146,9 @@ namespace CrossCompiler
 		int32 Version = 0;
 
 		/** Cross compilation flags. This is used for high-level cross compilation (such as Metal output) that is send over to SPIRV-Cross, e.g. { "invariant_float_math", "1" }. */
-		FShaderCompilerDefinitions CompileFlags;
+		PRAGMA_DISABLE_DEPRECATION_WARNINGS		// FShaderCompilerDefinitions will be made internal in the future, marked deprecated until then
+		TPimplPtr<FShaderCompilerDefinitions> CompileFlags;
+		PRAGMA_ENABLE_DEPRECATION_WARNINGS
 
 		/** Optional callback to rename certain variable types. */
 		TFunction<bool(const FAnsiStringView& VariableName, const FAnsiStringView& TypeName, FString& OutRenamedTypeName)> VariableTypeRenameCallback;
@@ -184,8 +190,10 @@ namespace CrossCompiler
 		FShaderConductorContext& operator = (const FShaderConductorContext&) = delete;
 
 		/** Loads the shader source and converts the input descriptor to a format suitable for ShaderConductor. If 'Definitions' is null, the previously loaded definitions are not modified. */
+		PRAGMA_DISABLE_DEPRECATION_WARNINGS		// FShaderCompilerDefinitions will be made internal in the future, marked deprecated until then
 		bool LoadSource(const FString& ShaderSource, const FString& Filename, const FString& EntryPoint, EShaderFrequency ShaderStage, const FShaderCompilerDefinitions* Definitions = nullptr, const TArray<FString>* ExtraDxcArgs = nullptr);
 		bool LoadSource(const ANSICHAR* ShaderSource, const ANSICHAR* Filename, const ANSICHAR* EntryPoint, EShaderFrequency ShaderStage, const FShaderCompilerDefinitions* Definitions = nullptr, const TArray<FString>* ExtraDxcArgs = nullptr);
+		PRAGMA_ENABLE_DEPRECATION_WARNINGS
 
 		/** Rewrites the specified HLSL shader source code. This allows to reduce the HLSL code by removing unused global resources for instance.
 		This will update the internally loaded source (see 'LoadSource'), so the output parameter 'OutSource' is optional. */
diff --git a/Engine/Source/Developer/ShaderFormatOpenGL/Private/OpenGLShaderCompiler.cpp b/Engine/Source/Developer/ShaderFormatOpenGL/Private/OpenGLShaderCompiler.cpp
--- a/Engine/Source/Developer/ShaderFormatOpenGL/Private/OpenGLShaderCompiler.cpp
+++ b/Engine/Source/Developer/ShaderFormatOpenGL/Private/OpenGLShaderCompiler.cpp
@@ -13,6 +13,7 @@
 #include "SpirvReflectCommon.h"
 #include <algorithm>
 #include <regex>
+#include "Runtime/RenderCore/Internal/ShaderCompilerDefinitions.h"
 
 #if PLATFORM_WINDOWS
 #include "Windows/AllowWindowsPlatformTypes.h"
@@ -1110,6 +1111,9 @@ void FOpenGLFrontend::PrecompileShader(FShaderCompilerOutput& ShaderOutput, cons
 	PlatformReleaseOpenGL(ContextPtr, PrevContextPtr);
 }
 
+
+PRAGMA_DISABLE_DEPRECATION_WARNINGS		// FShaderCompilerDefinitions will be made internal in the future, marked deprecated until then
+
 void FOpenGLFrontend::SetupPerVersionCompilationEnvironment(GLSLVersion Version, FShaderCompilerDefinitions& AdditionalDefines, EHlslCompileTarget& HlslCompilerTarget)
 {
 	switch (Version)
@@ -1134,6 +1138,8 @@ void FOpenGLFrontend::SetupPerVersionCompilationEnvironment(GLSLVersion Version,
 	AdditionalDefines.SetDefine(TEXT("OPENGL_PROFILE"), 1);
 }
 
+PRAGMA_ENABLE_DEPRECATION_WARNINGS
+
 uint32 FOpenGLFrontend::GetMaxSamplers(GLSLVersion Version)
 {
 	return 16;
@@ -2840,7 +2846,9 @@ bool GenerateDeferredMobileShaders(std::string& GlslSource, GLSLCompileParameter
 		OutString += ESVersionString + "\n";
 		OutString += "#ifdef UE_MRT_FRAMEBUFFER_FETCH\n";
 
-		FShaderCompilerDefinitions CompileFlagsCopy = GLSLCompileParams.TargetDesc->CompileFlags;
+		PRAGMA_DISABLE_DEPRECATION_WARNINGS		// FShaderCompilerDefinitions will be made internal in the future, marked deprecated until then
+
+		FShaderCompilerDefinitions CompileFlagsCopy = *GLSLCompileParams.TargetDesc->CompileFlags;
 
 		// Add defines for FBF in spirv-glsl
 		for (uint32_t i = 1; i < 4; ++i)
@@ -2849,7 +2857,7 @@ bool GenerateDeferredMobileShaders(std::string& GlslSource, GLSLCompileParameter
 			{
 				FString DefineKey = FString::Printf(TEXT("remap_ext_framebuffer_fetch%d"), i);
 				FString DefineValue = FString::Printf(TEXT("%d %d"), i + 1, i);
-				GLSLCompileParams.TargetDesc->CompileFlags.SetDefine(*DefineKey, *DefineValue);
+				GLSLCompileParams.TargetDesc->CompileFlags->SetDefine(*DefineKey, *DefineValue);
 			}
 		}
 
@@ -2865,7 +2873,9 @@ bool GenerateDeferredMobileShaders(std::string& GlslSource, GLSLCompileParameter
 			FBFSourceString.replace(VersionStringPos, ESVersionString.length(), "");
 		}
 
-		GLSLCompileParams.TargetDesc->CompileFlags = CompileFlagsCopy;
+		*GLSLCompileParams.TargetDesc->CompileFlags = CompileFlagsCopy;
+
+		PRAGMA_ENABLE_DEPRECATION_WARNINGS
 
 		OutString += FBFSourceString;
 
@@ -2920,7 +2930,9 @@ bool GenerateDeferredMobileShaders(std::string& GlslSource, GLSLCompileParameter
 					DefineValue = GBufferOutputNames[i];
 				}
 
-				GLSLCompileParams.TargetDesc->CompileFlags.SetDefine(*DefineKey, *DefineValue);
+				PRAGMA_DISABLE_DEPRECATION_WARNINGS		// FShaderCompilerDefinitions will be made internal in the future, marked deprecated until then
+				GLSLCompileParams.TargetDesc->CompileFlags->SetDefine(*DefineKey, *DefineValue);
+				PRAGMA_ENABLE_DEPRECATION_WARNINGS
 			}
 		}
 
@@ -3037,10 +3049,14 @@ static bool CompileToGlslWithShaderConductor(
 		SourceData = HlslFrameBufferDeclarations + SourceData;
 	}
 	
+	PRAGMA_DISABLE_DEPRECATION_WARNINGS		// FShaderCompilerDefinitions will be made internal in the future, marked deprecated until then
+
 	// Inject additional macro definitions to circumvent missing features: external textures
 	FShaderCompilerDefinitions AdditionalDefines;
 	AdditionalDefines.SetDefine(TEXT("TextureExternal"), TEXT("Texture2D"));
 
+	PRAGMA_ENABLE_DEPRECATION_WARNINGS
+
 	UE::ShaderCompilerCommon::FDebugShaderDataOptions DebugDataOptions;
 	DebugDataOptions.HlslCCFlags = CCFlags;
 	UE::ShaderCompilerCommon::DumpDebugShaderData(Input, PreprocessedShader, DebugDataOptions);
@@ -3124,7 +3140,9 @@ static bool CompileToGlslWithShaderConductor(
 
 		CrossCompiler::FShaderConductorTarget TargetDesc;
 
-		TargetDesc.CompileFlags.SetDefine(TEXT("relax_nan_checks"), 1);
+		PRAGMA_DISABLE_DEPRECATION_WARNINGS		// FShaderCompilerDefinitions will be made internal in the future, marked deprecated until then
+
+		TargetDesc.CompileFlags->SetDefine(TEXT("relax_nan_checks"), 1);
 
 		switch (Version)
 		{
@@ -3139,9 +3157,9 @@ static bool CompileToGlslWithShaderConductor(
 		case GLSL_150_ES3_1:
 		case GLSL_ES3_1_ANDROID:
 		default:
-			TargetDesc.CompileFlags.SetDefine(TEXT("force_flattened_io_blocks"), 1);
-			TargetDesc.CompileFlags.SetDefine(TEXT("emit_uniform_buffer_as_plain_uniforms"), 1);
-			TargetDesc.CompileFlags.SetDefine(TEXT("pad_ubo_blocks"), 1);
+			TargetDesc.CompileFlags->SetDefine(TEXT("force_flattened_io_blocks"), 1);
+			TargetDesc.CompileFlags->SetDefine(TEXT("emit_uniform_buffer_as_plain_uniforms"), 1);
+			TargetDesc.CompileFlags->SetDefine(TEXT("pad_ubo_blocks"), 1);
 			// TODO: Currently disabled due to bug when assigning an array to temporary variable
 			//TargetDesc.CompileFlags.SetDefine(TEXT("force_temporary"), 1);
 
@@ -3149,18 +3167,20 @@ static bool CompileToGlslWithShaderConductor(
 			const bool bMultiView = Input.Environment.GetCompileArgument(TEXT("MOBILE_MULTI_VIEW"), false);
 			if (Frequency == SF_Vertex && bMultiView)
 			{
-				TargetDesc.CompileFlags.SetDefine(TEXT("ovr_multiview_view_count"), 2);
+				TargetDesc.CompileFlags->SetDefine(TEXT("ovr_multiview_view_count"), 2);
 			}
 
 			if (Version == GLSL_150_ES3_1)
 			{
-				TargetDesc.CompileFlags.SetDefine(TEXT("force_glsl_clipspace"), 1);
+				TargetDesc.CompileFlags->SetDefine(TEXT("force_glsl_clipspace"), 1);
 			}
 			TargetDesc.Language = CrossCompiler::EShaderConductorLanguage::Essl;
 			TargetDesc.Version = 320;
 			break;
 		}
 
+		PRAGMA_ENABLE_DEPRECATION_WARNINGS
+
 		TSet<FString> ExternalTextures;
 		int32 Pos = 0;
 #if !PLATFORM_MAC
@@ -3291,6 +3311,8 @@ static inline FString GetExtension(EHlslShaderFrequency Frequency, bool bAddDot
  */
 void FOpenGLFrontend::CompileShader(const FShaderCompilerInput& Input, FShaderCompilerOutput& Output, const FString& WorkingDirectory, GLSLVersion Version)
 {
+	PRAGMA_DISABLE_DEPRECATION_WARNINGS		// FShaderCompilerDefinitions will be made internal in the future, marked deprecated until then
+
 	FString PreprocessedShader;
 	FShaderCompilerDefinitions AdditionalDefines;
 	EHlslCompileTarget HlslCompilerTarget = HCT_InvalidTarget;
@@ -3323,6 +3345,8 @@ void FOpenGLFrontend::CompileShader(const FShaderCompilerInput& Input, FShaderCo
 		AdditionalDefines.SetDefine(TEXT("FORCE_FLOATS"), (uint32)1);
 	}
 
+	PRAGMA_ENABLE_DEPRECATION_WARNINGS
+
 	if (Input.bSkipPreprocessedCache)
 	{
 		if (!FFileHelper::LoadFileToString(PreprocessedShader, *Input.VirtualSourceFilePath))
diff --git a/Engine/Source/Developer/ShaderFormatVectorVM/Private/VectorVMShaderCompiler.cpp b/Engine/Source/Developer/ShaderFormatVectorVM/Private/VectorVMShaderCompiler.cpp
--- a/Engine/Source/Developer/ShaderFormatVectorVM/Private/VectorVMShaderCompiler.cpp
+++ b/Engine/Source/Developer/ShaderFormatVectorVM/Private/VectorVMShaderCompiler.cpp
@@ -11,6 +11,7 @@
 #include "ShaderCore.h"
 #include "ShaderFormatVectorVM.h"
 #include "ShaderPreprocessor.h"
+#include "Runtime/RenderCore/Internal/ShaderCompilerDefinitions.h"
 
 #include "VectorVM.h"
 #include "VectorVMBackend.h"
@@ -53,6 +54,8 @@ bool CompileShader_VectorVM(const FShaderCompilerInput& Input, FShaderCompilerOu
 {
 	SCOPE_CYCLE_COUNTER(STAT_VectorVM_Compiler_CompileShader_VectorVM);
 
+	PRAGMA_DISABLE_DEPRECATION_WARNINGS		// FShaderCompilerDefinitions will be made internal in the future, marked deprecated until then
+
 	FString PreprocessedShader;
 	FShaderCompilerDefinitions AdditionalDefines;
 	EHlslCompileTarget HlslCompilerTarget = HCT_FeatureLevelSM5;
@@ -65,6 +68,8 @@ bool CompileShader_VectorVM(const FShaderCompilerInput& Input, FShaderCompilerOu
 	const bool bDumpDebugInfo = (Input.DumpDebugInfoPath != TEXT("") && IFileManager::Get().DirectoryExists(*Input.DumpDebugInfoPath));
 
 	AdditionalDefines.SetDefine(TEXT("FORCE_FLOATS"), (uint32)1);
+	
+	PRAGMA_ENABLE_DEPRECATION_WARNINGS
 
 	if (Input.bSkipPreprocessedCache)
 	{
diff --git a/Engine/Source/Developer/ShaderPreprocessor/Private/ShaderPreprocessor.cpp b/Engine/Source/Developer/ShaderPreprocessor/Private/ShaderPreprocessor.cpp
--- a/Engine/Source/Developer/ShaderPreprocessor/Private/ShaderPreprocessor.cpp
+++ b/Engine/Source/Developer/ShaderPreprocessor/Private/ShaderPreprocessor.cpp
@@ -6,6 +6,7 @@
 #include "Misc/ScopeLock.h"
 #include "Modules/ModuleManager.h"
 #include "PreprocessorPrivate.h"
+#include "Runtime/RenderCore/Internal/ShaderCompilerDefinitions.h"
 
 #include "stb_preprocess/preprocessor.h"
 #include "stb_preprocess/stb_alloc.h"
@@ -23,6 +24,7 @@ namespace
 	}
 }
 
+PRAGMA_DISABLE_DEPRECATION_WARNINGS		// FShaderCompilerDefinitions will be made internal in the future, marked deprecated until then
 static void AddStbDefine(stb_arena* MacroArena, macro_definition**& StbDefines, const TCHAR* Name, const TCHAR* Value);
 static void AddStbDefines(stb_arena* MacroArena, macro_definition**& StbDefines, const FShaderCompilerDefinitions& Defines);
 
@@ -32,8 +34,8 @@ public:
 	static void DumpShaderDefinesAsCommentedCode(const FShaderCompilerEnvironment& Environment, FString* OutDefines)
 	{
 		TArray<FString> DefinesLines;
-		DefinesLines.Reserve(Environment.Definitions.Num());
-		for (FShaderCompilerDefinitions::FConstIterator DefineIt(Environment.Definitions); DefineIt; ++DefineIt)
+		DefinesLines.Reserve(Environment.Definitions->Num());
+		for (FShaderCompilerDefinitions::FConstIterator DefineIt(*Environment.Definitions); DefineIt; ++DefineIt)
 		{
 			DefinesLines.Add(FString::Printf(TEXT("// #define %s %s\n"), DefineIt.Key(), DefineIt.Value()));
 		}
@@ -50,10 +52,11 @@ public:
 
 	static void PopulateDefines(const FShaderCompilerEnvironment& Environment, const FShaderCompilerDefinitions& AdditionalDefines, stb_arena* MacroArena, macro_definition**& OutDefines)
 	{
-		AddStbDefines(MacroArena, OutDefines, Environment.Definitions);
+		AddStbDefines(MacroArena, OutDefines, *Environment.Definitions);
 		AddStbDefines(MacroArena, OutDefines, AdditionalDefines);
 	}
 };
+PRAGMA_ENABLE_DEPRECATION_WARNINGS
 
 //////////////////////////////////////////////////////////////////////////
 extern "C"
@@ -302,6 +305,8 @@ static void AddStbDefine(stb_arena* MacroArena, macro_definition**& StbDefines,
 	arrput(StbDefines, pp_define(MacroArena, (ANSICHAR*)ConvertedDefine.Get()));
 }
 
+PRAGMA_DISABLE_DEPRECATION_WARNINGS		// FShaderCompilerDefinitions will be made internal in the future, marked deprecated until then
+
 static void AddStbDefines(stb_arena* MacroArena, macro_definition**& StbDefines, const FShaderCompilerDefinitions& Defines)
 {
 	for (FShaderCompilerDefinitions::FConstIterator It(Defines); It; ++It)
@@ -321,6 +326,8 @@ bool InnerPreprocessShaderStb(
 	macro_definition** StbDefines = nullptr;
 	FShaderPreprocessorUtilities::PopulateDefines(Environment, AdditionalDefines, &MacroArena, StbDefines);
 
+	PRAGMA_ENABLE_DEPRECATION_WARNINGS
+
 	FStbPreprocessContext Context{ Input, Environment };
 
 	auto InFilename = StringCast<ANSICHAR>(*Input.VirtualSourceFilePath);
@@ -378,7 +385,9 @@ bool PreprocessShader(
 	FString& OutPreprocessedShader,
 	FShaderCompilerOutput& ShaderOutput,
 	const FShaderCompilerInput& ShaderInput,
+	PRAGMA_DISABLE_DEPRECATION_WARNINGS		// FShaderCompilerDefinitions will be made internal in the future, marked deprecated until then
 	const FShaderCompilerDefinitions& AdditionalDefines,
+	PRAGMA_ENABLE_DEPRECATION_WARNINGS
 	EDumpShaderDefines DefinesPolicy)
 {
 	FShaderPreprocessOutput Output;
@@ -409,7 +418,9 @@ bool PreprocessShader(
 	FShaderPreprocessOutput& Output,
 	const FShaderCompilerInput& Input,
 	const FShaderCompilerEnvironment& Environment,
+	PRAGMA_DISABLE_DEPRECATION_WARNINGS		// FShaderCompilerDefinitions will be made internal in the future, marked deprecated until then
 	const FShaderCompilerDefinitions& AdditionalDefines,
+	PRAGMA_ENABLE_DEPRECATION_WARNINGS
 	EDumpShaderDefines DefinesPolicy
 )
 {
diff --git a/Engine/Source/Developer/ShaderPreprocessor/Public/ShaderPreprocessor.h b/Engine/Source/Developer/ShaderPreprocessor/Public/ShaderPreprocessor.h
--- a/Engine/Source/Developer/ShaderPreprocessor/Public/ShaderPreprocessor.h
+++ b/Engine/Source/Developer/ShaderPreprocessor/Public/ShaderPreprocessor.h
@@ -38,7 +38,9 @@ extern SHADERPREPROCESSOR_API bool PreprocessShader(
 	FString& OutPreprocessedShader,
 	FShaderCompilerOutput& ShaderOutput,
 	const FShaderCompilerInput& ShaderInput,
+	PRAGMA_DISABLE_DEPRECATION_WARNINGS		// FShaderCompilerDefinitions will be made internal in the future, marked deprecated until then
 	const FShaderCompilerDefinitions& AdditionalDefines,
+	PRAGMA_ENABLE_DEPRECATION_WARNINGS
 	EDumpShaderDefines DefinesPolicy = EDumpShaderDefines::DontCare);
 
 /**
@@ -56,5 +58,7 @@ extern SHADERPREPROCESSOR_API bool PreprocessShader(
 	FShaderPreprocessOutput& Output,
 	const FShaderCompilerInput& Input,
 	const FShaderCompilerEnvironment& MergedEnvironment,
+	PRAGMA_DISABLE_DEPRECATION_WARNINGS		// FShaderCompilerDefinitions will be made internal in the future, marked deprecated until then
 	const FShaderCompilerDefinitions& AdditionalDefines,
+	PRAGMA_ENABLE_DEPRECATION_WARNINGS
 	EDumpShaderDefines DefinesPolicy = EDumpShaderDefines::DontCare);
diff --git a/Engine/Source/Developer/VulkanShaderFormat/Private/VulkanShaderCompiler.cpp b/Engine/Source/Developer/VulkanShaderFormat/Private/VulkanShaderCompiler.cpp
--- a/Engine/Source/Developer/VulkanShaderFormat/Private/VulkanShaderCompiler.cpp
+++ b/Engine/Source/Developer/VulkanShaderFormat/Private/VulkanShaderCompiler.cpp
@@ -10,6 +10,7 @@
 #include "hlslcc.h"
 #include "SpirvReflectCommon.h"
 #include "RHIShaderFormatDefinitions.inl"
+#include "Runtime/RenderCore/Internal/ShaderCompilerDefinitions.h"
 
 #if PLATFORM_MAC
 // Horrible hack as we need the enum available but the Vulkan headers do not compile on Mac
@@ -2361,7 +2362,9 @@ static bool CompileWithShaderConductor(
 	CrossCompiler::FShaderConductorContext CompilerContext;
 
 	// Inject additional macro definitions to circumvent missing features: external textures
+	PRAGMA_DISABLE_DEPRECATION_WARNINGS		// FShaderCompilerDefinitions will be made internal in the future, marked deprecated until then
 	FShaderCompilerDefinitions AdditionalDefines;
+	PRAGMA_ENABLE_DEPRECATION_WARNINGS
 
 	// Load shader source into compiler context
 	CompilerContext.LoadSource(PreprocessedShader, Input.VirtualSourceFilePath, EntryPointName, Frequency, &AdditionalDefines);
@@ -2595,6 +2598,8 @@ void DoCompileVulkanShader(const FShaderCompilerInput& Input, FShaderCompilerOut
 		return;
 	}
 
+	PRAGMA_DISABLE_DEPRECATION_WARNINGS		// FShaderCompilerDefinitions will be made internal in the future, marked deprecated until then
+
 	FString PreprocessedShader;
 	FShaderCompilerDefinitions AdditionalDefines;
 	AdditionalDefines.SetDefine(TEXT("COMPILER_HLSLCC"), 1);
@@ -2653,6 +2658,8 @@ void DoCompileVulkanShader(const FShaderCompilerInput& Input, FShaderCompilerOut
 		AdditionalDefines.SetDefine(TEXT("EMULATE_VKDEVICEADRESS"), 1);
 	}
 
+	PRAGMA_ENABLE_DEPRECATION_WARNINGS
+
 	const double StartPreprocessTime = FPlatformTime::Seconds();
 
 	// Preprocess the shader.
diff --git a/Engine/Source/Developer/Windows/ShaderFormatD3D/Private/D3DShaderCompiler.cpp b/Engine/Source/Developer/Windows/ShaderFormatD3D/Private/D3DShaderCompiler.cpp
--- a/Engine/Source/Developer/Windows/ShaderFormatD3D/Private/D3DShaderCompiler.cpp
+++ b/Engine/Source/Developer/Windows/ShaderFormatD3D/Private/D3DShaderCompiler.cpp
@@ -13,6 +13,7 @@
 #include "Serialization/MemoryWriter.h"
 #include "RayTracingDefinitions.h"
 #include "SpirvCommon.h"
+#include "Runtime/RenderCore/Internal/ShaderCompilerDefinitions.h"
 
 DEFINE_LOG_CATEGORY_STATIC(LogD3D11ShaderCompiler, Log, All);
 
@@ -650,13 +651,16 @@ static bool CompileAndProcessD3DShaderFXCExt(
 			CrossCompiler::FShaderConductorTarget TargetDesc;
 			TargetDesc.Language = CrossCompiler::EShaderConductorLanguage::Hlsl;
 			TargetDesc.Version = 50;
-			TargetDesc.CompileFlags.SetDefine(TEXT("implicit_resource_binding"), 1);
-			TargetDesc.CompileFlags.SetDefine(TEXT("reconstruct_global_uniforms"), 1);
-			TargetDesc.CompileFlags.SetDefine(TEXT("reconstruct_cbuffer_names"), 1);
-			TargetDesc.CompileFlags.SetDefine(TEXT("reconstruct_semantics"), 1);
-			TargetDesc.CompileFlags.SetDefine(TEXT("force_zero_initialized_variables"), 1);
-			TargetDesc.CompileFlags.SetDefine(TEXT("relax_nan_checks"), 1);
-			TargetDesc.CompileFlags.SetDefine(TEXT("preserve_structured_buffers"), 1);
+
+			PRAGMA_DISABLE_DEPRECATION_WARNINGS		// FShaderCompilerDefinitions will be made internal in the future, marked deprecated until then
+			TargetDesc.CompileFlags->SetDefine(TEXT("implicit_resource_binding"), 1);
+			TargetDesc.CompileFlags->SetDefine(TEXT("reconstruct_global_uniforms"), 1);
+			TargetDesc.CompileFlags->SetDefine(TEXT("reconstruct_cbuffer_names"), 1);
+			TargetDesc.CompileFlags->SetDefine(TEXT("reconstruct_semantics"), 1);
+			TargetDesc.CompileFlags->SetDefine(TEXT("force_zero_initialized_variables"), 1);
+			TargetDesc.CompileFlags->SetDefine(TEXT("relax_nan_checks"), 1);
+			TargetDesc.CompileFlags->SetDefine(TEXT("preserve_structured_buffers"), 1);
+			PRAGMA_ENABLE_DEPRECATION_WARNINGS
 
 			// Patch SPIR-V for workarounds to prevent potential additional FXC failures
 			PatchSpirvForPrecompilation(Spirv);
@@ -1106,6 +1110,7 @@ bool PreprocessD3DShader(
 	FShaderPreprocessOutput& Output,
 	ELanguage Language)
 {
+	PRAGMA_DISABLE_DEPRECATION_WARNINGS		// FShaderCompilerDefinitions will be made internal in the future, marked deprecated until then
 	FShaderCompilerDefinitions AdditionalDefines;
 	if (Language == ELanguage::SM6)
 	{
@@ -1150,6 +1155,7 @@ bool PreprocessD3DShader(
 		}
 	}
 
+	PRAGMA_ENABLE_DEPRECATION_WARNINGS
 	
 	if (Input.bSkipPreprocessedCache)
 	{
diff --git a/Engine/Source/Runtime/Engine/Private/ShaderCompiler/ShaderGenerationUtil.cpp b/Engine/Source/Runtime/Engine/Private/ShaderCompiler/ShaderGenerationUtil.cpp
--- a/Engine/Source/Runtime/Engine/Private/ShaderCompiler/ShaderGenerationUtil.cpp
+++ b/Engine/Source/Runtime/Engine/Private/ShaderCompiler/ShaderGenerationUtil.cpp
@@ -15,6 +15,7 @@
 #include "RenderUtils.h"
 #include "SceneManagement.h"
 #include "ShaderMaterial.h"
+#include "Runtime/RenderCore/Internal/ShaderCompilerDefinitions.h"
 
 static TAutoConsoleVariable<int32> CVarShaderUseGBufferRefactor(
 	TEXT("r.Shaders.UseGBufferRefactor"),
@@ -290,6 +291,8 @@ struct FShaderInitialDefinesInitializer
 		ApplyFetchEnvironmentInternal(LightmapDefines, GatherNames);
 		ApplyFetchEnvironmentInternal(CompilerDefines, GatherNames);
 
+		PRAGMA_DISABLE_DEPRECATION_WARNINGS		// FShaderCompilerDefinitions will be made internal in the future, marked deprecated until then
+
 		FShaderCompilerDefinitions InitialDefines;
 		for (FName DefineKey : GatherNames.Names)
 		{
@@ -297,6 +300,8 @@ struct FShaderInitialDefinesInitializer
 			InitialDefines.FindOrAddMapIndex(DefineKey);
 		}
 		FShaderCompilerDefinitions::InitializeInitialDefines(InitialDefines);
+
+		PRAGMA_ENABLE_DEPRECATION_WARNINGS
 	}
 };
 static FShaderInitialDefinesInitializer GInitialDefinesInitializer;
@@ -2161,7 +2166,7 @@ void FShaderCompileUtilities::GenerateBrdfHeaders(const FName& ShaderFormat)
 
 EGBufferLayout FShaderCompileUtilities::FetchGBufferLayout(const FShaderCompilerEnvironment& Environment)
 {
-	const uint32 Layout = Environment.Definitions.GetIntegerValue(TEXT("GBUFFER_LAYOUT"));
+	const uint32 Layout = Environment.GetIntegerValue(TEXT("GBUFFER_LAYOUT"));
 	if (Layout >= GBL_Num)
 	{
 		return GBL_Default;
diff --git a/Engine/Source/Runtime/RenderCore/Internal/ShaderCompilerDefinitions.h b/Engine/Source/Runtime/RenderCore/Internal/ShaderCompilerDefinitions.h
new file mode 100644
--- /dev/null
+++ b/Engine/Source/Runtime/RenderCore/Internal/ShaderCompilerDefinitions.h
@@ -0,0 +1,470 @@
+// Copyright Epic Games, Inc. All Rights Reserved.
+
+#pragma once
+
+#include "ShaderCore.h"
+
+enum class EShaderCompilerDefineVariant : uint8
+{
+	None = 0,
+	Integer,
+	Unsigned,
+	Float,
+	String
+};
+
+class UE_DEPRECATED(5.4, "FShaderCompilerDefinitions will be made internal in the future, use FShaderCompilerEnvironment instead.") FShaderCompilerDefinitions;
+
+PRAGMA_DISABLE_DEPRECATION_WARNINGS
+
+/** Container for shader compiler definitions. */
+class FShaderCompilerDefinitions
+{
+public:
+	RENDERCORE_API FShaderCompilerDefinitions(bool bIncludeInitialDefines = false);
+	RENDERCORE_API FShaderCompilerDefinitions(const FShaderCompilerDefinitions&);
+
+	/** Value types supported:   bool, int32, uint32, float, const TCHAR*, FString& */
+	template <typename ValueType>
+	void SetDefine(FName Name, ValueType Value) { InternalSetValue(FindOrAddMapIndex(Name), Value); }
+	void SetDefine(FName Name, const TCHAR* Value) { InternalSetValue(FindOrAddMapIndex(Name), Value); }
+	void SetDefine(FName Name, const FString& Value) { InternalSetValue(FindOrAddMapIndex(Name), *Value); }
+
+	template <typename ValueType>
+	void SetDefine(FShaderCompilerDefineNameCache& Name, ValueType Value) { InternalSetValue(FindOrAddMapIndex(Name), Value); }
+	void SetDefine(FShaderCompilerDefineNameCache& Name, const TCHAR* Value) { InternalSetValue(FindOrAddMapIndex(Name), Value); }
+	void SetDefine(FShaderCompilerDefineNameCache& Name, const FString& Value) { InternalSetValue(FindOrAddMapIndex(Name), *Value); }
+
+	int32 GetIntegerValue(FName Name) const
+	{
+		int32 Result = 0;
+		int32 KeyIndex = FindMapIndex(Name, GetTypeHash(Name));
+		if (KeyIndex != INDEX_NONE)
+		{
+			// For None, Integer, or Unsigned, return ValueInteger (None will have a default of zero)
+			if (ValueTypes[KeyIndex] <= EShaderCompilerDefineVariant::Unsigned)
+			{
+				Result = Pairs[KeyIndex].ValueInteger;
+			}
+			else if (ValueTypes[KeyIndex] == EShaderCompilerDefineVariant::String)
+			{
+				Result = FCString::Atoi(*StringValues[Pairs[KeyIndex].ValueInteger]);
+			}
+			else
+			{
+				// EShaderCompilerDefineVariant::Float
+				Result = (int32)Pairs[KeyIndex].ValueFloat;
+			}
+		}
+		return Result;
+	}
+
+	int32 GetIntegerValue(FShaderCompilerDefineNameCache& NameCache, int32 ResultIfNotFound) const
+	{
+		int32 Result = ResultIfNotFound;
+		int32 KeyIndex;
+
+		// Check if this is an initial define, meaning it has a fixed map index.  If Name.MapIndex is INDEX_NONE
+		// (MapIndex not initialized yet) or InitialDefineCount is zero (FShaderCompilerDefinitions constructed without
+		// initial state), this condition will be false, and continue to the rest of the function.
+		if ((uint32)NameCache.MapIndex < InitialDefineCount)
+		{
+			KeyIndex = NameCache.MapIndex;
+		}
+		else
+		{
+			KeyIndex = FindMapIndex(NameCache.Name, GetTypeHash(NameCache.Name));
+		}
+
+		if (KeyIndex != INDEX_NONE)
+		{
+			// Initialize MapIndex if necessary.  If the define is not an initial define, its index will be greater than
+			// or equal to InitialDefineCount, indicating that it doesn't have a fixed map index.
+			if (InitialDefineCount && NameCache.MapIndex == INDEX_NONE)
+			{
+				NameCache.MapIndex = KeyIndex;
+			}
+
+			EShaderCompilerDefineVariant ValueType = ValueTypes[KeyIndex];
+			if (ValueType == EShaderCompilerDefineVariant::Integer ||
+				ValueType == EShaderCompilerDefineVariant::Unsigned)
+			{
+				Result = Pairs[KeyIndex].ValueInteger;
+			}
+			else if (ValueType == EShaderCompilerDefineVariant::String)
+			{
+				Result = FCString::Atoi(*StringValues[Pairs[KeyIndex].ValueInteger]);
+			}
+			else if (ValueType == EShaderCompilerDefineVariant::Float)
+			{
+				// EShaderCompilerDefineVariant::Float
+				Result = (int32)Pairs[KeyIndex].ValueFloat;
+			}
+		}
+		return Result;
+	}
+
+	bool Contains(FName Name) const
+	{
+		int32 KeyIndex = FindMapIndex(Name, GetTypeHash(Name));
+		return KeyIndex != INDEX_NONE && ValueTypes[KeyIndex] != EShaderCompilerDefineVariant::None;
+	}
+
+	FORCEINLINE int32 Num() const
+	{
+		return ValueCount;
+	}
+
+	void Empty()
+	{
+		if (ValueCount)
+		{
+			KeyHashTable.Free();
+			Pairs.Empty();
+			ValueTypes.Empty();
+			StringValues.Empty();
+			ValueCount = 0;
+
+			// Re-copy initial defines if originally constructed with initial defines
+			if (InitialDefineCount)
+			{
+				*this = *GInitialDefines;
+			}
+		}
+	}
+
+	void Merge(const FShaderCompilerDefinitions& Other)
+	{
+		for (FConstIterator OtherIt(Other); OtherIt; ++OtherIt)
+		{
+			int32 OtherIndex = OtherIt.GetIndex();
+			switch (Other.ValueTypes[OtherIndex])
+			{
+			case EShaderCompilerDefineVariant::Integer:
+				SetDefine(OtherIt.Key(), Other.Pairs[OtherIndex].ValueInteger);
+				break;
+			case EShaderCompilerDefineVariant::Unsigned:
+				SetDefine(OtherIt.Key(), Other.Pairs[OtherIndex].ValueUnsigned);
+				break;
+			case EShaderCompilerDefineVariant::Float:
+				SetDefine(OtherIt.Key(), Other.Pairs[OtherIndex].ValueFloat);
+				break;
+			case EShaderCompilerDefineVariant::String:
+				SetDefine(OtherIt.Key(), OtherIt.Value());
+				break;
+			}
+		}
+	}
+
+	FShaderCompilerDefinitions& operator=(const FShaderCompilerDefinitions& Other)
+	{
+		KeyHashTable = Other.KeyHashTable;
+		Pairs = Other.Pairs;
+		ValueTypes = Other.ValueTypes;
+		StringValues = Other.StringValues;
+		InitialDefineCount = Other.InitialDefineCount;
+		ValueCount = Other.ValueCount;
+		return *this;
+	}
+
+	friend FArchive& operator<<(FArchive& Ar, FShaderCompilerDefinitions& Defs)
+	{
+		if (Ar.IsSaving())
+		{
+			// Only write set values in the map
+			Ar << Defs.ValueCount;
+			for (FShaderCompilerDefinitions::FConstIterator DefineIt(Defs); DefineIt; ++DefineIt)
+			{
+				int32 Index = DefineIt.GetIndex();
+				Ar << Defs.Pairs[Index];
+				Ar << Defs.ValueTypes[Index];
+			}
+			Ar << Defs.StringValues;
+		}
+		else if (Ar.IsLoading())
+		{
+			Ar << Defs.ValueCount;
+			for (uint32 ValueIndex = 0; ValueIndex < Defs.ValueCount; ValueIndex++)
+			{
+				FShaderCompilerDefinitions::FPairType Pair;
+				Ar << Pair;
+				int32 Index = Defs.FindOrAddMapIndex(Pair.Key);
+				Defs.Pairs[Index] = Pair;
+				Ar << Defs.ValueTypes[Index];
+			}
+			Ar << Defs.StringValues;
+		}
+		return Ar;
+	}
+
+	class FConstIterator
+	{
+	public:
+		FORCEINLINE FConstIterator(const FShaderCompilerDefinitions& InDefines)
+			: Defines(InDefines), Index(-1)
+		{
+			// NULL terminate these strings to start
+			KeyStringBuffer[0] = 0;
+			ValueStringBuffer[0] = 0;
+
+			// Index set to -1 above, advance to first valid element
+			++(*this);
+		}
+
+		/** conversion to "bool" returning true if the iterator is valid. */
+		FORCEINLINE explicit operator bool() const
+		{
+			return Index < Defines.Pairs.Num();
+		}
+		/** inverse of the "bool" operator */
+		FORCEINLINE bool operator !() const
+		{
+			return !(bool)*this;
+		}
+
+		FORCEINLINE FConstIterator& operator++()
+		{
+			Index++;
+
+			// Skip over None values in the map
+			int32 PairNum = Defines.Pairs.Num();
+			while (Index < PairNum && Defines.ValueTypes[Index] == EShaderCompilerDefineVariant::None)
+			{
+				Index++;
+			}
+			return *this;
+		}
+
+		// Note that the output of Key() is transient, only valid until the iterator is incremented!
+		FORCEINLINE const TCHAR* Key()
+		{
+			Defines.Pairs[Index].Key.ToString(KeyStringBuffer);
+			return KeyStringBuffer;
+		}
+
+		// Note that the output of Value() is transient, only valid until the iterator is incremented!
+		FORCEINLINE const TCHAR* Value()
+		{
+			const TCHAR* Result;
+
+			if (Defines.ValueTypes[Index] == EShaderCompilerDefineVariant::Integer)
+			{
+				int32 ValueInteger = Defines.Pairs[Index].ValueInteger;
+				if (ValueInteger >= 0 && ValueInteger <= 9)
+				{
+					ValueStringBuffer[0] = (TCHAR)ValueInteger + '0';
+					ValueStringBuffer[1] = 0;
+				}
+				else
+				{
+					FCString::Sprintf(ValueStringBuffer, TEXT("%d"), ValueInteger);
+				}
+				Result = ValueStringBuffer;
+			}
+			else if (Defines.ValueTypes[Index] == EShaderCompilerDefineVariant::Unsigned)
+			{
+				uint32 ValueUnsigned = Defines.Pairs[Index].ValueUnsigned;
+				if (ValueUnsigned >= 0 && ValueUnsigned <= 9)
+				{
+					ValueStringBuffer[0] = (TCHAR)ValueUnsigned + '0';
+					ValueStringBuffer[1] = 0;
+				}
+				else
+				{
+					FCString::Sprintf(ValueStringBuffer, TEXT("%u"), ValueUnsigned);
+				}
+				Result = ValueStringBuffer;
+			}
+			else if (Defines.ValueTypes[Index] == EShaderCompilerDefineVariant::Float)
+			{
+				// Make sure the printed value perfectly matches the given number
+				FCString::Sprintf(ValueStringBuffer, TEXT("%#.9gf"), Defines.Pairs[Index].ValueFloat);
+				Result = ValueStringBuffer;
+			}
+			else
+			{
+				check(Defines.ValueTypes[Index] == EShaderCompilerDefineVariant::String);
+				Result = *Defines.StringValues[Defines.Pairs[Index].ValueInteger];
+			}
+
+			return Result;
+		}
+
+		FORCEINLINE int32 GetIndex() const
+		{
+			return Index;
+		}
+
+	private:
+		const FShaderCompilerDefinitions& Defines;
+		int32 Index;
+		TCHAR KeyStringBuffer[FName::StringBufferSize];
+		TCHAR ValueStringBuffer[32];
+	};
+
+
+private:
+	int32 FindMapIndex(FName Key, uint32 KeyHash) const
+	{
+		for (uint32 KeyIndex = KeyHashTable.First(KeyHash); KeyHashTable.IsValid(KeyIndex); KeyIndex = KeyHashTable.Next(KeyIndex))
+		{
+			if (Pairs[KeyIndex].Key == Key)
+			{
+				return KeyIndex;
+			}
+		}
+		return INDEX_NONE;
+	}
+
+	int32 FindOrAddMapIndex(FName Name)
+	{
+		uint32 KeyHash = GetTypeHash(Name);
+		int32 KeyIndex = FindMapIndex(Name, KeyHash);
+		if (KeyIndex == INDEX_NONE)
+		{
+			KeyIndex = Pairs.Add(FPairType({ Name, {0} }));
+			ValueTypes.Add(EShaderCompilerDefineVariant::None);
+			KeyHashTable.Add(KeyHash, KeyIndex);
+		}
+		return KeyIndex;
+	}
+
+	int32 FindOrAddMapIndex(FShaderCompilerDefineNameCache& NameCache)
+	{
+		// Check if this is an initial define, meaning it has a fixed map index.  If Name.MapIndex is INDEX_NONE
+		// (MapIndex not initialized yet) or InitialDefineCount is zero (FShaderCompilerDefinitions constructed without
+		// initial state), this condition will be false, and continue to the rest of the function.
+		if ((uint32)NameCache.MapIndex < InitialDefineCount)
+		{
+			return NameCache.MapIndex;
+		}
+
+		uint32 KeyHash = GetTypeHash(NameCache.Name);
+		int32 KeyIndex = FindMapIndex(NameCache.Name, KeyHash);
+		if (KeyIndex == INDEX_NONE)
+		{
+			KeyIndex = Pairs.Add(FPairType({ NameCache.Name, {0} }));
+			ValueTypes.Add(EShaderCompilerDefineVariant::None);
+			KeyHashTable.Add(KeyHash, KeyIndex);
+		}
+
+		// Initialize MapIndex if necessary.
+		if (InitialDefineCount && NameCache.MapIndex == INDEX_NONE)
+		{
+			NameCache.MapIndex = KeyIndex;
+		}
+		return KeyIndex;
+	}
+
+	void InternalSetValue(int32 Index, const TCHAR* Value)
+	{
+		if (Value[0] >= '0' && Value[0] <= '9' && Value[1] == 0)
+		{
+			// If the string is a single digit integer, treat it as an integer.  Most usages of string define values are cases where
+			// clients inadvertently pass bools as text, like TEXT("0") or TEXT("1"), and we can trivially optimize that.
+			SetValueType(Index, EShaderCompilerDefineVariant::Integer);
+			Pairs[Index].ValueInteger = Value[0] - '0';
+		}
+		else if (ValueTypes[Index] == EShaderCompilerDefineVariant::String)
+		{
+			// Value was already a string variant, overwrite it with new value, rather than allocating new string
+			StringValues[Pairs[Index].ValueInteger] = Value;
+		}
+		else
+		{
+			// Set to string variant type, allocate and fill in the string
+			SetValueType(Index, EShaderCompilerDefineVariant::String);
+			Pairs[Index].ValueInteger = StringValues.Add(Value);
+		}
+	}
+
+	void InternalSetValue(int32 Index, bool Value)
+	{
+		SetValueType(Index, EShaderCompilerDefineVariant::Integer);
+		Pairs[Index].ValueInteger = Value ? 1 : 0;
+	}
+
+	void InternalSetValue(int32 Index, int32 Value)
+	{
+		SetValueType(Index, EShaderCompilerDefineVariant::Integer);
+		Pairs[Index].ValueInteger = Value;
+	}
+
+	void InternalSetValue(int32 Index, uint32 Value)
+	{
+		SetValueType(Index, EShaderCompilerDefineVariant::Unsigned);
+		Pairs[Index].ValueUnsigned = Value;
+	}
+
+	void InternalSetValue(int32 Index, float Value)
+	{
+		SetValueType(Index, EShaderCompilerDefineVariant::Float);
+		Pairs[Index].ValueFloat = Value;
+	}
+
+	FORCEINLINE void SetValueType(int32 Index, EShaderCompilerDefineVariant InValueType)
+	{
+		EShaderCompilerDefineVariant& ValueType = ValueTypes[Index];
+		if (ValueType == EShaderCompilerDefineVariant::None)
+		{
+			ValueCount++;
+		}
+		ValueType = InValueType;
+	}
+
+	/** Called from FShaderInitialDefinesInitializer */
+	RENDERCORE_API static void InitializeInitialDefines(const FShaderCompilerDefinitions& InDefines);
+
+	struct FPairType
+	{
+		FName Key;
+		union
+		{
+			int32 ValueInteger;
+			uint32 ValueUnsigned;
+			float ValueFloat;
+		};
+
+		FORCEINLINE friend FArchive& operator<<(FArchive& Ar, FPairType& Pair)
+		{
+			// Avoid dynamic allocation when serializing keys.  We need to serialize the FName as a string, because the index in the FName is
+			// non-deterministic (serialization generates hash values that need to be deterministic) and process specific (data needs to be
+			// serialized to worker processes).  Compiler defines must be ANSI text, so serialize as such to save memory.
+			TCHAR KeyBuffer[FName::StringBufferSize];
+			ANSICHAR KeyBufferAnsi[FName::StringBufferSize];
+			uint32 StringLength;
+
+			if (Ar.IsLoading())
+			{
+				Ar << StringLength;
+				check(StringLength < FName::StringBufferSize);
+				Ar.Serialize(KeyBufferAnsi, StringLength + 1);
+				Pair.Key = FName(KeyBufferAnsi);
+			}
+			else
+			{
+				StringLength = Pair.Key.ToString(KeyBuffer, FName::StringBufferSize);
+				for (uint32 CharIndex = 0; CharIndex < StringLength + 1; CharIndex++)
+				{
+					KeyBufferAnsi[CharIndex] = (ANSICHAR)KeyBuffer[CharIndex];
+				}
+				Ar << StringLength;
+				Ar.Serialize(KeyBufferAnsi, StringLength + 1);
+			}
+			Ar << Pair.ValueInteger;
+			return Ar;
+		}
+	};
+	FHashTable KeyHashTable;
+	TArray<FPairType> Pairs;
+	TArray<EShaderCompilerDefineVariant> ValueTypes;
+	TArray<FString> StringValues;
+	uint32 InitialDefineCount;			// Number of items that came from GInitialDefines
+	uint32 ValueCount;					// Number of valid values (ValueType != EShaderCompilerDefineVariant::None)
+
+	RENDERCORE_API static FShaderCompilerDefinitions* GInitialDefines;
+
+	friend struct FShaderInitialDefinesInitializer;
+};
+
+PRAGMA_ENABLE_DEPRECATION_WARNINGS
diff --git a/Engine/Source/Runtime/RenderCore/Private/ShaderCore.cpp b/Engine/Source/Runtime/RenderCore/Private/ShaderCore.cpp
--- a/Engine/Source/Runtime/RenderCore/Private/ShaderCore.cpp
+++ b/Engine/Source/Runtime/RenderCore/Private/ShaderCore.cpp
@@ -37,6 +37,7 @@
 #include "Misc/CoreMisc.h"
 #include "Interfaces/ITargetPlatform.h"
 #include "Interfaces/ITargetPlatformManagerModule.h"
+#include "Runtime/RenderCore/Internal/ShaderCompilerDefinitions.h"
 #if WITH_EDITOR
 #include "Misc/ConfigCacheIni.h"
 #include "Misc/StringBuilder.h"
@@ -116,6 +117,8 @@ DEFINE_STAT(STAT_Shaders_ShaderPreloadMemory);
 DEFINE_STAT(STAT_Shaders_NumShadersRegistered);
 DEFINE_STAT(STAT_Shaders_NumShadersDuplicated);
 
+PRAGMA_DISABLE_DEPRECATION_WARNINGS		// FShaderCompilerDefinitions will be made internal in the future, marked deprecated until then
+
 /**
  * Singleton initial set of defines added when constructing a defines structure with bIncludeInitialDefines==true.  The
  * advantage of using preset defines is that the index of the initial define can be cached in the FShaderCompilerDefineNameCache
@@ -148,290 +151,7 @@ void FShaderCompilerDefinitions::InitializeInitialDefines(const FShaderCompilerD
 	GInitialDefines->InitialDefineCount = InDefines.Pairs.Num();
 }
 
-int32 FShaderCompilerDefinitions::FindMapIndex(FName Key, uint32 KeyHash) const
-{
-	for (uint32 KeyIndex = KeyHashTable.First(KeyHash); KeyHashTable.IsValid(KeyIndex); KeyIndex = KeyHashTable.Next(KeyIndex))
-	{
-		if (Pairs[KeyIndex].Key == Key)
-		{
-			return KeyIndex;
-		}
-	}
-	return INDEX_NONE;
-}
-
-int32 FShaderCompilerDefinitions::FindOrAddMapIndex(FName Name)
-{
-	uint32 KeyHash = GetTypeHash(Name);
-	int32 KeyIndex = FindMapIndex(Name, KeyHash);
-	if (KeyIndex == INDEX_NONE)
-	{
-		KeyIndex = Pairs.Add(FPairType({ Name, {0} }));
-		ValueTypes.Add(EShaderCompilerDefineVariant::None);
-		KeyHashTable.Add(KeyHash, KeyIndex);
-	}
-	return KeyIndex;
-}
-
-int32 FShaderCompilerDefinitions::FindOrAddMapIndex(FShaderCompilerDefineNameCache& NameCache)
-{
-	// Check if this is an initial define, meaning it has a fixed map index.  If Name.MapIndex is INDEX_NONE
-	// (MapIndex not initialized yet) or InitialDefineCount is zero (FShaderCompilerDefinitions constructed without
-	// initial state), this condition will be false, and continue to the rest of the function.
-	if ((uint32)NameCache.MapIndex < InitialDefineCount)
-	{
-		return NameCache.MapIndex;
-	}
-
-	uint32 KeyHash = GetTypeHash(NameCache.Name);
-	int32 KeyIndex = FindMapIndex(NameCache.Name, KeyHash);
-	if (KeyIndex == INDEX_NONE)
-	{
-		KeyIndex = Pairs.Add(FPairType({ NameCache.Name, {0} }));
-		ValueTypes.Add(EShaderCompilerDefineVariant::None);
-		KeyHashTable.Add(KeyHash, KeyIndex);
-	}
-
-	// Initialize MapIndex if necessary.
-	if (InitialDefineCount && NameCache.MapIndex == INDEX_NONE)
-	{
-		NameCache.MapIndex = KeyIndex;
-	}
-	return KeyIndex;
-}
-
-void FShaderCompilerDefinitions::InternalSetValue(int32 Index, const TCHAR* Value)
-{
-	if (Value[0] >= '0' && Value[0] <= '9' && Value[1] == 0)
-	{
-		// If the string is a single digit integer, treat it as an integer.  Most usages of string define values are cases where
-		// clients inadvertently pass bools as text, like TEXT("0") or TEXT("1"), and we can trivially optimize that.
-		SetValueType(Index, EShaderCompilerDefineVariant::Integer);
-		Pairs[Index].ValueInteger = Value[0] - '0';
-	}
-	else if (ValueTypes[Index] == EShaderCompilerDefineVariant::String)
-	{
-		// Value was already a string variant, overwrite it with new value, rather than allocating new string
-		StringValues[Pairs[Index].ValueInteger] = Value;
-	}
-	else
-	{
-		// Set to string variant type, allocate and fill in the string
-		SetValueType(Index, EShaderCompilerDefineVariant::String);
-		Pairs[Index].ValueInteger = StringValues.Add(Value);
-	}
-}
-
-int32 FShaderCompilerDefinitions::GetIntegerValue(FName Name) const
-{
-	int32 Result = 0;
-	int32 KeyIndex = FindMapIndex(Name, GetTypeHash(Name));
-	if (KeyIndex != INDEX_NONE)
-	{
-		// For None, Integer, or Unsigned, return ValueInteger (None will have a default of zero)
-		if (ValueTypes[KeyIndex] <= EShaderCompilerDefineVariant::Unsigned)
-		{
-			Result = Pairs[KeyIndex].ValueInteger;
-		}
-		else if (ValueTypes[KeyIndex] == EShaderCompilerDefineVariant::String)
-		{
-			Result = FCString::Atoi(*StringValues[Pairs[KeyIndex].ValueInteger]);
-		}
-		else
-		{
-			// EShaderCompilerDefineVariant::Float
-			Result = (int32)Pairs[KeyIndex].ValueFloat;
-		}
-	}
-	return Result;
-}
-
-int32 FShaderCompilerDefinitions::GetIntegerValue(FShaderCompilerDefineNameCache& NameCache, int32 ResultIfNotFound) const
-{
-	int32 Result = ResultIfNotFound;
-	int32 KeyIndex;
-
-	// Check if this is an initial define, meaning it has a fixed map index.  If Name.MapIndex is INDEX_NONE
-	// (MapIndex not initialized yet) or InitialDefineCount is zero (FShaderCompilerDefinitions constructed without
-	// initial state), this condition will be false, and continue to the rest of the function.
-	if ((uint32)NameCache.MapIndex < InitialDefineCount)
-	{
-		KeyIndex = NameCache.MapIndex;
-	}
-	else
-	{
-		KeyIndex = FindMapIndex(NameCache.Name, GetTypeHash(NameCache.Name));
-	}
-
-	if (KeyIndex != INDEX_NONE)
-	{
-		// Initialize MapIndex if necessary.  If the define is not an initial define, its index will be greater than
-		// or equal to InitialDefineCount, indicating that it doesn't have a fixed map index.
-		if (InitialDefineCount && NameCache.MapIndex == INDEX_NONE)
-		{
-			NameCache.MapIndex = KeyIndex;
-		}
-
-		EShaderCompilerDefineVariant ValueType = ValueTypes[KeyIndex];
-		if (ValueType == EShaderCompilerDefineVariant::Integer ||
-			ValueType == EShaderCompilerDefineVariant::Unsigned)
-		{
-			Result = Pairs[KeyIndex].ValueInteger;
-		}
-		else if (ValueType == EShaderCompilerDefineVariant::String)
-		{
-			Result = FCString::Atoi(*StringValues[Pairs[KeyIndex].ValueInteger]);
-		}
-		else if (ValueType == EShaderCompilerDefineVariant::Float)
-		{
-			// EShaderCompilerDefineVariant::Float
-			Result = (int32)Pairs[KeyIndex].ValueFloat;
-		}
-	}
-	return Result;
-}
-
-RENDERCORE_API void FShaderCompilerDefinitions::Empty()
-{
-	if (ValueCount)
-	{
-		KeyHashTable.Free();
-		Pairs.Empty();
-		ValueTypes.Empty();
-		StringValues.Empty();
-		ValueCount = 0;
-
-		// Re-copy initial defines if originally constructed with initial defines
-		if (InitialDefineCount)
-		{
-			*this = *GInitialDefines;
-		}
-	}
-}
-
-void FShaderCompilerDefinitions::Merge(const FShaderCompilerDefinitions& Other)
-{
-	for (FConstIterator OtherIt(Other); OtherIt; ++OtherIt)
-	{
-		int32 OtherIndex = OtherIt.GetIndex();
-		switch (Other.ValueTypes[OtherIndex])
-		{
-		case EShaderCompilerDefineVariant::Integer:
-			SetDefine(OtherIt.Key(), Other.Pairs[OtherIndex].ValueInteger);
-			break;
-		case EShaderCompilerDefineVariant::Unsigned:
-			SetDefine(OtherIt.Key(), Other.Pairs[OtherIndex].ValueUnsigned);
-			break;
-		case EShaderCompilerDefineVariant::Float:
-			SetDefine(OtherIt.Key(), Other.Pairs[OtherIndex].ValueFloat);
-			break;
-		case EShaderCompilerDefineVariant::String:
-			SetDefine(OtherIt.Key(), OtherIt.Value());
-			break;
-		}
-	}
-}
-
-RENDERCORE_API FShaderCompilerDefinitions& FShaderCompilerDefinitions::operator=(const FShaderCompilerDefinitions& Other)
-{
-	KeyHashTable = Other.KeyHashTable;
-	Pairs = Other.Pairs;
-	ValueTypes = Other.ValueTypes;
-	StringValues = Other.StringValues;
-	InitialDefineCount = Other.InitialDefineCount;
-	ValueCount = Other.ValueCount;
-	return *this;
-}
-
-FArchive& operator<<(FArchive& Ar, FShaderCompilerDefinitions& Defs)
-{
-	if (Ar.IsSaving())
-	{
-		// Only write set values in the map
-		Ar << Defs.ValueCount;
-		for (FShaderCompilerDefinitions::FConstIterator DefineIt(Defs); DefineIt; ++DefineIt)
-		{
-			int32 Index = DefineIt.GetIndex();
-			Ar << Defs.Pairs[Index];
-			Ar << Defs.ValueTypes[Index];
-		}
-		Ar << Defs.StringValues;
-	}
-	else if (Ar.IsLoading())
-	{
-		Ar << Defs.ValueCount;
-		for (uint32 ValueIndex = 0; ValueIndex < Defs.ValueCount; ValueIndex++)
-		{
-			FShaderCompilerDefinitions::FPairType Pair;
-			Ar << Pair;
-			int32 Index = Defs.FindOrAddMapIndex(Pair.Key);
-			Defs.Pairs[Index] = Pair;
-			Ar << Defs.ValueTypes[Index];
-		}
-		Ar << Defs.StringValues;
-	}
-	return Ar;
-}
-
-RENDERCORE_API FShaderCompilerDefinitions::FConstIterator::FConstIterator(const FShaderCompilerDefinitions& InDefines)
-	: Defines(InDefines), Index(-1)
-{
-	// NULL terminate these strings to start
-	KeyStringBuffer[0] = 0;
-	ValueStringBuffer[0] = 0;
-
-	// Index set to -1 above, advance to first valid element
-	++(*this);
-}
-
-RENDERCORE_API const TCHAR* FShaderCompilerDefinitions::FConstIterator::Value()
-{
-	const TCHAR* Result;
-
-	if (Defines.ValueTypes[Index] == EShaderCompilerDefineVariant::Integer)
-	{
-		int32 ValueInteger = Defines.Pairs[Index].ValueInteger;
-		if (ValueInteger >= 0 && ValueInteger <= 9)
-		{
-			ValueStringBuffer[0] = (TCHAR)ValueInteger + '0';
-			ValueStringBuffer[1] = 0;
-		}
-		else
-		{
-			FCString::Sprintf(ValueStringBuffer, TEXT("%d"), ValueInteger);
-		}
-		Result = ValueStringBuffer;
-	}
-	else if (Defines.ValueTypes[Index] == EShaderCompilerDefineVariant::Unsigned)
-	{
-		uint32 ValueUnsigned = Defines.Pairs[Index].ValueUnsigned;
-		if (ValueUnsigned >= 0 && ValueUnsigned <= 9)
-		{
-			ValueStringBuffer[0] = (TCHAR)ValueUnsigned + '0';
-			ValueStringBuffer[1] = 0;
-		}
-		else
-		{
-			FCString::Sprintf(ValueStringBuffer, TEXT("%u"), ValueUnsigned);
-		}
-		Result = ValueStringBuffer;
-	}
-	else if (Defines.ValueTypes[Index] == EShaderCompilerDefineVariant::Float)
-	{
-		// Make sure the printed value perfectly matches the given number
-		FCString::Sprintf(ValueStringBuffer, TEXT("%#.9gf"), Defines.Pairs[Index].ValueFloat);
-		Result = ValueStringBuffer;
-	}
-	else
-	{
-		check(Defines.ValueTypes[Index] == EShaderCompilerDefineVariant::String);
-		Result = *Defines.StringValues[Defines.Pairs[Index].ValueInteger];
-	}
-
-	return Result;
-}
-
-
+PRAGMA_ENABLE_DEPRECATION_WARNINGS
 
 
 // Apply lock striping as we're mostly reader lock bound.
@@ -1080,6 +800,129 @@ void FShaderResourceTableMap::FixupOnLoad(const TMap<FString, FUniformBufferEntr
 	}
 }
 
+FShaderCompilerEnvironment::FShaderCompilerEnvironment()
+{
+	// Enable initial defines in FShaderCompilerEnvironment to improve performance (helpful here, but not for defines declared in various shader compiler backends).
+	const bool bIncludeInitialDefines = true;
+
+	PRAGMA_DISABLE_DEPRECATION_WARNINGS		// FShaderCompilerDefinitions will be made internal in the future, marked deprecated until then
+	Definitions = MakePimpl<FShaderCompilerDefinitions, EPimplPtrMode::DeepCopy>(bIncludeInitialDefines);
+	PRAGMA_ENABLE_DEPRECATION_WARNINGS
+
+	// Presize to reduce re-hashing while building shader jobs
+	IncludeVirtualPathToContentsMap.Empty(15);
+}
+
+/** Initialization constructor. */
+PRAGMA_DISABLE_DEPRECATION_WARNINGS		// FShaderCompilerDefinitions will be made internal in the future, marked deprecated until then
+FShaderCompilerEnvironment::FShaderCompilerEnvironment(const FShaderCompilerDefinitions& InDefinitions)
+{
+	Definitions = MakePimpl<FShaderCompilerDefinitions, EPimplPtrMode::DeepCopy>(InDefinitions);
+}
+PRAGMA_ENABLE_DEPRECATION_WARNINGS
+
+void FShaderCompilerEnvironment::Merge(const FShaderCompilerEnvironment& Other)
+{
+	// Merge the include maps
+	// Merge the values of any existing keys
+	for (TMap<FString, FString>::TConstIterator It(Other.IncludeVirtualPathToContentsMap); It; ++It)
+	{
+		FString* ExistingContents = IncludeVirtualPathToContentsMap.Find(It.Key());
+
+		if (ExistingContents)
+		{
+			ExistingContents->Append(It.Value());
+		}
+		else
+		{
+			IncludeVirtualPathToContentsMap.Add(It.Key(), It.Value());
+		}
+	}
+
+	check(Other.IncludeVirtualPathToExternalContentsMap.Num() == 0);
+
+	CompilerFlags.Append(Other.CompilerFlags);
+	ResourceTableMap.Append(Other.ResourceTableMap);
+	UniformBufferMap.Append(Other.UniformBufferMap);
+	PRAGMA_DISABLE_DEPRECATION_WARNINGS		// FShaderCompilerDefinitions will be made internal in the future, marked deprecated until then
+	Definitions->Merge(*Other.Definitions);
+	PRAGMA_ENABLE_DEPRECATION_WARNINGS
+	CompileArgs.Append(Other.CompileArgs);
+	RenderTargetOutputFormatsMap.Append(Other.RenderTargetOutputFormatsMap);
+	FullPrecisionInPS |= Other.FullPrecisionInPS;
+}
+
+PRAGMA_DISABLE_DEPRECATION_WARNINGS		// FShaderCompilerDefinitions will be made internal in the future, marked deprecated until then
+
+// Pass through functions to definitions
+void FShaderCompilerEnvironment::SetDefine(const TCHAR* Name, const TCHAR* Value)	{ Definitions->SetDefine(Name, Value); }
+void FShaderCompilerEnvironment::SetDefine(const TCHAR* Name, const FString& Value)	{ Definitions->SetDefine(Name, Value); }
+void FShaderCompilerEnvironment::SetDefine(const TCHAR* Name, uint32 Value)			{ Definitions->SetDefine(Name, Value); }
+void FShaderCompilerEnvironment::SetDefine(const TCHAR* Name, int32 Value)			{ Definitions->SetDefine(Name, Value); }
+void FShaderCompilerEnvironment::SetDefine(const TCHAR* Name, bool Value)			{ Definitions->SetDefine(Name, Value); }
+void FShaderCompilerEnvironment::SetDefine(const TCHAR* Name, float Value)			{ Definitions->SetDefine(Name, Value); }
+
+void FShaderCompilerEnvironment::SetDefine(FName Name, const TCHAR* Value)		{ Definitions->SetDefine(Name, Value); }
+void FShaderCompilerEnvironment::SetDefine(FName Name, const FString& Value)	{ Definitions->SetDefine(Name, Value); }
+void FShaderCompilerEnvironment::SetDefine(FName Name, uint32 Value)			{ Definitions->SetDefine(Name, Value); }
+void FShaderCompilerEnvironment::SetDefine(FName Name, int32 Value)				{ Definitions->SetDefine(Name, Value); }
+void FShaderCompilerEnvironment::SetDefine(FName Name, bool Value)				{ Definitions->SetDefine(Name, Value); }
+void FShaderCompilerEnvironment::SetDefine(FName Name, float Value)				{ Definitions->SetDefine(Name, Value); }
+
+void FShaderCompilerEnvironment::SetDefine(FShaderCompilerDefineNameCache& Name, const TCHAR* Value)	{ Definitions->SetDefine(Name, Value); }
+void FShaderCompilerEnvironment::SetDefine(FShaderCompilerDefineNameCache& Name, const FString& Value)	{ Definitions->SetDefine(Name, Value); }
+void FShaderCompilerEnvironment::SetDefine(FShaderCompilerDefineNameCache& Name, uint32 Value)			{ Definitions->SetDefine(Name, Value); }
+void FShaderCompilerEnvironment::SetDefine(FShaderCompilerDefineNameCache& Name, int32 Value)			{ Definitions->SetDefine(Name, Value); }
+void FShaderCompilerEnvironment::SetDefine(FShaderCompilerDefineNameCache& Name, bool Value)			{ Definitions->SetDefine(Name, Value); }
+void FShaderCompilerEnvironment::SetDefine(FShaderCompilerDefineNameCache& Name, float Value)			{ Definitions->SetDefine(Name, Value); }
+
+int32 FShaderCompilerEnvironment::GetIntegerValue(FName Name) const
+{
+	return Definitions->GetIntegerValue(Name);
+}
+
+int32 FShaderCompilerEnvironment::GetIntegerValue(FShaderCompilerDefineNameCache& NameCache, int32 ResultIfNotFound) const
+{
+	return Definitions->GetIntegerValue(NameCache, ResultIfNotFound);
+}
+
+bool FShaderCompilerEnvironment::ContainsDefinition(FName Name) const
+{
+	return Definitions->Contains(Name);
+}
+
+PRAGMA_ENABLE_DEPRECATION_WARNINGS
+
+/** This "core" serialization is also used for the hashing the compiler job (where files are handled differently). Should stay in sync with the ShaderCompileWorker. */
+void FShaderCompilerEnvironment::SerializeEverythingButFiles(FArchive& Ar)
+{
+	Ar << *Definitions;
+	Ar << CompileArgs;
+	Ar << CompilerFlags;
+	Ar << RenderTargetOutputFormatsMap;
+	Ar << ResourceTableMap.Resources;
+	Ar << UniformBufferMap;
+	Ar << FullPrecisionInPS;
+	if (Ar.IsLoading())
+	{
+		ResourceTableMap.FixupOnLoad(UniformBufferMap);
+	}
+}
+
+// Serializes the portions of the environment that are used as input to the backend compilation process (i.e. after all preprocessing)
+void FShaderCompilerEnvironment::SerializeCompilationDependencies(FArchive& Ar)
+{
+	Ar << CompileArgs;
+	Ar << CompilerFlags;
+	Ar << ResourceTableMap.Resources;
+	Ar << UniformBufferMap;
+	Ar << FullPrecisionInPS;
+	if (Ar.IsLoading())
+	{
+		ResourceTableMap.FixupOnLoad(UniformBufferMap);
+	}
+}
+
 void FShaderCompilerOutput::GenerateOutputHash()
 {
 	FSHA1 HashState;
diff --git a/Engine/Source/Runtime/RenderCore/Public/ShaderCore.h b/Engine/Source/Runtime/RenderCore/Public/ShaderCore.h
--- a/Engine/Source/Runtime/RenderCore/Public/ShaderCore.h
+++ b/Engine/Source/Runtime/RenderCore/Public/ShaderCore.h
@@ -33,6 +33,7 @@
 #include "Stats/Stats.h"
 #include "Stats/Stats2.h"
 #include "Templates/Function.h"
+#include "Templates/PimplPtr.h"
 #include "Templates/RefCounting.h"
 #include "Templates/SharedPointer.h"
 #include "Templates/UnrealTemplate.h"
@@ -44,6 +45,7 @@ class Error;
 class FMemoryImageWriter;
 class FMemoryUnfreezeContent;
 class FPointerTableBase;
+class FShaderCompilerDefinitions;
 class FShaderCompileUtilities;
 class FShaderPreprocessorUtilities;
 class FSHA1;
@@ -317,212 +319,8 @@ private:
 	friend class FShaderCompilerDefinitions;
 };
 
-enum class EShaderCompilerDefineVariant : uint8
-{
-	None = 0,
-	Integer,
-	Unsigned,
-	Float,
-	String
-};
-
-/** Container for shader compiler definitions. */
-class FShaderCompilerDefinitions
-{
-public:
-	RENDERCORE_API FShaderCompilerDefinitions(bool bIncludeInitialDefines = false);
-	
-	RENDERCORE_API FShaderCompilerDefinitions(const FShaderCompilerDefinitions&);
-
-	/** Value types supported:   bool, int32, uint32, float, const TCHAR*, FString& */
-	template <typename ValueType>
-	void SetDefine(FName Name, ValueType Value)									{ InternalSetValue(FindOrAddMapIndex(Name), Value); }
-	void SetDefine(FName Name, const TCHAR* Value)								{ InternalSetValue(FindOrAddMapIndex(Name), Value); }
-	void SetDefine(FName Name, const FString& Value)							{ InternalSetValue(FindOrAddMapIndex(Name), *Value); }
-
-	template <typename ValueType>
-	void SetDefine(FShaderCompilerDefineNameCache& Name, ValueType Value)		{ InternalSetValue(FindOrAddMapIndex(Name), Value); }
-	void SetDefine(FShaderCompilerDefineNameCache& Name, const TCHAR* Value)	{ InternalSetValue(FindOrAddMapIndex(Name), Value); }
-	void SetDefine(FShaderCompilerDefineNameCache& Name, const FString& Value)	{ InternalSetValue(FindOrAddMapIndex(Name), *Value); }
-
-	RENDERCORE_API int32 GetIntegerValue(FName Name) const;
-	RENDERCORE_API int32 GetIntegerValue(FShaderCompilerDefineNameCache& NameCache, int32 ResultIfNotFound) const;
-
-	bool Contains(FName Name) const
-	{
-		int32 KeyIndex = FindMapIndex(Name, GetTypeHash(Name));
-		return KeyIndex != INDEX_NONE && ValueTypes[KeyIndex] != EShaderCompilerDefineVariant::None;
-	}
-
-	FORCEINLINE int32 Num() const
-	{
-		return ValueCount;
-	}
-
-	RENDERCORE_API void Empty();
-	RENDERCORE_API void Merge(const FShaderCompilerDefinitions& Other);
-
-	RENDERCORE_API FShaderCompilerDefinitions& operator=(const FShaderCompilerDefinitions& Other);
-
-	RENDERCORE_API friend FArchive& operator<<(FArchive& Ar,FShaderCompilerDefinitions& Defs);
-
-	class FConstIterator
-	{
-	public:
-		RENDERCORE_API FConstIterator(const FShaderCompilerDefinitions& InDefines);
-
-		/** conversion to "bool" returning true if the iterator is valid. */
-		FORCEINLINE explicit operator bool() const
-		{
-			return Index < Defines.Pairs.Num();
-		}
-		/** inverse of the "bool" operator */
-		FORCEINLINE bool operator !() const
-		{
-			return !(bool)*this;
-		}
-
-		FORCEINLINE FConstIterator& operator++()
-		{
-			Index++;
-
-			// Skip over None values in the map
-			int32 PairNum = Defines.Pairs.Num();
-			while (Index < PairNum && Defines.ValueTypes[Index] == EShaderCompilerDefineVariant::None)
-			{
-				Index++;
-			}
-			return *this;
-		}
-
-		// Note that the output of Key() is transient, only valid until the iterator is incremented!
-		FORCEINLINE const TCHAR* Key()
-		{
-			Defines.Pairs[Index].Key.ToString(KeyStringBuffer);
-			return KeyStringBuffer;
-		}
-
-		// Note that the output of Value() is transient, only valid until the iterator is incremented!
-		RENDERCORE_API const TCHAR* Value();
-
-		FORCEINLINE int32 GetIndex() const
-		{
-			return Index;
-		}
-
-	private:
-		const FShaderCompilerDefinitions& Defines;
-		int32 Index;
-		TCHAR KeyStringBuffer[FName::StringBufferSize];
-		TCHAR ValueStringBuffer[32];
-	};
-
-
-private:
-	RENDERCORE_API int32 FindMapIndex(FName Key, uint32 KeyHash) const;
-	RENDERCORE_API int32 FindOrAddMapIndex(FName Name);
-	RENDERCORE_API int32 FindOrAddMapIndex(FShaderCompilerDefineNameCache& NameCache);
-
-	RENDERCORE_API void InternalSetValue(int32 Index, const TCHAR* Value);
-
-	void InternalSetValue(int32 Index, bool Value)
-	{
-		SetValueType(Index, EShaderCompilerDefineVariant::Integer);
-		Pairs[Index].ValueInteger = Value ? 1 : 0;
-	}
-
-	void InternalSetValue(int32 Index, int32 Value)
-	{
-		SetValueType(Index, EShaderCompilerDefineVariant::Integer);
-		Pairs[Index].ValueInteger = Value;
-	}
-
-	void InternalSetValue(int32 Index, uint32 Value)
-	{
-		SetValueType(Index, EShaderCompilerDefineVariant::Unsigned);
-		Pairs[Index].ValueUnsigned = Value;
-	}
-
-	void InternalSetValue(int32 Index, float Value)
-	{
-		SetValueType(Index, EShaderCompilerDefineVariant::Float);
-		Pairs[Index].ValueFloat = Value;
-	}
-
-	FORCEINLINE void SetValueType(int32 Index, EShaderCompilerDefineVariant InValueType)
-	{
-		EShaderCompilerDefineVariant& ValueType = ValueTypes[Index];
-		if (ValueType == EShaderCompilerDefineVariant::None)
-		{
-			ValueCount++;
-		}
-		ValueType = InValueType;
-	}
-
-	/** Called from FShaderInitialDefinesInitializer */
-	RENDERCORE_API static void InitializeInitialDefines(const FShaderCompilerDefinitions& InDefines);
-
-	/** Needed for the deprecated FShaderCompilerEnvironment::GetDefinitions, do not use! */
-	const TMap<FString, FString>& GetDefinitionMap() const
-	{
-		return UnusedStringDefinitions;
-	}
-	friend struct FShaderCompilerEnvironment;
-
-	struct FPairType
-	{
-		FName Key;
-		union
-		{
-			int32 ValueInteger;
-			uint32 ValueUnsigned;
-			float ValueFloat;
-		};
-
-		FORCEINLINE friend FArchive& operator<<(FArchive& Ar, FPairType& Pair)
-		{
-			// Avoid dynamic allocation when serializing keys.  We need to serialize the FName as a string, because the index in the FName is
-			// non-deterministic (serialization generates hash values that need to be deterministic) and process specific (data needs to be
-			// serialized to worker processes).  Compiler defines must be ANSI text, so serialize as such to save memory.
-			TCHAR KeyBuffer[FName::StringBufferSize];
-			ANSICHAR KeyBufferAnsi[FName::StringBufferSize];
-			uint32 StringLength;
-
-			if (Ar.IsLoading())
-			{
-				Ar << StringLength;
-				check(StringLength < FName::StringBufferSize);
-				Ar.Serialize(KeyBufferAnsi, StringLength + 1);
-				Pair.Key = FName(KeyBufferAnsi);
-			}
-			else
-			{
-				StringLength = Pair.Key.ToString(KeyBuffer, FName::StringBufferSize);
-				for (uint32 CharIndex = 0; CharIndex < StringLength + 1; CharIndex++)
-				{
-					KeyBufferAnsi[CharIndex] = (ANSICHAR)KeyBuffer[CharIndex];
-				}
-				Ar << StringLength;
-				Ar.Serialize(KeyBufferAnsi, StringLength + 1);
-			}
-			Ar << Pair.ValueInteger;
-			return Ar;
-		}
-	};
-	FHashTable KeyHashTable;
-	TArray<FPairType> Pairs;
-	TArray<EShaderCompilerDefineVariant> ValueTypes;
-	TArray<FString> StringValues;
-	uint32 InitialDefineCount;			// Number of items that came from GInitialDefines
-	uint32 ValueCount;					// Number of valid values (ValueType != EShaderCompilerDefineVariant::None)
-
-	/** Unused data kept around as something to return from GetDefinitionMap() call, which is used only in deprecated FShaderCompilerEnvironment::GetDefinitions call */
-	TMap<FString,FString> UnusedStringDefinitions;
-
-	RENDERCORE_API static FShaderCompilerDefinitions* GInitialDefines;
-
-	friend struct FShaderInitialDefinesInitializer;
-};
+// Temporarily included here until we can fully deprecate access to FShaderCompilerDefinitions in a future version
+#include "Runtime/RenderCore/Internal/ShaderCompilerDefinitions.h"
 
 struct FShaderResourceTable
 {
@@ -697,18 +495,13 @@ struct FShaderCompilerEnvironment
 	bool FullPrecisionInPS = 0;
 
 	/** Default constructor. */
-	FShaderCompilerEnvironment()
-		: Definitions(true)			/* Enable initial defines -- helps performance here, but not necessary for standalone "added" defines declared elsewhere */
-	{
-		// Presize to reduce re-hashing while building shader jobs
-		IncludeVirtualPathToContentsMap.Empty(15);
-	}
+	RENDERCORE_API FShaderCompilerEnvironment();
 
 	/** Initialization constructor. */
-	explicit FShaderCompilerEnvironment(const FShaderCompilerDefinitions& InDefinitions)
-		: Definitions(InDefinitions)
-	{
-	}
+	PRAGMA_DISABLE_DEPRECATION_WARNINGS
+	UE_DEPRECATED(5.4, "FShaderCompilerDefinitions is being made private in the future, do not use this constructor.")
+	RENDERCORE_API explicit FShaderCompilerEnvironment(const FShaderCompilerDefinitions& InDefinitions);
+	PRAGMA_ENABLE_DEPRECATION_WARNINGS
 
 	// Used as a baseclass, make sure we're not incorrectly destroyed through a baseclass pointer
 	// This will be expensive to destroy anyway, additional vcall overhead should be small
@@ -736,42 +529,38 @@ struct FShaderCompilerEnvironment
 	 * e.g. SET_SHADER_DEFINE(NUM_SAMPLES, 1);
 	 * e.g. SET_SHADER_DEFINE(DOIT, true);
 	 */
-	void SetDefine(const TCHAR* Name, const TCHAR* Value)	{ Definitions.SetDefine(Name, Value); }
-	void SetDefine(const TCHAR* Name, const FString& Value) { Definitions.SetDefine(Name, Value); }
-	void SetDefine(const TCHAR* Name, uint32 Value)			{ Definitions.SetDefine(Name, Value); }
-	void SetDefine(const TCHAR* Name, int32 Value)			{ Definitions.SetDefine(Name, Value); }
-	void SetDefine(const TCHAR* Name, bool Value)			{ Definitions.SetDefine(Name, Value); }
-	void SetDefine(const TCHAR* Name, float Value)			{ Definitions.SetDefine(Name, Value); }
-
-	void SetDefine(FName Name, const TCHAR* Value)			{ Definitions.SetDefine(Name, Value); }
-	void SetDefine(FName Name, const FString& Value)		{ Definitions.SetDefine(Name, Value); }
-	void SetDefine(FName Name, uint32 Value)				{ Definitions.SetDefine(Name, Value); }
-	void SetDefine(FName Name, int32 Value)					{ Definitions.SetDefine(Name, Value); }
-	void SetDefine(FName Name, bool Value)					{ Definitions.SetDefine(Name, Value); }
-	void SetDefine(FName Name, float Value)					{ Definitions.SetDefine(Name, Value); }
-
-	void SetDefine(FShaderCompilerDefineNameCache& Name, const TCHAR* Value)	{ Definitions.SetDefine(Name, Value); }
-	void SetDefine(FShaderCompilerDefineNameCache& Name, const FString& Value)	{ Definitions.SetDefine(Name, Value); }
-	void SetDefine(FShaderCompilerDefineNameCache& Name, uint32 Value)			{ Definitions.SetDefine(Name, Value); }
-	void SetDefine(FShaderCompilerDefineNameCache& Name, int32 Value)			{ Definitions.SetDefine(Name, Value); }
-	void SetDefine(FShaderCompilerDefineNameCache& Name, bool Value)			{ Definitions.SetDefine(Name, Value); }
-	void SetDefine(FShaderCompilerDefineNameCache& Name, float Value)			{ Definitions.SetDefine(Name, Value); }
-
-	FORCEINLINE int32 GetIntegerValue(FName Name) const
-	{
-		return Definitions.GetIntegerValue(Name);
-	}
-	FORCEINLINE int32 GetIntegerValue(FShaderCompilerDefineNameCache& NameCache, int32 ResultIfNotFound = 0) const
-	{
-		return Definitions.GetIntegerValue(NameCache, ResultIfNotFound);
-	}
+	RENDERCORE_API void SetDefine(const TCHAR* Name, const TCHAR* Value);
+	RENDERCORE_API void SetDefine(const TCHAR* Name, const FString& Value);
+	RENDERCORE_API void SetDefine(const TCHAR* Name, uint32 Value);
+	RENDERCORE_API void SetDefine(const TCHAR* Name, int32 Value);
+	RENDERCORE_API void SetDefine(const TCHAR* Name, bool Value);
+	RENDERCORE_API void SetDefine(const TCHAR* Name, float Value);
+
+	RENDERCORE_API void SetDefine(FName Name, const TCHAR* Value);
+	RENDERCORE_API void SetDefine(FName Name, const FString& Value);
+	RENDERCORE_API void SetDefine(FName Name, uint32 Value);
+	RENDERCORE_API void SetDefine(FName Name, int32 Value);
+	RENDERCORE_API void SetDefine(FName Name, bool Value);
+	RENDERCORE_API void SetDefine(FName Name, float Value);
+
+	RENDERCORE_API void SetDefine(FShaderCompilerDefineNameCache& Name, const TCHAR* Value);
+	RENDERCORE_API void SetDefine(FShaderCompilerDefineNameCache& Name, const FString& Value);
+	RENDERCORE_API void SetDefine(FShaderCompilerDefineNameCache& Name, uint32 Value);
+	RENDERCORE_API void SetDefine(FShaderCompilerDefineNameCache& Name, int32 Value);
+	RENDERCORE_API void SetDefine(FShaderCompilerDefineNameCache& Name, bool Value);
+	RENDERCORE_API void SetDefine(FShaderCompilerDefineNameCache& Name, float Value);
+
+	RENDERCORE_API int32 GetIntegerValue(FName Name) const;
+	RENDERCORE_API int32 GetIntegerValue(FShaderCompilerDefineNameCache& NameCache, int32 ResultIfNotFound = 0) const;
+
+	RENDERCORE_API bool ContainsDefinition(FName Name) const;
 
 	template <typename ValueType> void SetDefineIfUnset(const TCHAR* Name, ValueType Value)
 	{
 		FName NameKey(Name);
-		if (!Definitions.Contains(NameKey))
+		if (!ContainsDefinition(NameKey))
 		{
-			Definitions.SetDefine(NameKey, Value);
+			SetDefine(NameKey, Value);
 		}
 	}
 
@@ -830,7 +619,7 @@ struct FShaderCompilerEnvironment
 	UE_DEPRECATED(5.3, "GetDefinitions is deprecated; preprocessor defines must now only be accessed by core shader system code. Use Get/SetCompileArgument for generic params instead.")
 	const TMap<FString,FString>& GetDefinitions() const
 	{
-		return Definitions.GetDefinitionMap();
+		return UnusedStringDefinitions;
 	}
 
 	void SetRenderTargetOutputFormat(uint32 RenderTargetIndex, EPixelFormat PixelFormat)
@@ -839,34 +628,10 @@ struct FShaderCompilerEnvironment
 	}
 
 	/** This "core" serialization is also used for the hashing the compiler job (where files are handled differently). Should stay in sync with the ShaderCompileWorker. */
-	inline void SerializeEverythingButFiles(FArchive& Ar)
-	{
-		Ar << Definitions;
-		Ar << CompileArgs;
-		Ar << CompilerFlags;
-		Ar << RenderTargetOutputFormatsMap;
-		Ar << ResourceTableMap.Resources;
-		Ar << UniformBufferMap;
-		Ar << FullPrecisionInPS;
-		if (Ar.IsLoading())
-		{
-			ResourceTableMap.FixupOnLoad(UniformBufferMap);
-		}
-	}
+	RENDERCORE_API void SerializeEverythingButFiles(FArchive& Ar);
 
 	// Serializes the portions of the environment that are used as input to the backend compilation process (i.e. after all preprocessing)
-	inline void SerializeCompilationDependencies(FArchive& Ar)
-	{
-		Ar << CompileArgs;
-		Ar << CompilerFlags;
-		Ar << ResourceTableMap.Resources;
-		Ar << UniformBufferMap;
-		Ar << FullPrecisionInPS;
-		if (Ar.IsLoading())
-		{
-			ResourceTableMap.FixupOnLoad(UniformBufferMap);
-		}
-	}
+	RENDERCORE_API void SerializeCompilationDependencies(FArchive& Ar);
 
 	friend FArchive& operator<<(FArchive& Ar,FShaderCompilerEnvironment& Environment)
 	{
@@ -879,42 +644,21 @@ struct FShaderCompilerEnvironment
 		return Ar;
 	}
 	
-	void Merge(const FShaderCompilerEnvironment& Other)
-	{
-		// Merge the include maps
-		// Merge the values of any existing keys
-		for (TMap<FString,FString>::TConstIterator It(Other.IncludeVirtualPathToContentsMap); It; ++It )
-		{
-			FString* ExistingContents = IncludeVirtualPathToContentsMap.Find(It.Key());
-
-			if (ExistingContents)
-			{
-				ExistingContents->Append(It.Value());
-			}
-			else
-			{
-				IncludeVirtualPathToContentsMap.Add(It.Key(), It.Value());
-			}
-		}
-
-		check(Other.IncludeVirtualPathToExternalContentsMap.Num() == 0);
-
-		CompilerFlags.Append(Other.CompilerFlags);
-		ResourceTableMap.Append(Other.ResourceTableMap);
-		UniformBufferMap.Append(Other.UniformBufferMap);
-		Definitions.Merge(Other.Definitions);
-		CompileArgs.Append(Other.CompileArgs);
-		RenderTargetOutputFormatsMap.Append(Other.RenderTargetOutputFormatsMap);
-		FullPrecisionInPS |= Other.FullPrecisionInPS;
-	}
+	RENDERCORE_API void Merge(const FShaderCompilerEnvironment& Other);
 
 private:
 
 	friend class FShaderCompileUtilities;
 	friend class FShaderPreprocessorUtilities;
 
-	FShaderCompilerDefinitions Definitions;
+	PRAGMA_DISABLE_DEPRECATION_WARNINGS		// FShaderCompilerDefinitions will be made internal in the future, marked deprecated until then
+	TPimplPtr<FShaderCompilerDefinitions, EPimplPtrMode::DeepCopy> Definitions;
+	PRAGMA_ENABLE_DEPRECATION_WARNINGS
+
 	TMap<FString, TVariant<bool, float, int32, uint32, FString>> CompileArgs;
+
+	/** Unused data kept around for deprecated FShaderCompilerEnvironment::GetDefinitions call */
+	TMap<FString, FString> UnusedStringDefinitions;
 };
 
 
