Handle the case where the PSO cache mask is set before the cache is opened.

Before, tasks would only be enqueued whenever the mask was changed. If the pipeline file was opened after the mask was set, the task created as a result of the file opening would not be kicked off until the mask was set again (which could be never).

#pf 64b999681c55de20cf59e305

diff --git a/Engine/Source/Runtime/RHI/Private/PipelineFileCache.cpp b/Engine/Source/Runtime/RHI/Private/PipelineFileCache.cpp
--- a/Engine/Source/Runtime/RHI/Private/PipelineFileCache.cpp
+++ b/Engine/Source/Runtime/RHI/Private/PipelineFileCache.cpp
@@ -185,6 +185,7 @@ FPipelineFileCacheManager::PSOOrder FPipelineFileCacheManager::RequestedOrder =
 bool FPipelineFileCacheManager::FileCacheEnabled = false;
 FPipelineFileCacheManager::FPipelineStateLoggedEvent FPipelineFileCacheManager::PSOLoggedEvent;
 uint64 FPipelineFileCacheManager::GameUsageMask = 0;
+bool FPipelineFileCacheManager::GameUsageMaskSet = false;
 bool FPipelineFileCacheManager::LogNewPSOsToConsoleAndCSV = true;
 
 static int64 GetCurrentUnixTime()
@@ -3201,6 +3202,7 @@ uint64 FPipelineFileCacheManager::SetGameUsageMaskWithComparison(uint64 InGameUs
 		}
 		
 		FPipelineFileCacheManager::MaskComparisonFn = InComparisonFnPtr;
+		FPipelineFileCacheManager::GameUsageMaskSet = true;
 	}
 	
 	return OldMask;
diff --git a/Engine/Source/Runtime/RHI/Public/PipelineFileCache.h b/Engine/Source/Runtime/RHI/Public/PipelineFileCache.h
--- a/Engine/Source/Runtime/RHI/Public/PipelineFileCache.h
+++ b/Engine/Source/Runtime/RHI/Public/PipelineFileCache.h
@@ -399,7 +399,8 @@ public:
 	 * @returns the old mask
 	 */
 	RHI_API static uint64 SetGameUsageMaskWithComparison(uint64 GameUsageMask, FPSOMaskComparisonFn InComparisonFnPtr);
-	static uint64 GetGameUsageMask()	{ return GameUsageMask;}
+	RHI_API static uint64 GetGameUsageMask() { return GameUsageMask; }
+	RHI_API static bool IsGameUsageMaskSet() { return GameUsageMaskSet; }
 	
 	RHI_API static void PreCompileComplete();
 
@@ -447,6 +448,7 @@ private:
 	static bool FileCacheEnabled;
 	static FPipelineStateLoggedEvent PSOLoggedEvent;
 	RHI_API static uint64 GameUsageMask;
+	RHI_API static bool GameUsageMaskSet;
 	static FPSOMaskComparisonFn MaskComparisonFn;	
 
 	RHI_API static bool LogNewPSOsToConsoleAndCSV; // Whether to log new PSOs to the log file and CSV.
diff --git a/Engine/Source/Runtime/RenderCore/Private/ShaderPipelineCache.cpp b/Engine/Source/Runtime/RenderCore/Private/ShaderPipelineCache.cpp
--- a/Engine/Source/Runtime/RenderCore/Private/ShaderPipelineCache.cpp
+++ b/Engine/Source/Runtime/RenderCore/Private/ShaderPipelineCache.cpp
@@ -399,7 +399,12 @@ public:
 	FShaderPipelineCacheTask(FShaderPipelineCache* InCache, bool bIsUserCacheIn, EShaderPlatform InShaderPlatform) :
 		ShaderPlatform(InShaderPlatform), bIsUserCache(bIsUserCacheIn), Cache(InCache)
 	{
-		bReady = !CVarPSOFileCacheGameFileMaskEnabled.GetValueOnAnyThread() || CVarPSOFileCachePreOptimizeEnabled.GetValueOnAnyThread();
+		// Mark task as ready if either:
+		// * The usage mask mechanism is disabled.
+		// * The usage mask mechanism is enabled and the mask has already been set.
+		// * Pre-optimization is enabled.
+		bReady = (!CVarPSOFileCacheGameFileMaskEnabled.GetValueOnAnyThread() || FPipelineFileCacheManager::IsGameUsageMaskSet())
+			|| CVarPSOFileCachePreOptimizeEnabled.GetValueOnAnyThread();
 	}
 
 	bool Open(const FString& Key, const FString& Name, EShaderPlatform Platform);
@@ -665,21 +670,7 @@ bool FShaderPipelineCache::SetGameUsageMaskWithComparison(uint64 InMask, FPSOMas
 			bMaskChanged |= OldMask != InMask;
 			if (bMaskChanged)
 			{
-				// reset the current queue and all of the tasks.
-				ShaderPipelineCache->ClearPendingPrecompileTaskQueue();
-				int PendingPSOCacheCount = 0;
-				for (TPair < FString, TUniquePtr<class FShaderPipelineCacheTask>>& ShaderCacheTask : ShaderPipelineCache->ShaderCacheTasks)
-				{
-					if(!ShaderCacheTask.Value->CompletedMasks.Contains(InMask))
-					{
-						// Mask has changed and we have an open file refetch PSO's for this Mask - leave the FPipelineFileCacheManager file open - no need to close - just pull out the relevant PSOs.
-						// If this PSO compile run has completed for this Mask in which case don't refetch + compile for that mask
-						// Don't clear already compiled PSOHash list - this is not a full reset
-						ShaderCacheTask.Value->bReady = true;
-						ShaderPipelineCache->EnqueuePendingPrecompileTask(ShaderCacheTask.Key);
-						PendingPSOCacheCount++;
-					}
-				}
+				int PendingPSOCacheCount = ShaderPipelineCache->ApplyNewUsageMaskToAllTasks(InMask);
 
 				bMaskChanged = false;
 						
@@ -694,7 +685,6 @@ bool FShaderPipelineCache::SetGameUsageMaskWithComparison(uint64 InMask, FPSOMas
 		}
 		else
 		{
-			// NOTE: if this is called and then the cache is opened, but this function is never called again, the PSOs will not get precompiled. Should probably update the open call itself to see if there was a new mask set and call the code above
 			uint64 OldMask = FPipelineFileCacheManager::SetGameUsageMaskWithComparison(InMask, InComparisonFnPtr);
 			bMaskChanged |= OldMask != InMask;
 
@@ -710,6 +700,28 @@ bool FShaderPipelineCache::SetGameUsageMaskWithComparison(uint64 InMask, FPSOMas
 	return bMaskChanged;
 }
 
+
+int FShaderPipelineCache::ApplyNewUsageMaskToAllTasks(uint64 Mask)
+{
+	// Reset the current queue and all of the tasks.
+	ShaderPipelineCache->ClearPendingPrecompileTaskQueue();
+	int PendingPSOCacheCount = 0;
+	for (TPair < FString, TUniquePtr<class FShaderPipelineCacheTask>>& ShaderCacheTask : ShaderPipelineCache->ShaderCacheTasks)
+	{
+		if (!ShaderCacheTask.Value->CompletedMasks.Contains(Mask))
+		{
+			// Mask has changed and we have an open file refetch PSO's for this Mask - leave the FPipelineFileCacheManager file open - no need to close - just pull out the relevant PSOs.
+			// If this PSO compile run has completed for this Mask in which case don't refetch + compile for that mask
+			// Don't clear already compiled PSOHash list - this is not a full reset
+			ShaderCacheTask.Value->bReady = true;
+			ShaderPipelineCache->EnqueuePendingPrecompileTask(ShaderCacheTask.Key);
+			PendingPSOCacheCount++;
+		}
+	}
+
+	return PendingPSOCacheCount;
+}
+
 void FShaderPipelineCache::Initialize(EShaderPlatform Platform)
 {
 	check(ShaderPipelineCache == nullptr);
diff --git a/Engine/Source/Runtime/RenderCore/Public/ShaderPipelineCache.h b/Engine/Source/Runtime/RenderCore/Public/ShaderPipelineCache.h
--- a/Engine/Source/Runtime/RenderCore/Public/ShaderPipelineCache.h
+++ b/Engine/Source/Runtime/RenderCore/Public/ShaderPipelineCache.h
@@ -229,6 +229,7 @@ private:
 
 	RENDERCORE_API void ClearPendingPrecompileTaskQueue();
 	RENDERCORE_API void EnqueuePendingPrecompileTask(const FString& PSOCacheKey);
+	RENDERCORE_API int ApplyNewUsageMaskToAllTasks(uint64 Mask);
 
 	RENDERCORE_API FRHIBlendState* GetOrCreateBlendState(const FBlendStateInitializerRHI& Initializer);
 	RENDERCORE_API FRHIRasterizerState* GetOrCreateRasterizerState(const FRasterizerStateInitializerRHI& Initializer);
