Add debug utility to crash the GPU. From the console, use the command GPUDebugCrash with arguments for the crash type and queue you want to crash on. Usage: GPUDebugCrash [hang | pagefault | platformbreak [direct | compute]]. Crash type is required, queue will default to the graphics queue if unspecified.

diff --git a/Engine/Shaders/Private/Tools/GPUDebugCrashUtils.usf b/Engine/Shaders/Private/Tools/GPUDebugCrashUtils.usf
new file mode 100644
--- /dev/null
+++ b/Engine/Shaders/Private/Tools/GPUDebugCrashUtils.usf
@@ -0,0 +1,38 @@
+// Copyright Epic Games, Inc. All Rights Reserved.
+
+/*=============================================================================
+	GPUDebugCrashUtils.usf: Shaders for crashing the GPU for debugging purposes.
+=============================================================================*/
+
+#include "../Common.ush"
+
+RWTexture2D<uint> PageFaultUAV;
+
+[numthreads(THREADGROUPSIZE, THREADGROUPSIZE, 1)]
+void MainCS(uint3 ThreadID : SV_DispatchThreadID)
+{
+#if PLATFORM_BREAK_REQUESTED
+	PLATFORM_BREAK();
+#endif
+	uint what = PageFaultUAV[ThreadID.xy] + 1; // read
+
+	uint OldValue = 0;
+	uint2 index = uint2(0, 0);
+	uint2 index1 = uint2(0, 1);
+	InterlockedAdd(PageFaultUAV[index], 1, OldValue);
+	if (OldValue == 0)
+	{
+		PageFaultUAV[index1] = what;
+	}
+
+	// Try to trigger a crash (timeout) using a deliberate infinite loop
+	for (;;)
+	{
+		InterlockedAdd(PageFaultUAV[index], 0, OldValue);
+		if (OldValue == 0)
+		{
+			break;
+		}
+	}
+	
+}
diff --git a/Engine/Source/Runtime/Core/Private/HAL/ConsoleManager.cpp b/Engine/Source/Runtime/Core/Private/HAL/ConsoleManager.cpp
--- a/Engine/Source/Runtime/Core/Private/HAL/ConsoleManager.cpp
+++ b/Engine/Source/Runtime/Core/Private/HAL/ConsoleManager.cpp
@@ -2649,10 +2649,14 @@ void CreateConsoleVariables()
 	IConsoleManager::Get().RegisterConsoleCommand(TEXT("r.ResetRenderTargetsExtent"), TEXT("To reset internal render target extents"), ECVF_Cheat);
 #endif // !(UE_BUILD_SHIPPING || UE_BUILD_TEST)
 
-#if WITH_PROFILEGPU
+#if WITH_DUMPGPU
 	IConsoleManager::Get().RegisterConsoleCommand(TEXT("DumpGPU"), TEXT("Dump one frame of rendering intermediary resources to disk."), ECVF_Cheat);
 #endif
 
+#if WITH_GPUDEBUGCRASH
+	IConsoleManager::Get().RegisterConsoleCommand(TEXT("GPUDebugCrash"), TEXT("Crash GPU intentionally for debugging."), ECVF_Cheat);
+#endif
+
 #if	!UE_BUILD_SHIPPING
 	IConsoleManager::Get().RegisterConsoleCommand( TEXT( "DumpConsoleCommands" ), TEXT( "Dumps all console vaiables and commands and all exec that can be discovered to the log/console" ), ECVF_Default );
 
diff --git a/Engine/Source/Runtime/Core/Public/Misc/Build.h b/Engine/Source/Runtime/Core/Public/Misc/Build.h
--- a/Engine/Source/Runtime/Core/Public/Misc/Build.h
+++ b/Engine/Source/Runtime/Core/Public/Misc/Build.h
@@ -390,6 +390,17 @@
 // DumpGPU command
 #define WITH_DUMPGPU (!(UE_BUILD_SHIPPING || UE_BUILD_TEST) || (UE_BUILD_TEST && ALLOW_DUMPGPU_IN_TEST) || (UE_BUILD_SHIPPING && ALLOW_DUMPGPU_IN_SHIPPING))
 
+#ifndef ALLOW_GPUDEBUGCRASH_IN_TEST
+#define ALLOW_GPUDEBUGCRASH_IN_TEST 1
+#endif
+
+#ifndef ALLOW_GPUDEBUGCRASH_IN_SHIPPING
+#define ALLOW_GPUDEBUGCRASH_IN_SHIPPING 0
+#endif
+
+// GPUDebugCrash
+#define WITH_GPUDEBUGCRASH (!(UE_BUILD_SHIPPING || UE_BUILD_TEST) || (UE_BUILD_TEST && ALLOW_GPUDEBUGCRASH_IN_TEST) || (UE_BUILD_SHIPPING && ALLOW_GPUDEBUGCRASH_IN_SHIPPING))
+
 #ifndef ALLOW_CHEAT_CVARS_IN_TEST
 	#define ALLOW_CHEAT_CVARS_IN_TEST 1
 #endif
diff --git a/Engine/Source/Runtime/D3D12RHI/Private/D3D12Texture.cpp b/Engine/Source/Runtime/D3D12RHI/Private/D3D12Texture.cpp
--- a/Engine/Source/Runtime/D3D12RHI/Private/D3D12Texture.cpp
+++ b/Engine/Source/Runtime/D3D12RHI/Private/D3D12Texture.cpp
@@ -923,6 +923,13 @@ FD3D12Texture* FD3D12DynamicRHI::CreateD3D12Texture(const FRHITextureCreateDesc&
 
 		NewTexture->CreateViews();
 
+#if WITH_GPUDEBUGCRASH
+		if (EnumHasAnyFlags(CreateDesc.Flags, TexCreate_Invalid))
+		{
+			ID3D12Pageable* EvictableTexture = NewTexture->GetResource()->GetPageable();
+			Device->GetDevice()->Evict(1, &EvictableTexture);
+		}
+#endif
 		return NewTexture;
 	});
 
@@ -934,7 +941,6 @@ FD3D12Texture* FD3D12DynamicRHI::CreateD3D12Texture(const FRHITextureCreateDesc&
 		D3D12TextureOut->InitializeTextureData(RHICmdList, CreateDesc, InitialState);
 		CreateDesc.BulkData->Discard();
 	}
-
 	return D3D12TextureOut;
 #else
 	checkf(false, TEXT("XBOX_CODE_MERGE : Removed. The Xbox platform version should be used."));
diff --git a/Engine/Source/Runtime/Engine/Classes/Engine/Engine.h b/Engine/Source/Runtime/Engine/Classes/Engine/Engine.h
--- a/Engine/Source/Runtime/Engine/Classes/Engine/Engine.h
+++ b/Engine/Source/Runtime/Engine/Classes/Engine/Engine.h
@@ -2259,6 +2259,10 @@ public:
 	ENGINE_API bool HandleDumpGPUCommand( const TCHAR* Cmd, FOutputDevice& Ar );
 #endif
 
+#if WITH_GPUDEBUGCRASH
+	ENGINE_API bool HandleGPUDebugCrashCommand( const TCHAR* Cmd, FOutputDevice& Ar );
+#endif
+
 	// Compile in Debug or Development
 #if !(UE_BUILD_SHIPPING || UE_BUILD_TEST) && WITH_HOT_RELOAD
 	ENGINE_API bool HandleHotReloadCommand( const TCHAR* Cmd, FOutputDevice& Ar );
diff --git a/Engine/Source/Runtime/Engine/Private/UnrealEngine.cpp b/Engine/Source/Runtime/Engine/Private/UnrealEngine.cpp
--- a/Engine/Source/Runtime/Engine/Private/UnrealEngine.cpp
+++ b/Engine/Source/Runtime/Engine/Private/UnrealEngine.cpp
@@ -4878,6 +4878,13 @@ bool UEngine::Exec( UWorld* InWorld, const TCHAR* Cmd, FOutputDevice& Ar )
 	}
 #endif
 
+#if WITH_GPUDEBUGCRASH
+	else if (FParse::Command(&Cmd, TEXT("GPUDEBUGCRASH")))
+	{
+		return HandleGPUDebugCrashCommand(Cmd, Ar);
+	}
+#endif
+
 	else if ( FParse::Command(&Cmd,TEXT("SCALABILITY")) )
 	{
 		Scalability::ProcessCommand(Cmd, Ar);
@@ -5903,6 +5910,57 @@ bool UEngine::HandleDumpGPUCommand(const TCHAR* Cmd, FOutputDevice& Ar)
 
 #endif // WITH_DUMPGPU
 
+#if WITH_GPUDEBUGCRASH
+bool UEngine::HandleGPUDebugCrashCommand(const TCHAR* Cmd, FOutputDevice& Ar)
+{
+	if (!FApp::CanEverRender())
+	{
+		return true;
+	}
+
+	TArray<FString> Args;
+	Ar.Logf(TEXT("Scheduling GPU crash."));
+
+	bool bSupportsPageFaultsAndPlatformBreak = (GDynamicRHI->GetInterfaceType() == ERHIInterfaceType::D3D11 || GDynamicRHI->GetInterfaceType() == ERHIInterfaceType::OpenGL) ? false : true;
+
+	// schedule the GPU crash
+	if (FParse::Command(&Cmd, TEXT("pagefault")) && bSupportsPageFaultsAndPlatformBreak)
+	{
+		GRHIGlobals.TriggerGPUCrash = ERequestedGPUCrash::Type_PageFault;
+	}
+	if (FParse::Command(&Cmd, TEXT("platformbreak")) && bSupportsPageFaultsAndPlatformBreak)
+	{
+		GRHIGlobals.TriggerGPUCrash = ERequestedGPUCrash::Type_PlatformBreak;
+	}
+	if (FParse::Command(&Cmd, TEXT("hang")) || !bSupportsPageFaultsAndPlatformBreak)
+	{
+		if (!bSupportsPageFaultsAndPlatformBreak)
+		{
+			Ar.Logf(TEXT("Current RHI does not support GPU page faults on demand or PLATFORM_BREAK(). Falling back to GPU hang for requested crash."));
+		}
+		GRHIGlobals.TriggerGPUCrash = ERequestedGPUCrash::Type_Hang;
+	}
+	
+	if (GRHIGlobals.TriggerGPUCrash == ERequestedGPUCrash::None)
+	{
+		Ar.Logf(TEXT("Please specify a crash type for your requested GPU crash. Crash will be scheduled on direct queue unless compute queue is specified. Usage: gpudebugcrash [hang|pagefault|platformbreak [compute]]"));
+	}
+	else
+	{
+		if (FParse::Command(&Cmd, TEXT("compute")))
+		{
+			GRHIGlobals.TriggerGPUCrash |= ERequestedGPUCrash::Queue_Compute;
+		}
+		else
+		{
+			GRHIGlobals.TriggerGPUCrash |= ERequestedGPUCrash::Queue_Direct;
+		}
+	}
+
+	return true;
+}
+#endif
+
 #if !UE_BUILD_SHIPPING
 bool UEngine::HandleStartFPSChartCommand( const TCHAR* Cmd, FOutputDevice& Ar )
 {
diff --git a/Engine/Source/Runtime/RHI/Public/RHIDefinitions.h b/Engine/Source/Runtime/RHI/Public/RHIDefinitions.h
--- a/Engine/Source/Runtime/RHI/Public/RHIDefinitions.h
+++ b/Engine/Source/Runtime/RHI/Public/RHIDefinitions.h
@@ -980,6 +980,9 @@ enum class ETextureCreateFlags : uint64
 
 	/** Don't lump this texture with streaming memory when tracking total texture allocation sizes */
 	ForceIntoNonStreamingMemoryTracking = 1ull << 39,
+
+	/** Textures marked with this are meant to be immediately evicted after creation for intentionally crashing the GPU with a page fault. */
+	Invalid                           = 1ull << 40,
 };
 ENUM_CLASS_FLAGS(ETextureCreateFlags);
 
@@ -1022,6 +1025,7 @@ ENUM_CLASS_FLAGS(ETextureCreateFlags);
 #define TexCreate_MultiGPUGraphIgnore            ETextureCreateFlags::MultiGPUGraphIgnore
 #define TexCreate_ReservedResource               ETextureCreateFlags::ReservedResource
 #define TexCreate_ImmediateCommit                ETextureCreateFlags::ImmediateCommit
+#define TexCreate_Invalid                        ETextureCreateFlags::Invalid
 
 enum EAsyncComputePriority
 {
@@ -1209,6 +1213,18 @@ enum class EResourceTransitionFlags
 };
 ENUM_CLASS_FLAGS(EResourceTransitionFlags);
 
+enum class ERequestedGPUCrash : uint8
+{
+	None = 0,
+	Type_Hang = 1 << 0,
+	Type_PageFault = 1 << 1,
+	Type_PlatformBreak = 1 << 2,
+
+	Queue_Direct = 1 << 3,
+	Queue_Compute = 1 << 4
+};
+ENUM_CLASS_FLAGS(ERequestedGPUCrash);
+
 /** Returns whether the shader parameter type references an RDG texture. */
 inline bool IsRDGTextureReferenceShaderParameterType(EUniformBufferBaseType BaseType)
 {
diff --git a/Engine/Source/Runtime/RHI/Public/RHIGlobals.h b/Engine/Source/Runtime/RHI/Public/RHIGlobals.h
--- a/Engine/Source/Runtime/RHI/Public/RHIGlobals.h
+++ b/Engine/Source/Runtime/RHI/Public/RHIGlobals.h
@@ -354,6 +354,9 @@ struct FRHIGlobals
 	/** Whether we are profiling GPU hitches. */
 	bool TriggerGPUHitchProfile = false;
 
+	/** Whether an intentional GPU crash has been scheduled. */
+	ERequestedGPUCrash TriggerGPUCrash = ERequestedGPUCrash::None;
+
 	/** Non-empty if we are performing a gpu trace. Also says where to place trace file. */
 	FString GPUTraceFileName;
 
@@ -693,6 +696,7 @@ extern RHI_API FRHIGlobals GRHIGlobals;
 #define GVertexElementTypeSupport                              GRHIGlobals.VertexElementTypeSupport
 #define GTriggerGPUProfile                                     GRHIGlobals.TriggerGPUProfile
 #define GTriggerGPUHitchProfile                                GRHIGlobals.TriggerGPUHitchProfile
+#define GTriggerGPUCrash                                       GRHIGlobals.TriggerGPUCrash
 #define GGPUTraceFileName                                      GRHIGlobals.GPUTraceFileName
 #define GRHISupportsTextureStreaming                           GRHIGlobals.SupportsTextureStreaming
 #define GCurrentTextureMemorySize                              GRHIGlobals.CurrentTextureMemorySize
diff --git a/Engine/Source/Runtime/RenderCore/Private/GPUDebugCrashUtils.cpp b/Engine/Source/Runtime/RenderCore/Private/GPUDebugCrashUtils.cpp
new file mode 100644
--- /dev/null
+++ b/Engine/Source/Runtime/RenderCore/Private/GPUDebugCrashUtils.cpp
@@ -0,0 +1,77 @@
+// Copyright Epic Games, Inc. All Rights Reserved.
+
+/*=============================================================================
+	GPUDebugCrashUtils.cpp: Utilities for crashing the GPU in various ways on purpose.
+	=============================================================================*/
+
+#include "GPUDebugCrashUtils.h"
+#include "RenderResource.h"
+#include "RenderGraphUtils.h"
+#include "GlobalShader.h"
+#include "RHIStaticStates.h"
+#include "DataDrivenShaderPlatformInfo.h"
+
+class FGPUDebugCrashUtilsCS : public FGlobalShader
+{
+public:
+	DECLARE_GLOBAL_SHADER(FGPUDebugCrashUtilsCS)
+	SHADER_USE_PARAMETER_STRUCT(FGPUDebugCrashUtilsCS, FGlobalShader)
+
+	class FPlatformBreakRequested : SHADER_PERMUTATION_BOOL("PLATFORM_BREAK_REQUESTED");
+	using FPermutationDomain = TShaderPermutationDomain<FPlatformBreakRequested>;
+
+	BEGIN_SHADER_PARAMETER_STRUCT(FParameters, )
+		SHADER_PARAMETER_RDG_TEXTURE_UAV(RWTexture2D<uint>, PageFaultUAV)
+	END_SHADER_PARAMETER_STRUCT()
+
+	static void ModifyCompilationEnvironment(const FGlobalShaderPermutationParameters& Parameters, FShaderCompilerEnvironment& OutEnvironment)
+	{
+		FGlobalShader::ModifyCompilationEnvironment(Parameters, OutEnvironment);
+		OutEnvironment.SetDefine(TEXT("THREADGROUPSIZE"), FComputeShaderUtils::kGolden2DGroupSize);
+	}
+};
+
+IMPLEMENT_GLOBAL_SHADER(FGPUDebugCrashUtilsCS, "/Engine/Private/Tools/GPUDebugCrashUtils.usf", "MainCS", SF_Compute);
+
+// create RDG pass for the shader/pass that's going to crash the GPU
+void ScheduleGPUDebugCrash(FRDGBuilder& GraphBuilder)
+{
+	FGPUDebugCrashUtilsCS::FPermutationDomain PermutationVector;
+	PermutationVector.Set<FGPUDebugCrashUtilsCS::FPlatformBreakRequested>(EnumHasAnyFlags(GRHIGlobals.TriggerGPUCrash, ERequestedGPUCrash::Type_PlatformBreak));
+
+	auto ComputeShader = GetGlobalShaderMap(GMaxRHIFeatureLevel)->GetShader<FGPUDebugCrashUtilsCS>(PermutationVector);
+	FGPUDebugCrashUtilsCS::FParameters* PassParameters = GraphBuilder.AllocParameters<FGPUDebugCrashUtilsCS::FParameters>();
+	ETextureCreateFlags TexFlags = TexCreate_UAV | TexCreate_ShaderResource | TexCreate_NoFastClear | TexCreate_RenderTargetable;
+	FRDGTextureDesc CreateInfo = FRDGTextureDesc::Create2D(
+		FIntPoint(16, 16),
+		EPixelFormat::PF_R32_UINT,
+		FClearValueBinding::None,
+		TexFlags);
+	
+
+	if (EnumHasAnyFlags(GRHIGlobals.TriggerGPUCrash, ERequestedGPUCrash::Type_PageFault))
+	{
+		CreateInfo.Flags |= TexCreate_Invalid;
+	}
+
+	FRDGTextureRef PageFaultUAV = GraphBuilder.CreateTexture(CreateInfo, TEXT("GPUDebugCrashUtils.PageFaultUAV"));
+	PassParameters->PageFaultUAV = GraphBuilder.CreateUAV(PageFaultUAV);
+
+	FString CrashTypeString = EnumHasAnyFlags(GRHIGlobals.TriggerGPUCrash, ERequestedGPUCrash::Type_PageFault) ? TEXT("PageFault") : TEXT("Hang");
+	if (EnumHasAnyFlags(GRHIGlobals.TriggerGPUCrash, ERequestedGPUCrash::Type_PlatformBreak))
+	{
+		CrashTypeString = TEXT("PlatformBreak");
+	}
+
+	FComputeShaderUtils::AddPass(
+		GraphBuilder,
+		RDG_EVENT_NAME("GPUDebugCrash_%s_%s", 
+			EnumHasAnyFlags(GRHIGlobals.TriggerGPUCrash, ERequestedGPUCrash::Queue_Compute) ? TEXT("ComputeQueue") : TEXT("DirectQueue"),
+			 *CrashTypeString),
+		(EnumHasAnyFlags(GRHIGlobals.TriggerGPUCrash, ERequestedGPUCrash::Queue_Compute) ? ERDGPassFlags::AsyncCompute : ERDGPassFlags::Compute) | ERDGPassFlags::NeverCull,
+		ComputeShader,
+		PassParameters,
+		FIntVector(FComputeShaderUtils::kGolden2DGroupSize, FComputeShaderUtils::kGolden2DGroupSize, 1)
+	);
+	GRHIGlobals.TriggerGPUCrash = ERequestedGPUCrash::None;
+}
diff --git a/Engine/Source/Runtime/RenderCore/Public/GPUDebugCrashUtils.h b/Engine/Source/Runtime/RenderCore/Public/GPUDebugCrashUtils.h
new file mode 100644
--- /dev/null
+++ b/Engine/Source/Runtime/RenderCore/Public/GPUDebugCrashUtils.h
@@ -0,0 +1,8 @@
+// Copyright Epic Games, Inc. All Rights Reserved.
+
+/*=============================================================================
+	GPUDebugCrashUtils.h: Utilities for crashing the GPU in various ways on purpose.
+	=============================================================================*/
+class FRDGBuilder;
+
+extern RENDERCORE_API void ScheduleGPUDebugCrash(FRDGBuilder& GraphBuilder);
diff --git a/Engine/Source/Runtime/Renderer/Private/SceneRendering.cpp b/Engine/Source/Runtime/Renderer/Private/SceneRendering.cpp
--- a/Engine/Source/Runtime/Renderer/Private/SceneRendering.cpp
+++ b/Engine/Source/Runtime/Renderer/Private/SceneRendering.cpp
@@ -91,6 +91,7 @@
 #include "PrimitiveSceneShaderData.h"
 #include "Engine/SpecularProfile.h"
 #include "Engine/VolumeTexture.h"
+#include "GPUDebugCrashUtils.h"
 
 /*-----------------------------------------------------------------------------
 	Globals
@@ -4537,7 +4538,12 @@ static void RenderViewFamilies_RenderThread(FRHICommandListImmediate& RHICmdList
 
 		// We need to execute the pre-render view extensions before we do any view dependent work.
 		FSceneRenderer::ViewExtensionPreRender_RenderThread(GraphBuilder, SceneRenderer);
-
+#if WITH_GPUDEBUGCRASH
+		if (GRHIGlobals.TriggerGPUCrash != ERequestedGPUCrash::None) 
+		{
+			ScheduleGPUDebugCrash(GraphBuilder);
+		}
+#endif
 		SCOPE_CYCLE_COUNTER(STAT_TotalSceneRenderingTime);
 		SCOPED_NAMED_EVENT_TCHAR_CONDITIONAL(*ViewFamily.ProfileDescription, FColor::Red, !ViewFamily.ProfileDescription.IsEmpty());
 		const uint64 FamilyRenderStart = FPlatformTime::Cycles64();
