Fix non-unity build error.

diff --git a/Engine/Source/Runtime/RenderCore/Public/GPUDebugCrashUtils.h b/Engine/Source/Runtime/RenderCore/Public/GPUDebugCrashUtils.h
--- a/Engine/Source/Runtime/RenderCore/Public/GPUDebugCrashUtils.h
+++ b/Engine/Source/Runtime/RenderCore/Public/GPUDebugCrashUtils.h
@@ -3,6 +3,6 @@
 /*=============================================================================
 	GPUDebugCrashUtils.h: Utilities for crashing the GPU in various ways on purpose.
 	=============================================================================*/
-class FRDGBuilder;
+#include "RenderGraphBuilder.h"
 
 extern RENDERCORE_API void ScheduleGPUDebugCrash(FRDGBuilder& GraphBuilder);
