Add offline BLAS compression for supporting platforms:
* Add IsCompressed() to FRHIRayTracingGeometry.
* If BLAS is compressed it will trigger a build that will do the decompression.
* Controlled by target platform setting "Enable Static Mesh Offline BVH Compression"
* Enabled by default.

diff --git a/Engine/Source/Developer/TargetPlatform/Private/TargetPlatformBase.cpp b/Engine/Source/Developer/TargetPlatform/Private/TargetPlatformBase.cpp
--- a/Engine/Source/Developer/TargetPlatform/Private/TargetPlatformBase.cpp
+++ b/Engine/Source/Developer/TargetPlatform/Private/TargetPlatformBase.cpp
@@ -61,6 +61,11 @@ EOfflineBVHMode FTargetPlatformBase::GetStaticMeshOfflineBVHMode() const
 	return EOfflineBVHMode::Disabled;
 }
 
+bool FTargetPlatformBase::GetStaticMeshOfflineBVHCompression() const
+{
+	return false;
+}
+
 bool FTargetPlatformBase::ForcesSimpleSkyDiffuse() const
 {
 	return false;
diff --git a/Engine/Source/Developer/TargetPlatform/Public/Common/TargetPlatformBase.h b/Engine/Source/Developer/TargetPlatform/Public/Common/TargetPlatformBase.h
--- a/Engine/Source/Developer/TargetPlatform/Public/Common/TargetPlatformBase.h
+++ b/Engine/Source/Developer/TargetPlatform/Public/Common/TargetPlatformBase.h
@@ -92,6 +92,8 @@ public:
 
 	TARGETPLATFORM_API virtual EOfflineBVHMode GetStaticMeshOfflineBVHMode() const override;
 
+	TARGETPLATFORM_API virtual bool GetStaticMeshOfflineBVHCompression() const override;
+
 	TARGETPLATFORM_API virtual bool ForcesSimpleSkyDiffuse() const override;
 
 	TARGETPLATFORM_API virtual float GetDownSampleMeshDistanceFieldDivider() const override;
diff --git a/Engine/Source/Developer/TargetPlatform/Public/Interfaces/ITargetPlatform.h b/Engine/Source/Developer/TargetPlatform/Public/Interfaces/ITargetPlatform.h
--- a/Engine/Source/Developer/TargetPlatform/Public/Interfaces/ITargetPlatform.h
+++ b/Engine/Source/Developer/TargetPlatform/Public/Interfaces/ITargetPlatform.h
@@ -503,6 +503,11 @@ public:
 	*/
 	virtual EOfflineBVHMode GetStaticMeshOfflineBVHMode() const = 0;
 
+	/**
+	* Gets whether the platform will use compression for static mesh offline BVH.
+	*/
+	virtual bool GetStaticMeshOfflineBVHCompression() const = 0;
+
 	/**
 	* Gets whether the platform will use SH2 instead of SH3 for sky irradiance.
 	*/
diff --git a/Engine/Source/Runtime/RHI/Public/RHIResources.h b/Engine/Source/Runtime/RHI/Public/RHIResources.h
--- a/Engine/Source/Runtime/RHI/Public/RHIResources.h
+++ b/Engine/Source/Runtime/RHI/Public/RHIResources.h
@@ -3037,6 +3037,7 @@ public:
 
 	virtual FRayTracingAccelerationStructureAddress GetAccelerationStructureAddress(uint64 GPUIndex) const = 0;
 	virtual void SetInitializer(const FRayTracingGeometryInitializer& Initializer) = 0;
+	virtual bool IsCompressed() const { return false; }
 
 	const FRayTracingGeometryInitializer& GetInitializer() const
 	{
diff --git a/Engine/Source/Runtime/RenderCore/Private/RayTracingGeometry.cpp b/Engine/Source/Runtime/RenderCore/Private/RayTracingGeometry.cpp
--- a/Engine/Source/Runtime/RenderCore/Private/RayTracingGeometry.cpp
+++ b/Engine/Source/Runtime/RenderCore/Private/RayTracingGeometry.cpp
@@ -60,9 +60,9 @@ void FRayTracingGeometry::CreateRayTracingGeometryFromCPUData(FRHICommandList& R
 		Initializer.OfflineData->Discard();
 		Initializer.OfflineData = nullptr;
 	}
-
-	SetRequiresBuild(Initializer.OfflineData == nullptr);
+	
 	RayTracingGeometryRHI = RHICmdList.CreateRayTracingGeometry(Initializer);
+	SetRequiresBuild(Initializer.OfflineData == nullptr || RayTracingGeometryRHI->IsCompressed());
 }
 
 void FRayTracingGeometry::CreateRayTracingGeometryFromCPUData(TResourceArray<uint8>& OfflineData)
@@ -203,6 +203,11 @@ void FRayTracingGeometry::CreateRayTracingGeometry(FRHICommandList& RHICmdList,
 		}
 		else
 		{
+			if (RayTracingGeometryRHI && RayTracingGeometryRHI->IsCompressed())
+			{
+				RayTracingBuildRequestIndex = GRayTracingGeometryManager.RequestBuildAccelerationStructure(RHICmdList, this, InBuildPriority);
+			}
+
 			SetRequiresBuild(false);
 
 			// Offline data ownership is transferred to the RHI, which discards it after use.
