Shader Compiler:  Fix for circular include issue temporarily introduced as a step towards deprecation of public access to FShaderCompilerDefinitions structure.

#rnx

diff --git a/Engine/Source/Runtime/RenderCore/Internal/ShaderCompilerDefinitions.h b/Engine/Source/Runtime/RenderCore/Internal/ShaderCompilerDefinitions.h
--- a/Engine/Source/Runtime/RenderCore/Internal/ShaderCompilerDefinitions.h
+++ b/Engine/Source/Runtime/RenderCore/Internal/ShaderCompilerDefinitions.h
@@ -2,7 +2,9 @@
 
 #pragma once
 
-#include "ShaderCore.h"
+// Currently, "ShaderCore.h" includes this header, so this is a circular include, but we will remove this header from "ShaderCore.h" in 5.5,
+// as part of hiding FShaderCompilerDefinitions from public view, and will need to add this include back in.
+// #include "ShaderCore.h"
 
 enum class EShaderCompilerDefineVariant : uint8
 {
