Only include GPU crash utils when it's enabled, fix warning about missing header guard

diff --git a/Engine/Source/Runtime/RenderCore/Private/GPUDebugCrashUtils.cpp b/Engine/Source/Runtime/RenderCore/Private/GPUDebugCrashUtils.cpp
--- a/Engine/Source/Runtime/RenderCore/Private/GPUDebugCrashUtils.cpp
+++ b/Engine/Source/Runtime/RenderCore/Private/GPUDebugCrashUtils.cpp
@@ -5,6 +5,7 @@
 	=============================================================================*/
 
 #include "GPUDebugCrashUtils.h"
+#include "RenderCore.h"
 #include "RenderResource.h"
 #include "RenderGraphUtils.h"
 #include "GlobalShader.h"
diff --git a/Engine/Source/Runtime/RenderCore/Public/GPUDebugCrashUtils.h b/Engine/Source/Runtime/RenderCore/Public/GPUDebugCrashUtils.h
--- a/Engine/Source/Runtime/RenderCore/Public/GPUDebugCrashUtils.h
+++ b/Engine/Source/Runtime/RenderCore/Public/GPUDebugCrashUtils.h
@@ -1,8 +1,10 @@
 // Copyright Epic Games, Inc. All Rights Reserved.
 
+#pragma once
 /*=============================================================================
 	GPUDebugCrashUtils.h: Utilities for crashing the GPU in various ways on purpose.
 	=============================================================================*/
-#include "RenderGraphBuilder.h"
+
+class FRDGBuilder;
 
 extern RENDERCORE_API void ScheduleGPUDebugCrash(FRDGBuilder& GraphBuilder);
diff --git a/Engine/Source/Runtime/Renderer/Private/SceneRendering.cpp b/Engine/Source/Runtime/Renderer/Private/SceneRendering.cpp
--- a/Engine/Source/Runtime/Renderer/Private/SceneRendering.cpp
+++ b/Engine/Source/Runtime/Renderer/Private/SceneRendering.cpp
@@ -91,7 +91,9 @@
 #include "PrimitiveSceneShaderData.h"
 #include "Engine/SpecularProfile.h"
 #include "Engine/VolumeTexture.h"
+#if WITH_GPUDEBUGCRASH
 #include "GPUDebugCrashUtils.h"
+#endif
 
 /*-----------------------------------------------------------------------------
 	Globals
