Fix header guard warning.

diff --git a/Engine/Source/Runtime/RenderCore/Public/GPUDebugCrashUtils.h b/Engine/Source/Runtime/RenderCore/Public/GPUDebugCrashUtils.h
--- a/Engine/Source/Runtime/RenderCore/Public/GPUDebugCrashUtils.h
+++ b/Engine/Source/Runtime/RenderCore/Public/GPUDebugCrashUtils.h
@@ -1,4 +1,5 @@
 // Copyright Epic Games, Inc. All Rights Reserved.
+#pragma once
 
 /*=============================================================================
 	GPUDebugCrashUtils.h: Utilities for crashing the GPU in various ways on purpose.
