Shader Compiler static analysis warning:  need to move include of Runtime/RenderCore/Internal/ShaderCompilerDefinitions.h to top of ShaderCore.h to avoid "Include after first code block" warning.  To make this move, a structure required by that header (FShaderCompilerDefinitions) needed to be moved to a different header as well.

#rnx

diff --git a/Engine/Source/Runtime/RenderCore/Public/ShaderCore.h b/Engine/Source/Runtime/RenderCore/Public/ShaderCore.h
--- a/Engine/Source/Runtime/RenderCore/Public/ShaderCore.h
+++ b/Engine/Source/Runtime/RenderCore/Public/ShaderCore.h
@@ -41,6 +41,9 @@
 #include "UObject/UnrealNames.h"
 #include "UniformBuffer.h"
 
+// Temporarily included here until we can fully deprecate access to FShaderCompilerDefinitions in a future version
+#include "Runtime/RenderCore/Internal/ShaderCompilerDefinitions.h"
+
 class Error;
 class FMemoryImageWriter;
 class FMemoryUnfreezeContent;
@@ -299,29 +302,6 @@ public:
 	TMap<FString,FParameterAllocation> ParameterMap;
 };
 
-/** Utility class for caching FName and other info for a define.  Don't use directly, use SET_SHADER_DEFINE or SET_SHADER_DEFINE_AND_COMPILE_ARGUMENT macros. */
-class FShaderCompilerDefineNameCache
-{
-public:
-	FShaderCompilerDefineNameCache(const TCHAR* InName)
-		: Name(InName), MapIndex(INDEX_NONE)
-	{}
-
-	operator FName() const
-	{
-		return Name;
-	}
-
-private:
-	FName Name;
-	int32 MapIndex;
-
-	friend class FShaderCompilerDefinitions;
-};
-
-// Temporarily included here until we can fully deprecate access to FShaderCompilerDefinitions in a future version
-#include "Runtime/RenderCore/Internal/ShaderCompilerDefinitions.h"
-
 struct FShaderResourceTable
 {
 	/** Bits indicating which resource tables contain resources bound to this shader. */
diff --git a/Engine/Source/Runtime/RenderCore/Public/ShaderParameterMetadata.h b/Engine/Source/Runtime/RenderCore/Public/ShaderParameterMetadata.h
--- a/Engine/Source/Runtime/RenderCore/Public/ShaderParameterMetadata.h
+++ b/Engine/Source/Runtime/RenderCore/Public/ShaderParameterMetadata.h
@@ -471,3 +471,34 @@ private:
 	RENDERCORE_API void InitializeUniformBufferDeclaration();
 #endif
 };
+
+
+/**
+ * Utility class for caching FName and other info for a shader compiler define.  Don't use directly, use SET_SHADER_DEFINE or SET_SHADER_DEFINE_AND_COMPILE_ARGUMENT
+ * macros defined in ShaderCore.h.
+ *
+ * Placed here temporarily to solve static analysis warning related to deprecation of public access to FShaderCompilerDefinitions.  This structure needs to be
+ * declared before including "Runtime/RenderCore/Internal/ShaderCompilerDefinitions.h", which would normally be accomplished by including ShaderCore.h, but we
+ * can't include that without producing a circular include, which generates a static analysis warning.  The first attempt to solve this involved placing the
+ * include after this structure declaration in ShaderCore.h, but that leads to an "Include after first code block" static analysis warning.  We don't want to move
+ * this structure to the internal header, so we need to place it in another header that's also included in "ShaderCore.h", to avoid producing that warning.  This
+ * is the only candidate shader related source file, so it makes sense to place here for now, and can be moved back to ShaderCore.h next major version.
+ */
+class FShaderCompilerDefineNameCache
+{
+public:
+	FShaderCompilerDefineNameCache(const TCHAR* InName)
+		: Name(InName), MapIndex(INDEX_NONE)
+	{}
+
+	operator FName() const
+	{
+		return Name;
+	}
+
+private:
+	FName Name;
+	int32 MapIndex;
+
+	friend class FShaderCompilerDefinitions;
+};
