Fixed assert related to mobile GPUScene on SM5 static shader platforms

diff --git a/Engine/Source/Runtime/RenderCore/Private/VertexFactory.cpp b/Engine/Source/Runtime/RenderCore/Private/VertexFactory.cpp
--- a/Engine/Source/Runtime/RenderCore/Private/VertexFactory.cpp
+++ b/Engine/Source/Runtime/RenderCore/Private/VertexFactory.cpp
@@ -294,7 +294,7 @@ bool FVertexFactory::AddPrimitiveIdStreamElement(EVertexInputStreamType InputStr
 			Elements.Add(AccessStreamComponent(FVertexStreamComponent(&GPrimitiveIdDummy, 0, 0, 0u, VET_UInt, EVertexStreamUsage::Instancing), AttributeIndex, InputStreamType));
 			SetPrimitiveIdStreamIndex(GMaxRHIFeatureLevel, InputStreamType, Elements.Last().StreamIndex);
 			
-			if (AttributeIndex_Mobile != 0xff && GMaxRHIFeatureLevel != ERHIFeatureLevel::ES3_1)
+			if (GIsEditor && (AttributeIndex_Mobile != 0xff && GMaxRHIFeatureLevel != ERHIFeatureLevel::ES3_1))
 			{
 				SetPrimitiveIdStreamIndex(ERHIFeatureLevel::ES3_1, InputStreamType, Elements.Last().StreamIndex);
 			}
