RenderCore: Fix AddDrawTexturePass() for specific mip level in different format

diff --git a/Engine/Source/Runtime/RenderCore/Private/RenderGraphUtils.cpp b/Engine/Source/Runtime/RenderCore/Private/RenderGraphUtils.cpp
--- a/Engine/Source/Runtime/RenderCore/Private/RenderGraphUtils.cpp
+++ b/Engine/Source/Runtime/RenderCore/Private/RenderGraphUtils.cpp
@@ -339,11 +339,8 @@ RENDERCORE_API void AddDrawTexturePass(
 				const int32 SourceSliceIndex = SliceIndex + DrawInfo.SourceSliceIndex;
 				const int32 DestSliceIndex = SliceIndex + DrawInfo.DestSliceIndex;
 
-				FRDGTextureSRVDesc SRVDesc = FRDGTextureSRVDesc::Create(InputTexture);
+				FRDGTextureSRVDesc SRVDesc = FRDGTextureSRVDesc::CreateForMipLevel(InputTexture, SourceMipIndex);
 				SRVDesc.FirstArraySlice = SourceSliceIndex;
-				SRVDesc.NumArraySlices = 1;
-				SRVDesc.MipLevel = SourceMipIndex;
-				SRVDesc.NumMipLevels = 1;
 
 				auto* PassParameters = GraphBuilder.AllocParameters<FDrawTexturePS::FParameters>();
 				PassParameters->InputTexture = GraphBuilder.CreateSRV(SRVDesc);
