Better fix for UE-190842 generate new shaderversion guid and restore previous DDC logic

#tests ran editor, cook lyra

diff --git a/Engine/Shaders/Public/ShaderVersion.ush b/Engine/Shaders/Public/ShaderVersion.ush
--- a/Engine/Shaders/Public/ShaderVersion.ush
+++ b/Engine/Shaders/Public/ShaderVersion.ush
@@ -2,5 +2,4 @@
 // Each time the console command is executed it generates a new GUID. As this file is included
 // in Platform.ush (which should be included in any shader) it allows to invalidate the shader DDC.
 // 
-// GUID = 631DFD83-05E2-4AC1-A7AC-47CBF4E629EB
-
+// GUID = 111D4D59-4A68-0336-ACC5-4F9E3346A24F
diff --git a/Engine/Source/Runtime/RenderCore/Private/Shader.cpp b/Engine/Source/Runtime/RenderCore/Private/Shader.cpp
--- a/Engine/Source/Runtime/RenderCore/Private/Shader.cpp
+++ b/Engine/Source/Runtime/RenderCore/Private/Shader.cpp
@@ -1598,11 +1598,8 @@ void ShaderMapAppendKeyString(EShaderPlatform Platform, FString& KeyString)
 		// Extra data (names, etc)
 		KeyString += ShouldEnableExtraShaderData(ShaderFormatName) ? TEXT("_ExtraData") : TEXT("");
 		// Symbols and/or SymbolsInfo and version if symbols serialization changes
-		const bool bSymbols = ShouldGenerateShaderSymbols(ShaderFormatName);
-		const bool bSymbolsinfo = ShouldGenerateShaderSymbolsInfo(ShaderFormatName);
-		KeyString += bSymbols ? TEXT("_Symbols") : TEXT("");
-		KeyString += bSymbolsinfo ? TEXT("_SymbolsInfo") : TEXT("");
-		KeyString += bSymbols || bSymbolsinfo ? TEXT("_V2") : TEXT("");
+		KeyString += ShouldGenerateShaderSymbols(ShaderFormatName) ? TEXT("_Symbols") : TEXT("");
+		KeyString += ShouldGenerateShaderSymbolsInfo(ShaderFormatName) ? TEXT("_SymbolsInfo") : TEXT("");
 		// Are symbols based on source or results
 		KeyString += ShouldAllowUniqueShaderSymbols(ShaderFormatName) ? TEXT("_FullDbg") : TEXT("");
 	}
