[shaders] move custom vulkan bindless parameter parser behaviour into baseclass (parameter parser being part of the serialized preprocess output means subclassing is problematic)

diff --git a/Engine/Source/Developer/VulkanShaderFormat/Private/VulkanShaderCompiler.cpp b/Engine/Source/Developer/VulkanShaderFormat/Private/VulkanShaderCompiler.cpp
--- a/Engine/Source/Developer/VulkanShaderFormat/Private/VulkanShaderCompiler.cpp
+++ b/Engine/Source/Developer/VulkanShaderFormat/Private/VulkanShaderCompiler.cpp
@@ -2500,61 +2500,6 @@ static bool CompileWithShaderConductor(
 #endif // PLATFORM_MAC || PLATFORM_WINDOWS || PLATFORM_LINUX
 
 
-// Overload parameter declaration from the default FShaderParameterParser in order to take into account Vulkan particularities for resources
-class FVulkanShaderParameterParser : public FShaderParameterParser
-{
-public:
-	FVulkanShaderParameterParser(FShaderCompilerFlags CompilerFlags, const TCHAR* InConstantBufferType)
-		: FShaderParameterParser(CompilerFlags, InConstantBufferType)
-	{}
-
-protected:
-	FString GenerateBindlessParameterDeclaration(const FParsedShaderParameter& ParsedParameter) const override
-	{
-		if (bBindlessResources || bBindlessSamplers)
-		{
-			const bool IsSampler = (ParsedParameter.BindlessConversionType == EBindlessConversionType::Sampler);
-			const TCHAR* IndexPrefix = IsSampler ? FShaderParameterParser::kBindlessSamplerPrefix : FShaderParameterParser::kBindlessResourcePrefix;
-			const TCHAR* HeapPrefix = IsSampler ? VulkanBindless::kBindlessSamplerArrayPrefix : VulkanBindless::kBindlessResourceArrayPrefix;
-
-			const TCHAR* StorageClass = ParsedParameter.bGloballyCoherent ? TEXT("globallycoherent ") : TEXT("");
-
-			const FStringView Name = ParsedParameter.ParsedName;
-			const FStringView Type = ParsedParameter.ParsedType;
-
-			const FString RewriteType = FString::Printf(TEXT("SafeType%.*s"), Name.Len(), Name.GetData());
-
-			TStringBuilder<512> Result;
-
-			// If we weren't going to be added to a root constant buffer, that means we need to declare our index before we declare our getter.
-			if (ParsedParameter.ConstantBufferParameterType == EShaderParameterType::Num)
-			{
-				// e.g. `uint BindlessResource_##Name;`
-				Result << TEXT("uint ") << IndexPrefix << Name << TEXT("; ");
-			}
-
-			// Add the typedef
-			Result << TEXT("typedef ") << Type << TEXT(" ") << RewriteType << TEXT("; ");
-
-			// Declare a heap for the RewriteType
-			// e.g. `SafeType##Name ResourceDescriptorHeap_SafeType##Name[];`
-			Result << StorageClass << RewriteType << TEXT(" ") << HeapPrefix << RewriteType << TEXT("[]; ");
-			// :todo-jn: specify the descripor set and binding directly in source instead of patching SPIRV
-
-			// e.g. `static const Type Name = GetBindlessResource##Name()`
-			Result << TEXT("static const ") << StorageClass << RewriteType << TEXT(" ") << Name << TEXT(" = ") << HeapPrefix << RewriteType << TEXT("[") << IndexPrefix << Name << TEXT("];");
-
-			return Result.ToString();
-		}
-		else
-		{
-			// use original code path
-			return FShaderParameterParser::GenerateBindlessParameterDeclaration(ParsedParameter);
-		}
-	}
-};
-
-
 void DoCompileVulkanShader(const FShaderCompilerInput& Input, FShaderCompilerOutput& Output, const class FString& WorkingDirectory, EVulkanShaderVersion Version)
 {
 	check(IsVulkanShaderFormat(Input.ShaderFormat));
@@ -2684,14 +2629,14 @@ void DoCompileVulkanShader(const FShaderCompilerInput& Input, FShaderCompilerOut
 		}
 	}
 
-	FVulkanShaderParameterParser ShaderParameterParser(Input.Environment.CompilerFlags, nullptr);
-	if (!ShaderParameterParser.ParseAndModify(Input, Output.Errors, PreprocessedShaderSource))
+	FShaderParameterParser ShaderParameterParser(Input.Environment.CompilerFlags, nullptr);
+	if (!ShaderParameterParser.ParseAndModify(Input, Output.Errors, PreprocessedShaderSource, EBindlessParameterMode::Vulkan))
 	{
 		// The FShaderParameterParser will add any relevant errors.
 		return;
 	}
 
-	const FString EntryPointName = Input.EntryPointName;
+const FString EntryPointName = Input.EntryPointName;
 
 	RemoveUniformBuffersFromSource(Input.Environment, PreprocessedShaderSource);
 
diff --git a/Engine/Source/Runtime/RenderCore/Private/ShaderParameterParser.cpp b/Engine/Source/Runtime/RenderCore/Private/ShaderParameterParser.cpp
--- a/Engine/Source/Runtime/RenderCore/Private/ShaderParameterParser.cpp
+++ b/Engine/Source/Runtime/RenderCore/Private/ShaderParameterParser.cpp
@@ -853,41 +853,55 @@ void FShaderParameterParser::RemoveMovingParametersFromSource(FString& Preproces
 FString FShaderParameterParser::GenerateBindlessParameterDeclaration(const FParsedShaderParameter& ParsedParameter) const
 {
 	// NOTE: Macros AUTO_BINDLESS_SAMPLER_INDEX/VARIABLE and AUTO_BINDLESS_RESOURCE_INDEX/VARIABLE in BindlessResources.ush must be kept in sync with this function
-
 	const bool bIsSampler = (ParsedParameter.BindlessConversionType == EBindlessConversionType::Sampler);
-
-	const TCHAR* Kind = bIsSampler ? TEXT("Sampler") : TEXT("Resource");
-	const TCHAR* StorageClass = ParsedParameter.bGloballyCoherent ? TEXT("globallycoherent ") : TEXT("");
-
 	const FStringView Name = ParsedParameter.ParsedName;
 	const FStringView Type = ParsedParameter.ParsedType;
+	const TCHAR* StorageClass = ParsedParameter.bGloballyCoherent ? TEXT("globallycoherent ") : TEXT("");
+	const TCHAR* IndexPrefix = bIsSampler ? FShaderParameterParser::kBindlessSamplerPrefix : FShaderParameterParser::kBindlessResourcePrefix;
+
+	TStringBuilder<512> Result;
+
+	// If we weren't going to be added to a root constant buffer, that means we need to declare our index before we declare our getter.
+	if (ParsedParameter.ConstantBufferParameterType == EShaderParameterType::Num)
+	{
+		// e.g. `uint BindlessResource_##Name;`
+		Result << TEXT("uint ") << IndexPrefix << Name << TEXT("; ");
+	}
 
 	FString RewriteType(Type);
 
 	FString TypedefText = TEXT("");
-	if (ParsedParameter.BindlessConversionType == EBindlessConversionType::Resource)
+	// Vulkan requires these typedefs for both resources and samplers, while other platforms only require resources.
+	if (ParsedParameter.BindlessConversionType == EBindlessConversionType::Resource || BindlessParameterMode == EBindlessParameterMode::Vulkan)
 	{
 		RewriteType = FString::Printf(TEXT("SafeType%.*s"), Name.Len(), Name.GetData());
 		TypedefText = FString::Printf(TEXT("typedef %.*s %s;"), Type.Len(), Type.GetData(), *RewriteType);
 	}
 
-	TStringBuilder<512> Result;
+	Result << TypedefText;
 
-	// If we weren't going to be added to a root constant buffer, that means we need to declare our index before we declare our getter.
-	if (ParsedParameter.ConstantBufferParameterType == EShaderParameterType::Num)
+	if (BindlessParameterMode == EBindlessParameterMode::Vulkan)
 	{
-		// e.g. `uint BindlessResource_##Name;`
-		// or   `uint BindlessSampler_##Name;`
-		Result << TEXT("uint Bindless") << Kind << TEXT("_") << Name << TEXT(";");
-	}
+		const TCHAR* HeapPrefix = bIsSampler ? VulkanBindless::kBindlessSamplerArrayPrefix : VulkanBindless::kBindlessResourceArrayPrefix;
 
-	Result << TypedefText;
+		// Declare a heap for the RewriteType
+		// e.g. `SafeType##Name ResourceDescriptorHeap_SafeType##Name[];`
+		Result << StorageClass << RewriteType << TEXT(" ") << HeapPrefix << RewriteType << TEXT("[]; ");
+		// :todo-jn: specify the descriptor set and binding directly in source instead of patching SPIRV
 
-	// e.g. `Type GetBindlessResource##Name() { return GetResourceFromHeap(Type, BindlessResource_##Name); } static const Type Name = GetBindlessResource##Name()`
-	// or   `Type GetBindlessSampler##Name() { return GetSamplerFromHeap(Type, BindlessSampler_##Name); } static const Type Name = GetBindlessSampler##Name()`
-	Result << StorageClass << RewriteType << TEXT(" GetBindless") << Kind << Name << TEXT("()");
-	Result << TEXT("{ return Get") << Kind << TEXT("FromHeap(") << StorageClass << RewriteType << TEXT(", Bindless") << Kind << TEXT("_") << Name << TEXT("); } ");
-	Result << TEXT("static const ") << StorageClass << RewriteType << TEXT(" ") << Name << TEXT(" = GetBindless") << Kind << Name << TEXT("();");
+		// e.g. `static const SafeType##Name Name = ResourceDescriptorHeap_SafeType##Name[BindlessResource_##Name];`
+		Result << TEXT("static const ") << StorageClass << RewriteType << TEXT(" ") << Name << TEXT(" = ") << HeapPrefix << RewriteType << TEXT("[") << IndexPrefix << Name << TEXT("];");
+	}
+	else
+	{
+		const TCHAR* Kind = bIsSampler ? TEXT("Sampler") : TEXT("Resource");
+
+		// e.g. `Type GetBindlessResource##Name() { return GetResourceFromHeap(Type, BindlessResource_##Name); } static const Type Name = GetBindlessResource##Name()`
+		// or   `Type GetBindlessSampler##Name() { return GetSamplerFromHeap(Type, BindlessSampler_##Name); } static const Type Name = GetBindlessSampler##Name()`
+		Result << StorageClass << RewriteType << TEXT(" GetBindless") << Kind << Name << TEXT("()");
+		Result << TEXT("{ return Get") << Kind << TEXT("FromHeap(") << StorageClass << RewriteType << TEXT(", Bindless") << Kind << TEXT("_") << Name << TEXT("); } ");
+		Result << TEXT("static const ") << StorageClass << RewriteType << TEXT(" ") << Name << TEXT(" = GetBindless") << Kind << Name << TEXT("();");
+	}
 
 	return Result.ToString();
 }
@@ -1091,9 +1105,11 @@ PRAGMA_ENABLE_DEPRECATION_WARNINGS
 bool FShaderParameterParser::ParseAndModify(
 	const FShaderCompilerInput& CompilerInput,
 	TArray<FShaderCompilerError>& OutErrors,
-	FString& PreprocessedShaderSource)
+	FString& PreprocessedShaderSource,
+	EBindlessParameterMode InBindlessParameterMode)
 {
 	const bool bHasRootParameters = (CompilerInput.RootParametersStructure != nullptr);
+	BindlessParameterMode = InBindlessParameterMode;
 
 	// The shader doesn't have any parameter binding through shader structure, therefore don't do anything.
 	if (!(bBindlessResources || bBindlessSamplers || bHasRootParameters))
diff --git a/Engine/Source/Runtime/RenderCore/Public/ShaderParameterParser.h b/Engine/Source/Runtime/RenderCore/Public/ShaderParameterParser.h
--- a/Engine/Source/Runtime/RenderCore/Public/ShaderParameterParser.h
+++ b/Engine/Source/Runtime/RenderCore/Public/ShaderParameterParser.h
@@ -20,6 +20,19 @@ enum class EBindlessConversionType : uint8
 	Sampler
 };
 
+enum class EBindlessParameterMode : uint8
+{
+	Default,
+	Vulkan,
+};
+
+namespace VulkanBindless
+{
+	// Prefix used to declare arrays of samplers/resources for bindless
+	static constexpr const TCHAR* kBindlessResourceArrayPrefix = TEXT("ResourceDescriptorHeap_");
+	static constexpr const TCHAR* kBindlessSamplerArrayPrefix = TEXT("SamplerDescriptorHeap_");
+}
+
 /** Validates and moves all the shader loose data parameter defined in the root scope of the shader into the root uniform buffer. */
 class FShaderParameterParser
 {
@@ -100,7 +113,8 @@ public:
 	RENDERCORE_API bool ParseAndModify(
 		const FShaderCompilerInput& CompilerInput,
 		TArray<FShaderCompilerError>& OutErrors,
-		FString& PreprocessedShaderSource
+		FString& PreprocessedShaderSource,
+		EBindlessParameterMode BindlessParameterMode = EBindlessParameterMode::Default
 	);
 
 	UE_DEPRECATED(5.2, "ParseAndModify doesn't need ConstantBufferType anymore")
@@ -190,7 +204,7 @@ protected:
 	* Generates shader source code to declare a bindless resource or sampler (for automatic bindless conversion).
 	* May be overriden to allow custom implementations for different platforms.
 	*/
-	RENDERCORE_API virtual FString GenerateBindlessParameterDeclaration(const FParsedShaderParameter& ParsedParameter) const;
+	RENDERCORE_API FString GenerateBindlessParameterDeclaration(const FParsedShaderParameter& ParsedParameter) const;
 
 	const TCHAR* ConstantBufferType = nullptr;
 
@@ -203,6 +217,7 @@ protected:
 
 	bool bBindlessResources = false;
 	bool bBindlessSamplers = false;
+	EBindlessParameterMode BindlessParameterMode = EBindlessParameterMode::Default;
 
 	/** Indicates that parameters should be moved to the root cosntant buffer. */
 	bool bNeedToMoveToRootConstantBuffer = false;
diff --git a/Engine/Source/Runtime/VulkanRHI/Public/VulkanCommon.h b/Engine/Source/Runtime/VulkanRHI/Public/VulkanCommon.h
--- a/Engine/Source/Runtime/VulkanRHI/Public/VulkanCommon.h
+++ b/Engine/Source/Runtime/VulkanRHI/Public/VulkanCommon.h
@@ -135,10 +135,6 @@ namespace VulkanBindless
 		NumBindlessSets,
 		MaxNumSets = NumBindlessSets
 	};
-
-	// Prefix used to declare arrays of samplers/resources for bindless
-	static constexpr const TCHAR* kBindlessResourceArrayPrefix = TEXT("ResourceDescriptorHeap_");
-	static constexpr const TCHAR* kBindlessSamplerArrayPrefix = TEXT("SamplerDescriptorHeap_");
 };
 
 namespace EVulkanBindingType
