Fixed validation error in transient allocator due to transient allocator being destroyed prior to the CustomDepth texture being destroyed.

diff --git a/Engine/Source/Runtime/RenderCore/Private/RenderGraphResourcePool.cpp b/Engine/Source/Runtime/RenderCore/Private/RenderGraphResourcePool.cpp
--- a/Engine/Source/Runtime/RenderCore/Private/RenderGraphResourcePool.cpp
+++ b/Engine/Source/Runtime/RenderCore/Private/RenderGraphResourcePool.cpp
@@ -347,4 +347,4 @@ void FRDGTransientResourceAllocator::ReleasePendingDeallocations()
 	}
 }
 
-TGlobalResource<FRDGTransientResourceAllocator> GRDGTransientResourceAllocator;
+TGlobalResource<FRDGTransientResourceAllocator, FRenderResource::EInitPhase::Pre> GRDGTransientResourceAllocator;
diff --git a/Engine/Source/Runtime/RenderCore/Private/RenderGraphResourcePool.h b/Engine/Source/Runtime/RenderCore/Private/RenderGraphResourcePool.h
--- a/Engine/Source/Runtime/RenderCore/Private/RenderGraphResourcePool.h
+++ b/Engine/Source/Runtime/RenderCore/Private/RenderGraphResourcePool.h
@@ -110,4 +110,4 @@ private:
 	friend class FRDGTransientRenderTarget;
 };
 
-extern RENDERCORE_API TGlobalResource<FRDGTransientResourceAllocator> GRDGTransientResourceAllocator;
+extern RENDERCORE_API TGlobalResource<FRDGTransientResourceAllocator, FRenderResource::EInitPhase::Pre> GRDGTransientResourceAllocator;
