Fix out of memory crashes when using DumpGPU.

The views made by DumpGPU must be created with lifetime extension disabled, so they release the underlying resources immediately after each dump pass.

#rnx

diff --git a/Engine/Source/Runtime/RenderCore/Private/DumpGPU.cpp b/Engine/Source/Runtime/RenderCore/Private/DumpGPU.cpp
--- a/Engine/Source/Runtime/RenderCore/Private/DumpGPU.cpp
+++ b/Engine/Source/Runtime/RenderCore/Private/DumpGPU.cpp
@@ -1016,6 +1016,20 @@ public:
 		(*Texture)->DisableLifetimeExtension();
 	}
 
+	FShaderResourceViewRHIRef CreateSRVNoLifetimeExtension(FRHICommandList& RHICmdList, FRHITexture* Texture, const FRHITextureSRVCreateInfo& Desc)
+	{
+		FShaderResourceViewRHIRef SRV = RHICmdList.CreateShaderResourceView(Texture, Desc);
+		SRV->DisableLifetimeExtension();
+		return SRV;
+	}
+
+	FUnorderedAccessViewRHIRef CreateUAVNoLifetimeExtension(FRHICommandList& RHICmdList, FRHITexture* Texture, uint32 MipLevel)
+	{
+		FUnorderedAccessViewRHIRef UAV = RHICmdList.CreateUnorderedAccessView(Texture, MipLevel);
+		UAV->DisableLifetimeExtension();
+		return UAV;
+	}
+
 	FStagingTexturePoolEntry& ReuseStagingTexture(
 		FRHICommandListImmediate& RHICmdList,
 		const FRHITextureCreateDesc& Desc,
@@ -1072,7 +1086,7 @@ public:
 		if (UAV)
 		{
 			check(Access == ERHIAccess::UAVCompute);
-			NewEntry->UAV = RHICmdList.CreateUnorderedAccessView(*Texture, /* MipLevel = */ 0);
+			NewEntry->UAV = CreateUAVNoLifetimeExtension(RHICmdList, *Texture, /* MipLevel = */ 0);
 			*UAV = NewEntry->UAV;
 		}
 		else
@@ -1155,7 +1169,7 @@ public:
 						CreateStagingTexture(Desc, /* out */ &StagingSrcTextureRef);
 
 						RHICmdList.Transition(FRHITransitionInfo(StagingSrcTextureRef, ERHIAccess::Unknown, ERHIAccess::UAVCompute));
-						StagingOutput = RHICmdList.CreateUnorderedAccessView(StagingSrcTextureRef, /* MipLevel = */ 0);
+						StagingOutput = CreateUAVNoLifetimeExtension(RHICmdList, StagingSrcTextureRef, /* MipLevel = */ 0);
 					}
 					StagingSrcTexture = StagingSrcTextureRef;
 				}
@@ -1288,7 +1302,7 @@ public:
 		FShaderResourceViewRHIRef SubResourceSRV;
 		if (SubresourceDumpDesc.bPreprocessForStaging)
 		{
-			SubResourceSRV = RHICmdList.CreateShaderResourceView(RHITexture, FRHITextureSRVCreateInfo(SubresourceDesc));
+			SubResourceSRV = CreateSRVNoLifetimeExtension(RHICmdList, RHITexture, FRHITextureSRVCreateInfo(SubresourceDesc));
 			RHICmdListImmediate.Transition(FRHITransitionInfo(RHITexture, ERHIAccess::Unknown, ERHIAccess::SRVCompute));
 		}
 		else
