Temporarily disabled warning for RealtimeGPUProfiler query been unavailable

diff --git a/Engine/Source/Runtime/RenderCore/Private/ProfilingDebugging/RealtimeGPUProfiler.cpp b/Engine/Source/Runtime/RenderCore/Private/ProfilingDebugging/RealtimeGPUProfiler.cpp
--- a/Engine/Source/Runtime/RenderCore/Private/ProfilingDebugging/RealtimeGPUProfiler.cpp
+++ b/Engine/Source/Runtime/RenderCore/Private/ProfilingDebugging/RealtimeGPUProfiler.cpp
@@ -539,9 +539,10 @@ public:
 
 			if (!Event.GatherQueryResults(RHICmdList))
 			{
-#if !(UE_BUILD_SHIPPING || UE_BUILD_TEST)
-				UE_LOG(LogRendererCore, Warning, TEXT("Query '%s' not ready."), *Event.GetName().ToString());
-#endif
+                // TODO: clloyd - Commmented out for now to stop spam on Mac
+//#if !(UE_BUILD_SHIPPING || UE_BUILD_TEST)
+				//UE_LOG(LogRendererCore, Warning, TEXT("Query '%s' not ready."), *Event.GetName().ToString());
+//#endif
 				// The frame isn't ready yet. Don't update stats - we'll try again next frame. 
 				return false;
 			}
