- Add ShaderCodeLibrary api to ignore some plugins on mount (so that they can be loaded in parallel for several plugins)
- Add an optional param to OpenLibrary to not look for shader lib chunks (because there's strong thread contention surrounding FMountedPakFileInfo::KnownPakFilesAccessLock)
- Removed reading of shared-cooked override since that feature is obsolete
#rnx

diff --git a/Engine/Source/Runtime/RenderCore/Private/ShaderCodeLibrary.cpp b/Engine/Source/Runtime/RenderCore/Private/ShaderCodeLibrary.cpp
--- a/Engine/Source/Runtime/RenderCore/Private/ShaderCodeLibrary.cpp
+++ b/Engine/Source/Runtime/RenderCore/Private/ShaderCodeLibrary.cpp
@@ -94,6 +94,8 @@ namespace UE
 			static FDelegateHandle OnPluginMountedDelegateHandle;
 			static FDelegateHandle OnPluginUnmountedDelegateHandle;
 
+			static TSet<FString> PluginsToIgnoreOnMount;
+
 			/** Helper function shared between the cooker and runtime */
 			FString GetShaderLibraryNameForChunk(FString const& BaseName, int32 ChunkId)
 			{
@@ -2391,6 +2393,8 @@ class FShaderLibrariesCollection
 	FShaderCodeStats EditorShaderCodeStats[EShaderPlatform::SP_NumPlatforms];
 	// At cook time, whether the shader archive supports pipelines (only OpenGL should)
 	bool EditorArchivePipelines[EShaderPlatform::SP_NumPlatforms];
+	bool bIsEditorShaderCodeArchiveEmpty = true;
+	bool bIsEditorShaderStableInfoEmpty = true;
 #endif //WITH_EDITOR
 	bool bSupportsPipelines;
 	bool bNativeFormat;
@@ -2472,8 +2476,10 @@ public:
 #endif
 	}
 
-	bool OpenLibrary(FString const& Name, FString const& Directory)
+	bool OpenLibrary(FString const& Name, FString const& Directory, const bool bMonolithicOnly = false)
 	{
+		TRACE_CPUPROFILER_EVENT_SCOPE(OpenShaderLibrary);
+
 		using namespace UE::ShaderLibrary::Private;
 
 		bool bResult = false;
@@ -2514,12 +2520,8 @@ public:
 			{
 				bResult = true;
 				bOpenedAsMonolithic = true;
-
-				// Attempt to open the shared-cooked override code library if there is one.
-				// This is probably not ideal, but it should get shared-cooks working.
-				Library->OpenShaderCode(Directory, Name + TEXT("_SC"));
 			}
-			else // attempt to open a chunked library
+			else if (!bMonolithicOnly) // attempt to open a chunked library
 			{
 				TSet<int32> PrevComponentSet = Library->PresentChunks;
 				{
@@ -2637,20 +2639,26 @@ public:
 		}
 
 #if WITH_EDITOR
-		for (uint32 i = 0; i < EShaderPlatform::SP_NumPlatforms; i++)
+		if (!bIsEditorShaderCodeArchiveEmpty)
 		{
-			FEditorShaderCodeArchive* CodeArchive = EditorShaderCodeArchive[i];
-			if (CodeArchive)
+			for (uint32 i = 0; i < EShaderPlatform::SP_NumPlatforms; i++)
 			{
-				CodeArchive->OpenLibrary(Name);
+				FEditorShaderCodeArchive* CodeArchive = EditorShaderCodeArchive[i];
+				if (CodeArchive)
+				{
+					CodeArchive->OpenLibrary(Name);
+				}
 			}
 		}
-		for (uint32 i = 0; i < EShaderPlatform::SP_NumPlatforms; i++)
+		if (!bIsEditorShaderStableInfoEmpty)
 		{
-			FEditorShaderStableInfo* StableArchive = EditorShaderStableInfo[i];
-			if (StableArchive)
+			for (uint32 i = 0; i < EShaderPlatform::SP_NumPlatforms; i++)
 			{
-				StableArchive->OpenLibrary(Name);
+				FEditorShaderStableInfo* StableArchive = EditorShaderStableInfo[i];
+				if (StableArchive)
+				{
+					StableArchive->OpenLibrary(Name);
+				}
 			}
 		}
 #endif
@@ -2896,6 +2904,7 @@ public:
 			FEditorShaderCodeArchive* CodeArchive = EditorShaderCodeArchive[Platform];
 			if (!CodeArchive)
 			{
+				bIsEditorShaderCodeArchiveEmpty = false;
 				CodeArchive = new FEditorShaderCodeArchive(FName(FormatAndPlatformName), Descriptor.bNeedsDeterministicOrder);
 				EditorShaderCodeArchive[Platform] = CodeArchive;
 				EditorArchivePipelines[Platform] = !bNativeFormat;
@@ -2917,6 +2926,7 @@ public:
 			FEditorShaderStableInfo* StableArchive = EditorShaderStableInfo[Platform];
 			if (!StableArchive && bUseStableKeys)
 			{
+				bIsEditorShaderStableInfoEmpty = false;
 				StableArchive = new FEditorShaderStableInfo(PossiblyAdjustedFormat);
 				EditorShaderStableInfo[Platform] = StableArchive;
 				ShaderFormatsThatNeedStableKeys[(int)Platform] = true;
@@ -3298,10 +3308,9 @@ FShaderLibrariesCollection* FShaderLibrariesCollection::Impl = nullptr;
 
 static void FShaderCodeLibraryPluginMountedCallback(IPlugin& Plugin)
 {
-	if (Plugin.CanContainContent() && Plugin.IsEnabled())
+	if (UE::ShaderLibrary::Private::PluginsToIgnoreOnMount.Remove(Plugin.GetName()) == 0)
 	{
-		// load any shader libraries that may exist in this plugin
-		FShaderCodeLibrary::OpenLibrary(Plugin.GetName(), Plugin.GetContentDir());
+		FShaderCodeLibrary::OpenPluginShaderLibrary(Plugin);
 	}
 }
 
@@ -3398,7 +3407,7 @@ void FShaderCodeLibrary::InitForRuntime(EShaderPlatform ShaderPlatform)
 			auto Plugins = IPluginManager::Get().GetEnabledPluginsWithContent();
 			ParallelFor(Plugins.Num(), [&](int32 Index)
 			{
-				FShaderCodeLibraryPluginMountedCallback(*Plugins[Index]);
+				FShaderCodeLibrary::OpenPluginShaderLibrary(*Plugins[Index]);
 			});
 		}
 		else
@@ -3443,6 +3452,8 @@ void FShaderCodeLibrary::Shutdown()
 		UE::ShaderLibrary::Private::OnPluginUnmountedDelegateHandle.Reset();
 	}
 
+	UE::ShaderLibrary::Private::PluginsToIgnoreOnMount.Empty();
+
 	if (FShaderLibrariesCollection::Impl)
 	{
 		delete FShaderLibrariesCollection::Impl;
@@ -3617,12 +3628,12 @@ void FShaderCodeLibrary::AddKnownChunkIDs(const int32* IDs, const int32 NumChunk
 	}
 }
 
-bool FShaderCodeLibrary::OpenLibrary(FString const& Name, FString const& Directory)
+bool FShaderCodeLibrary::OpenLibrary(FString const& Name, FString const& Directory, bool bMonolithicOnly)
 {
 	bool bResult = false;
 	if (FShaderLibrariesCollection::Impl)
 	{
-		bResult = FShaderLibrariesCollection::Impl->OpenLibrary(Name, Directory);
+		bResult = FShaderLibrariesCollection::Impl->OpenLibrary(Name, Directory, bMonolithicOnly);
 	}
 	return bResult;
 }
@@ -3872,6 +3883,21 @@ void FShaderCodeLibrary::UnregisterSharedShaderCodeRequestDelegate_Handle(FDeleg
 	OnSharedShaderCodeRequest.Remove(Handle);
 }
 
+void FShaderCodeLibrary::DontOpenPluginShaderLibraryOnMount(const FString& PluginName)
+{
+	check(IsInGameThread());
+	UE::ShaderLibrary::Private::PluginsToIgnoreOnMount.Add(PluginName);
+}
+
+void FShaderCodeLibrary::OpenPluginShaderLibrary(IPlugin& Plugin, bool bMonolithicOnly)
+{
+	if (Plugin.CanContainContent() && Plugin.IsEnabled())
+	{
+		// load any shader libraries that may exist in this plugin
+		FShaderCodeLibrary::OpenLibrary(Plugin.GetName(), Plugin.GetContentDir(), bMonolithicOnly);
+	}
+}
+
 // FNamedShaderLibrary methods
 
 // At runtime, open shader code collection for specified shader platform. Returns true if new code was opened
diff --git a/Engine/Source/Runtime/RenderCore/Public/ShaderCodeLibrary.h b/Engine/Source/Runtime/RenderCore/Public/ShaderCodeLibrary.h
--- a/Engine/Source/Runtime/RenderCore/Public/ShaderCodeLibrary.h
+++ b/Engine/Source/Runtime/RenderCore/Public/ShaderCodeLibrary.h
@@ -35,6 +35,7 @@ class FShaderPipeline;
 using FShaderMapAssetPaths = TSet<FName>;
 class FIoBuffer;
 class FIoChunkId;
+class IPlugin;
 class UObject;
 
 #if WITH_EDITOR
@@ -210,9 +211,10 @@ struct FShaderCodeLibrary
 	 * At runtime this will open the shader library with this name.
 	 * @param Name is a high level description of the library (usually a project name or "Global")
 	 * @param Directory location of the .ushadercode file
+	 * @param bMonolithicOnly If true, only attempt to open a monolithic library (no chunks)
 	 * @return true if successful
 	 */
-	static RENDERCORE_API bool OpenLibrary(FString const& Name, FString const& Directory);
+	static RENDERCORE_API bool OpenLibrary(FString const& Name, FString const& Directory, bool bMonolithicOnly = false);
 
 	/**
 	 * Close a named library.
@@ -249,6 +251,13 @@ struct FShaderCodeLibrary
 	// Delegate called whenever shader code is requested.
 	static RENDERCORE_API FDelegateHandle RegisterSharedShaderCodeRequestDelegate_Handle(const FSharedShaderCodeRequest::FDelegate& Delegate);
 	static RENDERCORE_API void UnregisterSharedShaderCodeRequestDelegate_Handle(FDelegateHandle Handle);
+
+	// Disables opening the specified plugin's shader library on mount
+	static RENDERCORE_API void DontOpenPluginShaderLibraryOnMount(const FString& PluginName);
+	
+	// Open the plugin's shader library
+	// @param bMonolithicOnly If true, only attempt to open a monolithic library (no chunks)
+	static RENDERCORE_API void OpenPluginShaderLibrary(IPlugin& Plugin, bool bMonolithicOnly = false);
 };
 
 #if WITH_EDITOR
