[shaders] minor include cleanup

diff --git a/Engine/Source/Developer/Apple/MetalShaderFormat/Private/MetalDerivedData.cpp b/Engine/Source/Developer/Apple/MetalShaderFormat/Private/MetalDerivedData.cpp
--- a/Engine/Source/Developer/Apple/MetalShaderFormat/Private/MetalDerivedData.cpp
+++ b/Engine/Source/Developer/Apple/MetalShaderFormat/Private/MetalDerivedData.cpp
@@ -1,19 +1,19 @@
 // Copyright Epic Games, Inc. All Rights Reserved.
 
 #include "MetalDerivedData.h"
-#include "Serialization/MemoryWriter.h"
-#include "RHIDefinitions.h"
-#include "Misc/FileHelper.h"
+
+#include "HlslccHeaderWriter.h"
+#include "MetalShaderFormat.h"
 #include "MetalShaderResources.h"
+#include "Misc/FileHelper.h"
 #include "Misc/Paths.h"
 #include "Misc/Compression.h"
 #include "Misc/OutputDeviceRedirector.h"
 #include "MetalBackend.h"
+#include "RHIDefinitions.h"
 #include "Serialization/MemoryReader.h"
-#include "HlslccHeaderWriter.h"
-#include "Runtime/RenderCore/Internal/ShaderCompilerDefinitions.h"
-
-#include "MetalShaderFormat.h"
+#include "Serialization/MemoryWriter.h"
+#include "ShaderCompilerDefinitions.h"
 #include "SpirvReflectCommon.h"
 
 #include <regex>
diff --git a/Engine/Source/Developer/Apple/MetalShaderFormat/Private/MetalShaderCompiler.cpp b/Engine/Source/Developer/Apple/MetalShaderFormat/Private/MetalShaderCompiler.cpp
--- a/Engine/Source/Developer/Apple/MetalShaderFormat/Private/MetalShaderCompiler.cpp
+++ b/Engine/Source/Developer/Apple/MetalShaderFormat/Private/MetalShaderCompiler.cpp
@@ -1,21 +1,20 @@
 // Copyright Epic Games, Inc. All Rights Reserved.
-// ..
 
 #include "CoreMinimal.h"
-#include "ShaderCore.h"
+#include "HAL/PlatformFileManager.h"
+#include "MetalShaderFormat.h"
 #include "MetalShaderResources.h"
+#include "Misc/ConfigCacheIni.h"
+#include "Misc/EngineVersion.h"
+#include "Misc/FileHelper.h"
+#include "Misc/Paths.h"
 #include "ShaderCompilerCommon.h"
+#include "ShaderCompilerDefinitions.h"
+#include "ShaderCore.h"
 #include "ShaderParameterParser.h"
-#include "Misc/ConfigCacheIni.h"
 #include "Serialization/MemoryWriter.h"
 #include "Serialization/MemoryReader.h"
-#include "Misc/FileHelper.h"
-#include "Misc/Paths.h"
-#include "Misc/EngineVersion.h"
-#include "HAL/PlatformFileManager.h"
-#include "Runtime/RenderCore/Internal/ShaderCompilerDefinitions.h"
 
-#include "MetalShaderFormat.h"
 
 #if PLATFORM_WINDOWS
 #include "Windows/WindowsHWrapper.h"
diff --git a/Engine/Source/Developer/ShaderCompilerCommon/Private/ShaderConductorContext.cpp b/Engine/Source/Developer/ShaderCompilerCommon/Private/ShaderConductorContext.cpp
--- a/Engine/Source/Developer/ShaderCompilerCommon/Private/ShaderConductorContext.cpp
+++ b/Engine/Source/Developer/ShaderCompilerCommon/Private/ShaderConductorContext.cpp
@@ -2,7 +2,7 @@
 
 #include "ShaderConductorContext.h"
 #include "HAL/ExceptionHandling.h"
-#include "Runtime/RenderCore/Internal/ShaderCompilerDefinitions.h"
+#include "ShaderCompilerDefinitions.h"
 
 #if PLATFORM_MAC || PLATFORM_WINDOWS || PLATFORM_LINUX
 THIRD_PARTY_INCLUDES_START
diff --git a/Engine/Source/Developer/ShaderFormatOpenGL/Private/OpenGLShaderCompiler.cpp b/Engine/Source/Developer/ShaderFormatOpenGL/Private/OpenGLShaderCompiler.cpp
--- a/Engine/Source/Developer/ShaderFormatOpenGL/Private/OpenGLShaderCompiler.cpp
+++ b/Engine/Source/Developer/ShaderFormatOpenGL/Private/OpenGLShaderCompiler.cpp
@@ -8,12 +8,12 @@
 #include "Misc/FileHelper.h"
 #include "Misc/Paths.h"
 #include "Serialization/MemoryWriter.h"
+#include "ShaderCompilerDefinitions.h"
 #include "ShaderFormatOpenGL.h"
 #include "ShaderParameterParser.h"
 #include "SpirvReflectCommon.h"
 #include <algorithm>
 #include <regex>
-#include "Runtime/RenderCore/Internal/ShaderCompilerDefinitions.h"
 
 #if PLATFORM_WINDOWS
 #include "Windows/AllowWindowsPlatformTypes.h"
diff --git a/Engine/Source/Developer/ShaderPreprocessor/Private/ShaderPreprocessor.cpp b/Engine/Source/Developer/ShaderPreprocessor/Private/ShaderPreprocessor.cpp
--- a/Engine/Source/Developer/ShaderPreprocessor/Private/ShaderPreprocessor.cpp
+++ b/Engine/Source/Developer/ShaderPreprocessor/Private/ShaderPreprocessor.cpp
@@ -6,7 +6,7 @@
 #include "Misc/ScopeLock.h"
 #include "Modules/ModuleManager.h"
 #include "PreprocessorPrivate.h"
-#include "Runtime/RenderCore/Internal/ShaderCompilerDefinitions.h"
+#include "ShaderCompilerDefinitions.h"
 
 #include "stb_preprocess/preprocessor.h"
 #include "stb_preprocess/stb_alloc.h"
diff --git a/Engine/Source/Developer/VulkanShaderFormat/Private/VulkanShaderCompiler.cpp b/Engine/Source/Developer/VulkanShaderFormat/Private/VulkanShaderCompiler.cpp
--- a/Engine/Source/Developer/VulkanShaderFormat/Private/VulkanShaderCompiler.cpp
+++ b/Engine/Source/Developer/VulkanShaderFormat/Private/VulkanShaderCompiler.cpp
@@ -6,8 +6,8 @@
 #include "HlslccHeaderWriter.h"
 #include "hlslcc.h"
 #include "RHIShaderFormatDefinitions.inl"
-#include "Runtime/RenderCore/Internal/ShaderCompilerDefinitions.h"
 #include "ShaderCompilerCommon.h"
+#include "ShaderCompilerDefinitions.h"
 #include "ShaderParameterParser.h"
 #include "ShaderPreprocessor.h"
 #include "ShaderPreprocessTypes.h"
diff --git a/Engine/Source/Developer/Windows/ShaderFormatD3D/Private/D3DShaderCompiler.cpp b/Engine/Source/Developer/Windows/ShaderFormatD3D/Private/D3DShaderCompiler.cpp
--- a/Engine/Source/Developer/Windows/ShaderFormatD3D/Private/D3DShaderCompiler.cpp
+++ b/Engine/Source/Developer/Windows/ShaderFormatD3D/Private/D3DShaderCompiler.cpp
@@ -1,19 +1,19 @@
 // Copyright Epic Games, Inc. All Rights Reserved.
 
+#include "D3D12RHI.h"
+#include "HAL/FileManager.h"
+#include "Misc/Paths.h"
+#include "Misc/FileHelper.h"
+#include "RayTracingDefinitions.h"
+#include "Serialization/MemoryWriter.h"
 #include "ShaderFormatD3D.h"
 #include "ShaderPreprocessor.h"
 #include "ShaderCompilerCommon.h"
+#include "ShaderCompilerDefinitions.h"
 #include "ShaderMinifier.h"
 #include "ShaderParameterParser.h"
 #include "ShaderPreprocessTypes.h"
-#include "D3D12RHI.h"
-#include "Misc/Paths.h"
-#include "Misc/FileHelper.h"
-#include "HAL/FileManager.h"
-#include "Serialization/MemoryWriter.h"
-#include "RayTracingDefinitions.h"
 #include "SpirvCommon.h"
-#include "Runtime/RenderCore/Internal/ShaderCompilerDefinitions.h"
 
 DEFINE_LOG_CATEGORY_STATIC(LogD3D11ShaderCompiler, Log, All);
 
diff --git a/Engine/Source/Runtime/Engine/Private/ShaderCompiler/ShaderGenerationUtil.cpp b/Engine/Source/Runtime/Engine/Private/ShaderCompiler/ShaderGenerationUtil.cpp
--- a/Engine/Source/Runtime/Engine/Private/ShaderCompiler/ShaderGenerationUtil.cpp
+++ b/Engine/Source/Runtime/Engine/Private/ShaderCompiler/ShaderGenerationUtil.cpp
@@ -8,14 +8,14 @@
 
 #include "DataDrivenShaderPlatformInfo.h"
 #include "HAL/PlatformFile.h"
+#include "HAL/PlatformFileManager.h"
 #include "Interfaces/ITargetPlatform.h"
-#include "Misc/FileHelper.h"
 #include "Interfaces/ITargetPlatformManagerModule.h"
-#include "HAL/PlatformFileManager.h"
+#include "Misc/FileHelper.h"
 #include "RenderUtils.h"
 #include "SceneManagement.h"
+#include "ShaderCompilerDefinitions.h"
 #include "ShaderMaterial.h"
-#include "Runtime/RenderCore/Internal/ShaderCompilerDefinitions.h"
 
 static TAutoConsoleVariable<int32> CVarShaderUseGBufferRefactor(
 	TEXT("r.Shaders.UseGBufferRefactor"),
diff --git a/Engine/Source/Runtime/RenderCore/Private/ShaderCore.cpp b/Engine/Source/Runtime/RenderCore/Private/ShaderCore.cpp
--- a/Engine/Source/Runtime/RenderCore/Private/ShaderCore.cpp
+++ b/Engine/Source/Runtime/RenderCore/Private/ShaderCore.cpp
@@ -7,37 +7,37 @@
 #include "ShaderCore.h"
 #include "Algo/Find.h"
 #include "Async/ParallelFor.h"
+#include "Compression/OodleDataCompression.h"
+#include "DataDrivenShaderPlatformInfo.h"
 #include "HAL/FileManager.h"
 #include "HAL/PlatformStackWalk.h"
+#include "Interfaces/IShaderFormat.h"
+#include "Interfaces/IShaderFormatModule.h"
+#include "Interfaces/ITargetPlatform.h"
+#include "Interfaces/ITargetPlatformManagerModule.h"
 #include "Math/BigInt.h"
 #include "Misc/CommandLine.h"
 #include "Misc/FileHelper.h"
 #include "Misc/Compression.h"
+#include "Misc/CoreMisc.h"
 #include "Misc/Paths.h"
 #include "Misc/PathViews.h"
 #include "Misc/ScopeLock.h"
-#include "Stats/StatsMisc.h"
 #include "Misc/ScopeRWLock.h"
 #include "Misc/ScopedSlowTask.h"
 #include "Misc/StringBuilder.h"
+#include "Modules/ModuleManager.h"
+#include "RHIShaderFormatDefinitions.inl"
 #include "Serialization/MemoryHasher.h"
+#include "Serialization/MemoryWriter.h"
 #include "Shader.h"
 #include "ShaderCompilerCore.h"
+#include "ShaderCompilerDefinitions.h"
 #include "ShaderCompilerJobTypes.h"
+#include "Stats/StatsMisc.h"
 #include "String/Find.h"
 #include "Tasks/Task.h"
 #include "VertexFactory.h"
-#include "Modules/ModuleManager.h"
-#include "Interfaces/IShaderFormat.h"
-#include "Interfaces/IShaderFormatModule.h"
-#include "RHIShaderFormatDefinitions.inl"
-#include "DataDrivenShaderPlatformInfo.h"
-#include "Compression/OodleDataCompression.h"
-#include "Serialization/MemoryWriter.h"
-#include "Misc/CoreMisc.h"
-#include "Interfaces/ITargetPlatform.h"
-#include "Interfaces/ITargetPlatformManagerModule.h"
-#include "Runtime/RenderCore/Internal/ShaderCompilerDefinitions.h"
 #if WITH_EDITOR
 #include "Misc/ConfigCacheIni.h"
 #include "Misc/StringBuilder.h"
diff --git a/Engine/Source/Runtime/RenderCore/Public/ShaderCore.h b/Engine/Source/Runtime/RenderCore/Public/ShaderCore.h
--- a/Engine/Source/Runtime/RenderCore/Public/ShaderCore.h
+++ b/Engine/Source/Runtime/RenderCore/Public/ShaderCore.h
@@ -42,7 +42,7 @@
 #include "UniformBuffer.h"
 
 // Temporarily included here until we can fully deprecate access to FShaderCompilerDefinitions in a future version
-#include "Runtime/RenderCore/Internal/ShaderCompilerDefinitions.h"
+#include "ShaderCompilerDefinitions.h"
 
 class Error;
 class FMemoryImageWriter;
