Add a new engine function to union shader byte code made from different engine cooks
New funciton takes in multiple meta data directories and writes to a desired directory

diff --git a/Engine/Source/Editor/UnrealEd/Private/CookOnTheFlyServer.cpp b/Engine/Source/Editor/UnrealEd/Private/CookOnTheFlyServer.cpp
--- a/Engine/Source/Editor/UnrealEd/Private/CookOnTheFlyServer.cpp
+++ b/Engine/Source/Editor/UnrealEd/Private/CookOnTheFlyServer.cpp
@@ -9108,19 +9108,29 @@ void UCookOnTheFlyServer::RegisterShaderChunkDataGenerator()
 
 FString UCookOnTheFlyServer::GetProjectShaderLibraryName() const
 {
-	if (!IsCookingDLC())
+	static FString ShaderLibraryName = [this]()
 	{
-		FString Result = FApp::GetProjectName();
-		if (Result.IsEmpty())
+		FString OverrideShaderLibraryName;
+		if (FParse::Value(FCommandLine::Get(), TEXT("OverrideShaderLibraryName="), OverrideShaderLibraryName))
 		{
-			Result = TEXT("UnrealGame");
+			return OverrideShaderLibraryName;
 		}
-		return Result;
-	}
-	else
-	{
-		return CookByTheBookOptions->DlcName;
-	}
+
+		if (!IsCookingDLC())
+		{
+			FString Result = FApp::GetProjectName();
+			if (Result.IsEmpty())
+			{
+				Result = TEXT("UnrealGame");
+			}
+			return Result;
+		}
+		else
+		{
+			return CookByTheBookOptions->DlcName;
+		}
+	}();
+	return ShaderLibraryName;
 }
 
 static FString GenerateShaderCodeLibraryName(FString const& Name, bool bIsIterateSharedBuild)
diff --git a/Engine/Source/Editor/UnrealEd/Private/ShaderLibraryChunkDataGenerator.cpp b/Engine/Source/Editor/UnrealEd/Private/ShaderLibraryChunkDataGenerator.cpp
--- a/Engine/Source/Editor/UnrealEd/Private/ShaderLibraryChunkDataGenerator.cpp
+++ b/Engine/Source/Editor/UnrealEd/Private/ShaderLibraryChunkDataGenerator.cpp
@@ -18,13 +18,16 @@ FShaderLibraryChunkDataGenerator::FShaderLibraryChunkDataGenerator(UCookOnTheFly
 	bOptedOut = false;
 	PlatformNameUsedForIni = TargetPlatform->IniPlatformName();
 
+	bool bChunkShadersWhenCookingDLC = false;
+
 	FConfigFile PlatformIniFile;
 	FConfigCacheIni::LoadLocalIniFile(PlatformIniFile, TEXT("Engine"), true, *PlatformNameUsedForIni);
 	PlatformIniFile.GetBool(TEXT("DevOptions.Shaders"), TEXT("bDoNotChunkShaderLib"), bOptedOut);
+	PlatformIniFile.GetBool(TEXT("DevOptions.Shaders"), TEXT("bChunkShaderWhenCookingDLC"), bChunkShadersWhenCookingDLC);
 
 	// Disable chunking for DLC - this causes problems as the main game can be optionally (for faster iteration) cooked with -fastcook. Fastcook disables chunking,
 	// so the game has no idea about ChunkIDs and cannot find DLC's chunked libs. If DLC lib is monolithic, both monolithic and chunked games will try to open it.
-	if (COTFS.IsCookingDLC())
+	if (!bChunkShadersWhenCookingDLC && COTFS.IsCookingDLC())
 	{
 		bOptedOut = true;
 	}
diff --git a/Engine/Source/Runtime/RenderCore/Private/ShaderCodeLibrary.cpp b/Engine/Source/Runtime/RenderCore/Private/ShaderCodeLibrary.cpp
--- a/Engine/Source/Runtime/RenderCore/Private/ShaderCodeLibrary.cpp
+++ b/Engine/Source/Runtime/RenderCore/Private/ShaderCodeLibrary.cpp
@@ -18,6 +18,7 @@ ShaderCodeLibrary.cpp: Bound shader state cache implementation.
 #include "HAL/PlatformSplash.h" // IWYU pragma: keep
 #include "Hash/CityHash.h"
 #include "Interfaces/IPluginManager.h"
+#include "Internationalization/Regex.h"
 #include "Math/UnitConversion.h"
 #include "Misc/CommandLine.h"
 #include "Misc/ConfigCacheIni.h"
@@ -1750,6 +1751,48 @@ struct FEditorShaderCodeArchive
 		return bOK;
 	}
 
+	void AddShaderCodeLibraryFromDirectory(const FString& BaseDir)
+	{
+		if (FArchive* PrevCookedAr = IFileManager::Get().CreateFileReader(*GetCodeArchiveFilename(BaseDir, LibraryName, FormatName)))
+		{
+			uint32 Version = 0;
+			*PrevCookedAr << Version;
+
+			if (Version == GShaderCodeArchiveVersion)
+			{
+				FSerializedShaderArchive PrevCookedShaders;
+
+				*PrevCookedAr << PrevCookedShaders;
+
+				// check if it also contains the asset info file
+				if (PrevCookedShaders.LoadAssetInfo(GetShaderAssetInfoFilename(BaseDir, LibraryName, FormatName)))
+				{
+					UE_LOG(LogShaderLibrary, Display, TEXT("Loaded asset info %s for the shader library %s: %d entries"),
+						*GetShaderAssetInfoFilename(BaseDir, LibraryName, FormatName),
+						*GetCodeArchiveFilename(BaseDir, LibraryName, FormatName),
+						PrevCookedShaders.ShaderCodeToAssets.Num()
+					);
+				}
+				else
+				{
+					UE_LOG(LogShaderLibrary, Warning, TEXT("Could not find or load asset info %s for the shader library %s"),
+						*GetShaderAssetInfoFilename(BaseDir, LibraryName, FormatName),
+						*GetCodeArchiveFilename(BaseDir, LibraryName, FormatName)
+					);
+				}
+
+				int64 PrevCookedShadersCodeStart = PrevCookedAr->Tell();
+				for (int32 PrevShaderMapIndex = 0; PrevShaderMapIndex < PrevCookedShaders.ShaderMapEntries.Num(); ++PrevShaderMapIndex)
+				{
+					AddShaderCode(PrevShaderMapIndex, PrevCookedShaders, PrevCookedShadersCodeStart, *PrevCookedAr);
+				}
+			}
+
+			PrevCookedAr->Close();
+			delete PrevCookedAr;
+		}
+	}
+
 	void AddExistingShaderCodeLibrary(FString const& OutputDir)
 	{
 		check(LibraryName.Len() > 0);
@@ -1763,46 +1806,7 @@ struct FEditorShaderCodeArchive
 		{
 			if (ShaderFileName.Contains(LibraryName + TEXT("-") + FormatName.ToString() + TEXT(".")))
 			{
-				FArchive* PrevCookedAr = IFileManager::Get().CreateFileReader(*GetCodeArchiveFilename(OutputDir, LibraryName, FormatName));
-
-				if (PrevCookedAr)
-				{
-					uint32 Version = 0;
-					*PrevCookedAr << Version;
-
-					if (Version == GShaderCodeArchiveVersion)
-					{
-						FSerializedShaderArchive PrevCookedShaders;
-
-						*PrevCookedAr << PrevCookedShaders;
-
-						// check if it also contains the asset info file
-						if (PrevCookedShaders.LoadAssetInfo(GetShaderAssetInfoFilename(OutputDir, LibraryName, FormatName)))
-						{
-							UE_LOG(LogShaderLibrary, Display, TEXT("Loaded asset info %s for the shader library %s: %d entries"),
-								*GetShaderAssetInfoFilename(OutputDir, LibraryName, FormatName),
-								*GetCodeArchiveFilename(OutputDir, LibraryName, FormatName),
-								PrevCookedShaders.ShaderCodeToAssets.Num()
-								);
-						}
-						else
-						{
-							UE_LOG(LogShaderLibrary, Warning, TEXT("Could not find or load asset info %s for the shader library %s"),
-								*GetShaderAssetInfoFilename(OutputDir, LibraryName, FormatName),
-								*GetCodeArchiveFilename(OutputDir, LibraryName, FormatName)
-							);
-						}
-
-						int64 PrevCookedShadersCodeStart = PrevCookedAr->Tell();
-						for (int32 PrevShaderMapIndex = 0; PrevShaderMapIndex < PrevCookedShaders.ShaderMapEntries.Num(); ++PrevShaderMapIndex)
-						{
-							AddShaderCode(PrevShaderMapIndex, PrevCookedShaders, PrevCookedShadersCodeStart, *PrevCookedAr);
-						}
-					}
-
-					PrevCookedAr->Close();
-					delete PrevCookedAr;
-				}
+				AddShaderCodeLibraryFromDirectory(OutputDir);
 			}
 		}
 	}
@@ -2160,12 +2164,23 @@ struct FEditorShaderStableInfo
 		MarkKeyValueDirty(StableKeyValue);
 	}
 
+	void AddShaderCodeLibraryFromDirectory(const FString& BaseDir)
+	{
+		TArray<FStableShaderKeyAndValue> StableKeys;
+		if (UE::PipelineCacheUtilities::LoadStableKeysFile(GetStableInfoArchiveFilename(BaseDir, LibraryName, FormatName), StableKeys))
+		{
+			for (FStableShaderKeyAndValue& Item : StableKeys)
+			{
+				AddShader(Item);
+			}
+		}
+	}
+
 	void AddExistingShaderCodeLibrary(FString const& OutputDir)
 	{
 		check(LibraryName.Len() > 0);
 
 		const FString ShaderIntermediateLocation = FPaths::ProjectSavedDir() / TEXT("Shaders") / FormatName.ToString();
-
 		TArray<FString> ShaderFiles;
 		IFileManager::Get().FindFiles(ShaderFiles, *ShaderIntermediateLocation, *ShaderExtension);
 
@@ -2173,14 +2188,7 @@ struct FEditorShaderStableInfo
 		{
 			if (ShaderFileName.Contains(LibraryName + TEXT("-") + FormatName.ToString() + TEXT(".")))
 			{
-				TArray<FStableShaderKeyAndValue> StableKeys;
-				if (UE::PipelineCacheUtilities::LoadStableKeysFile(GetStableInfoArchiveFilename(OutputDir, LibraryName, FormatName), StableKeys))
-				{
-					for (FStableShaderKeyAndValue& Item : StableKeys)
-					{
-						AddShader(Item);
-					}
-				}
+				AddShaderCodeLibraryFromDirectory(OutputDir);
 			}
 		}
 	}
@@ -3098,7 +3106,7 @@ public:
 		StableKeyValue.ComputeKeyHash();
 		StableArchive->AddShader(StableKeyValue);
 	}
-
+ 
 	void FinishPopulateShaderCode(const FString& ShaderCodeDir, const FString& MetaOutputDir, const TArray<FName>& ShaderFormats)
 	{
 		for (FName ShaderFormatName : ShaderFormats)
@@ -3814,6 +3822,117 @@ void FShaderLibraryCooker::FinishPopulateShaderLibrary(const ITargetPlatform* Ta
 	}
 }
 
+bool FShaderLibraryCooker::MergeShaderCodeArchive(const TArray<FString>& CookedMetadataDirs, const FString& OutputDir, TArray<FString>& OutWrittenFiles)
+{
+	OutWrittenFiles.Empty();
+
+	bool bNeedsDeterministicOrder = false;
+	GConfig->GetBool(TEXT("/Script/UnrealEd.ProjectPackagingSettings"), TEXT("bDeterministicShaderCodeOrder"), bNeedsDeterministicOrder, GGameIni);
+
+	// Map of file name to the archive we are unioning
+	TMap<FString, FEditorShaderCodeArchive> ShaderCodeArchives;
+	TMap<FString, FEditorShaderStableInfo> ShaderStableInfos;
+	static const FRegexPattern ShaderArchivePattern(TEXT("^ShaderArchive\\-(\\w*)\\-([\\w\\-]*)\\.ushaderbytecode$"));
+	static const FRegexPattern StableInfoPattern(TEXT("^ShaderStableInfo\\-(\\w*)\\-([\\w\\-]*)\\.shk$"));
+
+	for (const FString& MetadataDir : CookedMetadataDirs)
+	{
+		const FString ShaderCodeDir = MetadataDir / TEXT("ShaderLibrarySource");
+		const FString ShaderStableInfoDir = MetadataDir / TEXT("PipelineCaches");
+
+		TArray<FString> ShaderBytecodeFiles;
+		IFileManager::Get().FindFiles(ShaderBytecodeFiles, *ShaderCodeDir, *ShaderExtension);
+		for (const FString& ByteCodeFile : ShaderBytecodeFiles)
+		{
+			if (ShaderCodeArchives.Contains(ByteCodeFile))
+			{
+				ShaderCodeArchives[ByteCodeFile].AddShaderCodeLibraryFromDirectory(ShaderCodeDir);
+			}
+			else
+			{
+				FRegexMatcher FindShaderFormat(ShaderArchivePattern, ByteCodeFile);
+				if (ensureMsgf(FindShaderFormat.FindNext(), TEXT("Unable to parse out shader format from %s"), *ByteCodeFile))
+				{
+					const FName ShaderFormat(*FindShaderFormat.GetCaptureGroup(2));
+					FEditorShaderCodeArchive NewCodeArchive(ShaderFormat, bNeedsDeterministicOrder);
+					NewCodeArchive.OpenLibrary(FindShaderFormat.GetCaptureGroup(1));
+					NewCodeArchive.AddShaderCodeLibraryFromDirectory(ShaderCodeDir);
+
+					ShaderCodeArchives.FindOrAdd(ByteCodeFile, MoveTemp(NewCodeArchive));
+				}
+			}
+		}
+
+		TArray<FString> StableInfoFiles;
+		IFileManager::Get().FindFiles(StableInfoFiles, *ShaderStableInfoDir, *StableExtension);
+		for (const FString& StableInfoFile : StableInfoFiles)
+		{
+			if (ShaderStableInfos.Contains(StableInfoFile))
+			{
+				ShaderStableInfos[StableInfoFile].AddShaderCodeLibraryFromDirectory(ShaderStableInfoDir);
+			}
+			else
+			{
+				FRegexMatcher FindShaderInfo(StableInfoPattern, StableInfoFile);
+				if (ensureMsgf(FindShaderInfo.FindNext(), TEXT("Unable to parse out shader format from %s"), *StableInfoFile))
+				{
+					const FName ShaderFormat(*FindShaderInfo.GetCaptureGroup(2));
+					FEditorShaderStableInfo NewStableInfo(ShaderFormat);
+					NewStableInfo.OpenLibrary(FindShaderInfo.GetCaptureGroup(1));
+					NewStableInfo.AddShaderCodeLibraryFromDirectory(ShaderStableInfoDir);
+
+					ShaderStableInfos.FindOrAdd(StableInfoFile, MoveTemp(NewStableInfo));
+				}
+			}
+		}
+	}
+
+	const FString OutShaderCodeDir = OutputDir / TEXT("ShaderLibrarySource");
+	const FString OutContentDir = OutputDir / TEXT("../Content");
+	const FString OutShaderStableInfoDir = OutputDir / TEXT("PipelineCaches");
+
+	bool bSuccess = true;
+	for (TPair<FString, FEditorShaderCodeArchive>& ShaderCodeArchivePair : ShaderCodeArchives)
+	{
+		TArray<FString> CreatedFiles;
+		if (ShaderCodeArchivePair.Value.SaveToDisk(OutShaderCodeDir, TEXT(""), false, &CreatedFiles))
+		{
+			OutWrittenFiles.Append(MoveTemp(CreatedFiles));
+		}
+		else
+		{
+			UE_LOG(LogShaderLibrary, Error, TEXT("Failed to save %s"), *ShaderCodeArchivePair.Key);
+			bSuccess = false;
+		}
+		CreatedFiles.Reset();
+		if (ShaderCodeArchivePair.Value.SaveToDisk(OutContentDir, TEXT(""), false, &CreatedFiles))
+		{
+			OutWrittenFiles.Append(MoveTemp(CreatedFiles));
+		}
+		else
+		{
+			UE_LOG(LogShaderLibrary, Error, TEXT("Failed to save to Content Dir %s"), *ShaderCodeArchivePair.Key);
+			bSuccess = false;
+		}
+	}
+
+	for (TPair<FString, FEditorShaderStableInfo>& ShaderStableInfoPair : ShaderStableInfos)
+	{
+		FString WrittenFile;
+		if (ShaderStableInfoPair.Value.SaveToDisk(OutShaderStableInfoDir, WrittenFile))
+		{
+			OutWrittenFiles.Add(MoveTemp(WrittenFile));
+		}
+		else
+		{
+			UE_LOG(LogShaderLibrary, Error, TEXT("Failed to save %s"), *ShaderStableInfoPair.Key);
+			bSuccess = false;
+		}
+	}
+
+	return bSuccess;
+}
+
 bool FShaderLibraryCooker::SaveShaderLibraryWithoutChunking(const ITargetPlatform* TargetPlatform, FString const& Name,
 	FString const& SandboxDestinationPath, FString const& SandboxMetadataPath, TArray<FString>& PlatformSCLCSVPaths, FString& OutErrorMessage,
 	bool& bOutHasData)
diff --git a/Engine/Source/Runtime/RenderCore/Public/ShaderCodeLibrary.h b/Engine/Source/Runtime/RenderCore/Public/ShaderCodeLibrary.h
--- a/Engine/Source/Runtime/RenderCore/Public/ShaderCodeLibrary.h
+++ b/Engine/Source/Runtime/RenderCore/Public/ShaderCodeLibrary.h
@@ -326,6 +326,17 @@ struct FShaderLibraryCooker
 	static RENDERCORE_API void FinishPopulateShaderLibrary(const ITargetPlatform* TargetPlatform, FString const& Name, FString const& SandboxDestinationPath,
 		FString const& SandboxMetadataPath);
 
+	/**
+	 * Given multiple Cooked Metadata directories will attempt to merge the ShaderByteCode and the ShaderStableInfo into the given OutputDir.
+	 * It would be expected that the OutputDir is another MetaData directory but this can be any dir.
+	 * Sub directories for ShaderLibrarySource and PipelineCaches will be automatically generated.
+	 * 
+	 * @param CookedMetadataDirs - the cooked metadata directories to merge the shader archives from
+	 * @param OutputDir - where to place the union of the shader archives
+	 * @param OutWrittenFiles - full path to all the files written
+	 */
+	static RENDERCORE_API bool MergeShaderCodeArchive(const TArray<FString>& CookedMetadataDirs, const FString& OutputDir, TArray<FString>& OutWrittenFiles);
+
 	/**
 	 * Saves collected shader code to a single file per shader platform
 	 * When chunking is enabled, this call will not write the shader code, only the SCL.CSV file with the stable shader info.
