Isolated hang infinite loop from page fault logic for debugging page faults that don't trigger immediately or at all. Also enlarged PageFaultUAV to 64x64 so dispatch thread ID doesn't go off the end of the texture.

diff --git a/Engine/Shaders/Private/Tools/GPUDebugCrashUtils.usf b/Engine/Shaders/Private/Tools/GPUDebugCrashUtils.usf
--- a/Engine/Shaders/Private/Tools/GPUDebugCrashUtils.usf
+++ b/Engine/Shaders/Private/Tools/GPUDebugCrashUtils.usf
@@ -24,7 +24,7 @@ void MainCS(uint3 ThreadID : SV_DispatchThreadID)
 	{
 		PageFaultUAV[index1] = what;
 	}
-
+#if  HANG_REQUESTED
 	// Try to trigger a crash (timeout) using a deliberate infinite loop
 	for (;;)
 	{
@@ -34,5 +34,5 @@ void MainCS(uint3 ThreadID : SV_DispatchThreadID)
 			break;
 		}
 	}
-	
+#endif
 }
diff --git a/Engine/Source/Runtime/RenderCore/Private/GPUDebugCrashUtils.cpp b/Engine/Source/Runtime/RenderCore/Private/GPUDebugCrashUtils.cpp
--- a/Engine/Source/Runtime/RenderCore/Private/GPUDebugCrashUtils.cpp
+++ b/Engine/Source/Runtime/RenderCore/Private/GPUDebugCrashUtils.cpp
@@ -5,8 +5,8 @@
 	=============================================================================*/
 
 #include "GPUDebugCrashUtils.h"
-#include "RenderResource.h"
 #include "RenderGraphUtils.h"
+#include "RenderResource.h"
 #include "GlobalShader.h"
 #include "RHIStaticStates.h"
 #include "DataDrivenShaderPlatformInfo.h"
@@ -18,7 +18,8 @@ public:
 	SHADER_USE_PARAMETER_STRUCT(FGPUDebugCrashUtilsCS, FGlobalShader)
 
 	class FPlatformBreakRequested : SHADER_PERMUTATION_BOOL("PLATFORM_BREAK_REQUESTED");
-	using FPermutationDomain = TShaderPermutationDomain<FPlatformBreakRequested>;
+	class FHangRequested : SHADER_PERMUTATION_BOOL("HANG_REQUESTED");
+	using FPermutationDomain = TShaderPermutationDomain<FPlatformBreakRequested, FHangRequested>;
 
 	BEGIN_SHADER_PARAMETER_STRUCT(FParameters, )
 		SHADER_PARAMETER_RDG_TEXTURE_UAV(RWTexture2D<uint>, PageFaultUAV)
@@ -38,12 +39,13 @@ void ScheduleGPUDebugCrash(FRDGBuilder& GraphBuilder)
 {
 	FGPUDebugCrashUtilsCS::FPermutationDomain PermutationVector;
 	PermutationVector.Set<FGPUDebugCrashUtilsCS::FPlatformBreakRequested>(EnumHasAnyFlags(GRHIGlobals.TriggerGPUCrash, ERequestedGPUCrash::Type_PlatformBreak));
+	PermutationVector.Set<FGPUDebugCrashUtilsCS::FHangRequested>(EnumHasAnyFlags(GRHIGlobals.TriggerGPUCrash, ERequestedGPUCrash::Type_Hang));
 
 	auto ComputeShader = GetGlobalShaderMap(GMaxRHIFeatureLevel)->GetShader<FGPUDebugCrashUtilsCS>(PermutationVector);
 	FGPUDebugCrashUtilsCS::FParameters* PassParameters = GraphBuilder.AllocParameters<FGPUDebugCrashUtilsCS::FParameters>();
 	ETextureCreateFlags TexFlags = TexCreate_UAV | TexCreate_ShaderResource | TexCreate_NoFastClear | TexCreate_RenderTargetable;
 	FRDGTextureDesc CreateInfo = FRDGTextureDesc::Create2D(
-		FIntPoint(16, 16),
+		FIntPoint(64, 64),
 		EPixelFormat::PF_R32_UINT,
 		FClearValueBinding::None,
 		TexFlags);
