Converted FRenderCommandFence to use new task graph API.

diff --git a/Engine/Source/Runtime/RHI/Private/RHIUtilities.cpp b/Engine/Source/Runtime/RHI/Private/RHIUtilities.cpp
--- a/Engine/Source/Runtime/RHI/Private/RHIUtilities.cpp
+++ b/Engine/Source/Runtime/RHI/Private/RHIUtilities.cpp
@@ -161,8 +161,18 @@ class FRHIFrameFlipTrackingRunnable : public FRunnable
 	FCriticalSection CS;
 	struct FFramePair
 	{
+		FFramePair(uint64 InPresentIndex, const UE::Tasks::FTaskEvent& InEvent)
+			: PresentIndex(InPresentIndex)
+			, Event(InEvent)
+		{}
+
+		~FFramePair()
+		{
+			Event.Trigger();
+		}
+
 		uint64 PresentIndex;
-		FGraphEventRef Event;
+		UE::Tasks::FTaskEvent Event;
 	};
 	TArray<FFramePair> FramePairs;
 
@@ -175,7 +185,7 @@ public:
 	static void Initialize();
 	static void Shutdown();
 
-	static void CompleteGraphEventOnFlip(uint64 PresentIndex, FGraphEventRef Event);
+	static void TriggerTaskEventOnFlip(uint64 PresentIndex, UE::Tasks::FTaskEvent Event);
 };
 
 FRHIFrameFlipTrackingRunnable::FRHIFrameFlipTrackingRunnable()
@@ -399,16 +409,13 @@ uint32 FRHIFrameFlipTrackingRunnable::Run()
 		}
 
 		bool bUpdateRHIFrameTime = false;
-		// Complete any relevant task graph events.
+		// Complete any relevant task events.
 		FScopeLock Lock(&CS);
 		for (int32 PairIndex = FramePairs.Num() - 1; PairIndex >= 0; --PairIndex)
 		{
-			auto const& Pair = FramePairs[PairIndex];
+			auto& Pair = FramePairs[PairIndex];
 			if (Pair.PresentIndex <= SyncFrame)
 			{
-				// "Complete" the task graph event
-				Pair.Event->DispatchSubsequents();
-
 				FramePairs.RemoveAtSwap(PairIndex);
 				bUpdateRHIFrameTime = true;
 			}
@@ -467,13 +474,6 @@ void FRHIFrameFlipTrackingRunnable::Shutdown()
 
 	FScopeLock Lock(&Singleton.CS);
 
-	// Signal any remaining events
-	for (auto const& Pair : Singleton.FramePairs)
-	{
-		// "Complete" the task graph event
-		Pair.Event->DispatchSubsequents();
-	}
-
 	Singleton.FramePairs.Empty();
 
 #if USE_FRAME_OFFSET_THREAD
@@ -481,7 +481,7 @@ void FRHIFrameFlipTrackingRunnable::Shutdown()
 #endif
 }
 
-void FRHIFrameFlipTrackingRunnable::CompleteGraphEventOnFlip(uint64 PresentIndex, FGraphEventRef Event)
+void FRHIFrameFlipTrackingRunnable::TriggerTaskEventOnFlip(uint64 PresentIndex, UE::Tasks::FTaskEvent Event)
 {
 	if ( ! FPlatformMisc::UseRenderThread() )
 	{
@@ -492,11 +492,7 @@ void FRHIFrameFlipTrackingRunnable::CompleteGraphEventOnFlip(uint64 PresentIndex
 
 	if (Thread)
 	{
-		FFramePair Pair;
-		Pair.PresentIndex = PresentIndex;
-		Pair.Event = Event;
-
-		Singleton.FramePairs.Add(Pair);
+		Singleton.FramePairs.Emplace(PresentIndex, MoveTemp(Event));
 
 #if USE_FRAME_OFFSET_THREAD
 		FRHIFrameOffsetThread::Signal();
@@ -508,7 +504,7 @@ void FRHIFrameFlipTrackingRunnable::CompleteGraphEventOnFlip(uint64 PresentIndex
 	{
 		// Platform does not support flip tracking.
 		// Signal the event now...
-		Event->DispatchSubsequents();
+		Event.Trigger();
 	}
 }
 
@@ -543,9 +539,9 @@ void RHIGetPresentThresholds(float& OutTopPercent, float& OutBottomPercent)
 	OutBottomPercent = FMath::Clamp(CVarRHIPresentThresholdBottom.GetValueOnAnyThread(), 0.0f, 1.0f);
 }
 
-void RHICompleteGraphEventOnFlip(uint64 PresentIndex, FGraphEventRef Event)
+void RHITriggerTaskEventOnFlip(uint64 PresentIndex, const UE::Tasks::FTaskEvent& Event)
 {
-	FRHIFrameFlipTrackingRunnable::CompleteGraphEventOnFlip(PresentIndex, Event);
+	FRHIFrameFlipTrackingRunnable::TriggerTaskEventOnFlip(PresentIndex, Event);
 }
 
 void RHISetFrameDebugInfo(uint64 PresentIndex, uint64 FrameIndex, uint64 InputTime)
diff --git a/Engine/Source/Runtime/RHI/Public/RHIUtilities.h b/Engine/Source/Runtime/RHI/Public/RHIUtilities.h
--- a/Engine/Source/Runtime/RHI/Public/RHIUtilities.h
+++ b/Engine/Source/Runtime/RHI/Public/RHIUtilities.h
@@ -12,6 +12,7 @@
 #include "HAL/IConsoleManager.h"
 #include "RHIDefinitions.h"
 #include "RHICommandList.h"
+#include "Tasks/Task.h"
 
 static inline EPixelFormat GetBlockCompressedFormatUAVAliasFormat(EPixelFormat Format)
 {
@@ -627,7 +628,10 @@ extern RHI_API void RHIGetPresentThresholds(float& OutTopPercent, float& OutBott
 extern RHI_API bool RHIGetSyncAllowVariable();
 
 /** Signals the completion of the specified task graph event when the given frame has flipped. */
-extern RHI_API void RHICompleteGraphEventOnFlip(uint64 PresentIndex, FGraphEventRef Event);
+UE_DEPRECATED(5.4, "RHICompleteGraphEventOnFlip is replaced with RHITriggerTaskEventOnFlip")
+inline void RHICompleteGraphEventOnFlip(uint64 PresentIndex, FGraphEventRef Event) { checkNoEntry(); }
+
+extern RHI_API void RHITriggerTaskEventOnFlip(uint64 PresentIndex, const UE::Tasks::FTaskEvent& TaskEvent);
 
 /** Sets the FrameIndex and InputTime for the current frame. */
 extern RHI_API void RHISetFrameDebugInfo(uint64 PresentIndex, uint64 FrameIndex, uint64 InputTime);
diff --git a/Engine/Source/Runtime/RenderCore/Private/RenderingThread.cpp b/Engine/Source/Runtime/RenderCore/Private/RenderingThread.cpp
--- a/Engine/Source/Runtime/RenderCore/Private/RenderingThread.cpp
+++ b/Engine/Source/Runtime/RenderCore/Private/RenderingThread.cpp
@@ -968,8 +968,7 @@ bool IsRenderingThreadHealthy()
 	return GIsRenderingThreadHealthy;
 }
 
-static FGraphEventRef BundledCompletionEvent;
-static FGraphEventRef BundledCompletionEventPrereq; // We fire this when we are done, which queues the actual fence
+static TOptional<UE::Tasks::FTaskEvent> BundledCompletionEvent;
 
 void StartRenderCommandFenceBundler()
 {
@@ -978,34 +977,28 @@ void StartRenderCommandFenceBundler()
 		return;
 	}
 
-	check(IsInGameThread() && !BundledCompletionEvent.GetReference() && !BundledCompletionEventPrereq.GetReference()); // can't use this in a nested fashion
-	BundledCompletionEventPrereq = FGraphEvent::CreateGraphEvent();
-
-	FGraphEventArray Prereqs;
-	Prereqs.Add(BundledCompletionEventPrereq);
-
-	DECLARE_CYCLE_STAT(TEXT("FNullGraphTask.FenceRenderCommandBundled"),
-	STAT_FNullGraphTask_FenceRenderCommandBundled,
-		STATGROUP_TaskGraphTasks);
-
-	BundledCompletionEvent = TGraphTask<FNullGraphTask>::CreateTask(&Prereqs, ENamedThreads::GameThread).ConstructAndDispatchWhenReady(
-		GET_STATID(STAT_FNullGraphTask_FenceRenderCommandBundled), ENamedThreads::GetRenderThread());
+	check(IsInGameThread() && !BundledCompletionEvent); // can't use this in a nested fashion
+	BundledCompletionEvent.Emplace(TEXT("RenderCommandFenceBundlerEvent"));
 
 	StartBatchedRelease();
 }
 
 void StopRenderCommandFenceBundler()
 {
-	if (!GIsThreadedRendering || !BundledCompletionEvent.GetReference())
+	if (!GIsThreadedRendering || !BundledCompletionEvent)
 	{
 		return;
 	}
 
 	EndBatchedRelease();
-	checkf(IsInGameThread() && BundledCompletionEvent.GetReference() && !BundledCompletionEvent->IsComplete() && BundledCompletionEventPrereq.GetReference() && !BundledCompletionEventPrereq->IsComplete(), TEXT("IsInGameThread: %d, BundledCompletionEvent is completed: %d, BundledCompletionEventPrereq is completed: %d"), IsInGameThread(), BundledCompletionEvent->IsComplete(), BundledCompletionEventPrereq->IsComplete()); // can't use this in a nested fashion
-	BundledCompletionEventPrereq->DispatchSubsequents();
-	BundledCompletionEventPrereq = nullptr;
-	BundledCompletionEvent = nullptr;
+	checkf(IsInGameThread() && BundledCompletionEvent && !BundledCompletionEvent->IsCompleted(), TEXT("IsInGameThread: %d, BundledCompletionEvent is completed: %d"), IsInGameThread(), BundledCompletionEvent->IsCompleted()); // can't use this in a nested fashion
+
+	ENQUEUE_RENDER_COMMAND(InsertFence)(
+		[CompletionEvent = MoveTemp(*BundledCompletionEvent)](FRHICommandListImmediate&) mutable
+	{
+		CompletionEvent.Trigger();
+	});
+	BundledCompletionEvent.Reset();
 }
 
 TAutoConsoleVariable<int32> CVarGTSyncType(
@@ -1018,21 +1011,13 @@ TAutoConsoleVariable<int32> CVarGTSyncType(
 	TEXT(" 2 - Sync the game thread with the GPU swap chain flip (only on supported platforms).\n"),
 	ECVF_Default);
 
-TAutoConsoleVariable<bool> CVarAllowRHITriggerThread(
-	TEXT("r.AllowRHITriggerThread"),
-	true,
-	TEXT("In low latency mode, use the rhi thread to trigger the frame sync.")
-	TEXT(" true (default).\n")
-	TEXT(" false.\n"),
-	ECVF_Default);
-
 FRHICOMMAND_MACRO(FRHISyncFrameCommand)
 {
-	FGraphEventRef GraphEvent;
+	UE::Tasks::FTaskEvent TaskEvent;
 	int32 GTSyncType;
 
-	FORCEINLINE_DEBUGGABLE FRHISyncFrameCommand(FGraphEventRef InGraphEvent, int32 InGTSyncType)
-		: GraphEvent(InGraphEvent)
+	FORCEINLINE_DEBUGGABLE FRHISyncFrameCommand(UE::Tasks::FTaskEvent InTaskEvent, int32 InGTSyncType)
+		: TaskEvent(MoveTemp(InTaskEvent))
 		, GTSyncType(InGTSyncType)
 	{}
 
@@ -1043,13 +1028,13 @@ FRHICOMMAND_MACRO(FRHISyncFrameCommand)
 			// Sync the Game Thread with the RHI Thread
 
 			// "Complete" the graph event
-			GraphEvent->DispatchSubsequents();
+			TaskEvent.Trigger();
 		}
 		else
 		{
 			// This command runs *after* a present has happened, so the counter has already been incremented.
 			// Subtracting 1 gives us the index of the frame that has *just* been presented.
-			RHICompleteGraphEventOnFlip(GRHIPresentCounter - 1, GraphEvent);
+			RHITriggerTaskEventOnFlip(GRHIPresentCounter - 1, TaskEvent);
 		}
 	}
 };
@@ -1065,12 +1050,9 @@ void FRenderCommandFence::BeginFence(bool bSyncToRHIAndGPU)
 	}
 	else
 	{
-		// Render thread is a default trigger for the CompletionEvent
-		TriggerThreadIndex = ENamedThreads::ActualRenderingThread;
-				
-		if (BundledCompletionEvent.GetReference() && IsInGameThread())
+		if (BundledCompletionEvent && IsInGameThread())
 		{
-			CompletionEvent = BundledCompletionEvent;
+			CompletionTask = *BundledCompletionEvent;
 			return;
 		}
 
@@ -1088,45 +1070,31 @@ void FRenderCommandFence::BeginFence(bool bSyncToRHIAndGPU)
 			}
 		}
 
+		UE::Tasks::FTaskEvent CompletionTaskEvent{ UE_SOURCE_LOCATION };
 
-		if (bSyncToRHIAndGPU)
+		ENQUEUE_RENDER_COMMAND(FSyncFrameCommand)(
+			[CompletionTaskEvent, GTSyncType, bSyncToRHIAndGPU](FRHICommandListImmediate& RHICmdList) mutable
 		{
-			if (IsRHIThreadRunning())
+			if (bSyncToRHIAndGPU)
 			{
-				// Potentially change trigger thread to RHI
-				// On some platform RHI thread will block on present. Putting the RHI as the trigger index will block the GameThread until the present is finished.
-				// In low input latency mode, some consoles uses the RHIOffsetThread to kick of the Gamethread, so we dont want it to block on present.
-				TriggerThreadIndex = (GTSyncType == 2 && !CVarAllowRHITriggerThread.GetValueOnAnyThread()) ? TriggerThreadIndex : ENamedThreads::RHIThread;
-			}
-			
-			// Create a task graph event which we can pass to the render or RHI threads.
-			CompletionEvent = FGraphEvent::CreateGraphEvent();
-
-			FGraphEventRef InCompletionEvent = CompletionEvent;
-			ENQUEUE_RENDER_COMMAND(FSyncFrameCommand)(
-				[InCompletionEvent, GTSyncType](FRHICommandListImmediate& RHICmdList)
+				if (IsRHIThreadRunning())
 				{
-					if (IsRHIThreadRunning())
-					{
-						ALLOC_COMMAND_CL(RHICmdList, FRHISyncFrameCommand)(InCompletionEvent, GTSyncType);
-						RHICmdList.ImmediateFlush(EImmediateFlushType::DispatchToRHIThread);
-					}
-					else
-					{
-						FRHISyncFrameCommand Command(InCompletionEvent, GTSyncType);
-						Command.Execute(RHICmdList);
-					}
-				});
-		}
-		else
-		{
-			// Sync Game Thread with Render Thread only
-			DECLARE_CYCLE_STAT(TEXT("FNullGraphTask.FenceRenderCommand"),
-			STAT_FNullGraphTask_FenceRenderCommand,
-				STATGROUP_TaskGraphTasks);
+					ALLOC_COMMAND_CL(RHICmdList, FRHISyncFrameCommand)(MoveTemp(CompletionTaskEvent), GTSyncType);
+					RHICmdList.ImmediateFlush(EImmediateFlushType::DispatchToRHIThread);
+				}
+				else
+				{
+					FRHISyncFrameCommand Command(MoveTemp(CompletionTaskEvent), GTSyncType);
+					Command.Execute(RHICmdList);
+				}
+			}
+			else
+			{
+				CompletionTaskEvent.Trigger();
+			}
+		});
 
-			CompletionEvent = FFunctionGraphTask::CreateAndDispatchWhenReady([] {}, GET_STATID(STAT_FNullGraphTask_FenceRenderCommand), nullptr, ENamedThreads::GetRenderThread());
-		}
+		CompletionTask = MoveTemp(CompletionTaskEvent);
 	}
 }
 
@@ -1138,9 +1106,9 @@ bool FRenderCommandFence::IsFenceComplete() const
 	}
 	check(IsInGameThread() || IsInAsyncLoadingThread());
 	CheckRenderingThreadHealth();
-	if (!CompletionEvent.GetReference() || CompletionEvent->IsComplete())
+	if (CompletionTask.IsCompleted())
 	{
-		CompletionEvent = NULL; // this frees the handle for other uses, the NULL state is considered completed
+		CompletionTask = {}; // this frees the handle for other uses, the NULL state is considered completed
 		return true;
 	}
 	return false;
@@ -1173,16 +1141,15 @@ static FAutoConsoleVariableRef CVarTimeoutForBlockOnRenderFence(
 /**
  * Block the game thread waiting for a task to finish on the rendering thread.
  */
-static void GameThreadWaitForTask(const FGraphEventRef& Task, ENamedThreads::Type TriggerThreadIndex = ENamedThreads::ActualRenderingThread, bool bEmptyGameThreadTasks = false)
+static void GameThreadWaitForTask(const UE::Tasks::FTask& Task, bool bEmptyGameThreadTasks = false)
 {
-	TaskTrace::FWaitingScope WaitingScope(GetTraceIds({ Task }));
 	TRACE_CPUPROFILER_EVENT_SCOPE(GameThreadWaitForTask);
 	SCOPE_TIME_GUARD(TEXT("GameThreadWaitForTask"));
 
 	check(IsInGameThread());
-	check(IsValidRef(Task));	
+	check(Task.IsValid());
 
-	if (!Task->IsComplete())
+	if (!Task.IsCompleted())
 	{
 		SCOPE_CYCLE_COUNTER(STAT_GameIdleTime);
 		{
@@ -1205,11 +1172,6 @@ static void GameThreadWaitForTask(const FGraphEventRef& Task, ENamedThreads::Typ
 				bEmptyGameThreadTasks = false; // we don't do this on recursive calls or if we are at a blueprint breakpoint
 			}
 
-			// Grab an event from the pool and fire off a task to trigger it.
-			FEvent* Event = FPlatformProcess::GetSynchEventFromPool();
-			check(GIsThreadedRendering);
-			FTaskGraphInterface::Get().TriggerEventWhenTaskCompletes(Event, Task, ENamedThreads::GameThread, ENamedThreads::SetTaskPriority(TriggerThreadIndex, ENamedThreads::HighTaskPriority));
-
 			// Check rendering thread health needs to be called from time to
 			// time in order to pump messages, otherwise the RHI may block
 			// on vsync causing a deadlock. Also we should make sure the
@@ -1226,11 +1188,6 @@ static void GameThreadWaitForTask(const FGraphEventRef& Task, ENamedThreads::Typ
 
 			static bool bDisabled = FParse::Param(FCommandLine::Get(), TEXT("nothreadtimeout"));
 
-			if (TriggerThreadIndex == ENamedThreads::ActualRenderingThread && !bEmptyGameThreadTasks && GRenderThreadPollingOn)
-			{
-				FTaskGraphInterface::Get().WakeNamedThread(ENamedThreads::GetRenderThread());
-			}
-
 			do
 			{
 				CheckRenderingThreadHealth();
@@ -1239,7 +1196,7 @@ static void GameThreadWaitForTask(const FGraphEventRef& Task, ENamedThreads::Typ
 					// process gamethread tasks if there are any
 					FTaskGraphInterface::Get().ProcessThreadUntilIdle(ENamedThreads::GameThread);
 				}
-				bDone = Event->Wait(WaitTime);
+				bDone = Task.Wait(FTimespan::FromMilliseconds(WaitTime));
 
 				RenderThreadTimeoutClock.Tick();
 
@@ -1277,9 +1234,6 @@ static void GameThreadWaitForTask(const FGraphEventRef& Task, ENamedThreads::Typ
 			}
 			while (!bDone);
 
-			// Return the event to the pool and decrement the recursion counter.
-			FPlatformProcess::ReturnSynchEventToPool(Event);
-			Event = nullptr;
 			NumRecursiveCalls--;
 		}
 	}
@@ -1293,17 +1247,8 @@ void FRenderCommandFence::Wait(bool bProcessGameThreadTasks) const
 	if (!IsFenceComplete())
 	{
 		StopRenderCommandFenceBundler();
-#if 0
-		// on most platforms this is a better solution because it doesn't spin
-		// windows needs to pump messages
-		if (bProcessGameThreadTasks)
-		{
-			QUICK_SCOPE_CYCLE_COUNTER(STAT_FRenderCommandFence_Wait);
-			FTaskGraphInterface::Get().WaitUntilTaskCompletes(CompletionEvent, ENamedThreads::GameThread);
-		}
-#endif
-		GameThreadWaitForTask(CompletionEvent, TriggerThreadIndex, bProcessGameThreadTasks);
-		CompletionEvent = nullptr; // release the internal memory as soon as it's not needed anymore
+		GameThreadWaitForTask(CompletionTask, bProcessGameThreadTasks);
+		CompletionTask = {}; // release the internal memory as soon as it's not needed anymore
 	}
 }
 
diff --git a/Engine/Source/Runtime/RenderCore/Public/RenderCommandFence.h b/Engine/Source/Runtime/RenderCore/Public/RenderCommandFence.h
--- a/Engine/Source/Runtime/RenderCore/Public/RenderCommandFence.h
+++ b/Engine/Source/Runtime/RenderCore/Public/RenderCommandFence.h
@@ -2,10 +2,7 @@
 
 #pragma once
 
-#include "Async/TaskGraphFwd.h"
-#include "Templates/RefCounting.h"
-
-namespace ENamedThreads { enum Type : int32; }
+#include "Tasks/Task.h"
 
 ////////////////////////////////////
 // Render fences
@@ -40,10 +37,8 @@ public:
 	RENDERCORE_API ~FRenderCommandFence();
 
 private:
-	/** Graph event that represents completion of this fence **/
-	mutable FGraphEventRef CompletionEvent;
-	/** Thread that will trigger the CompletionEvent **/
-	ENamedThreads::Type TriggerThreadIndex;
+	/** Task that represents completion of this fence **/
+	mutable UE::Tasks::FTask CompletionTask;
 };
 
 #if UE_ENABLE_INCLUDE_ORDER_DEPRECATED_IN_5_2
