Render Command Pipe Implementation and API

Render Command Pipes dedicated asynchronous task pipes for render commands. Users can easily define new pipes and enqueue commands into them. Pipes can be synchronized using a scope to run serial render commands on the render thread, but initially pipes cannot be synchronized individually with each other. Render command overhead is reduced by recording command lambdas into MPSC queues which are serviced by the task graph; both for pipes and for the render thread. This reduces the task overhead as commands are no longer 1-to-1 with tasks.

Pipe behavior is controlled with new CVars. `r.RenderCommandPipeMode` controls overall behavior:
 0 - Legacy render thread tasks,
 1 - Render thread MPSC queue,
 2 - Render thread and async pipe MPSC queues.

To define a Render Command Pipe, use DEFINE_RENDER_COMMAND_PIPE(MyPipe), or DECLARE_RENDER_COMMAND_PIPE(MyPipe, MODULE_API) to declare an extern reference.

Enqueue a command into the pipe like so:

ENQUEUE_RENDER_COMMAND(MyCommand)(UE::RenderCommandPipe::MyPipe, [] (FRHICommandList&) {}).

Omitting a pipe will fallback to the 'general' pipe which is the render thread.

Eventually pipes need to be synced back to the general pipe for scene renders and other GPU work. On the game thread timeline, use UE::RenderCommandPipe::FSyncScope to synchronize the pipes. This waits for pipes and disables recording of new pipe commands until the scope completes, at which point pipe recording is restarted. This creates a 'sync point', so render commands issued prior to a sync scope will be waited on at the start of the scope, and render commands issued after the scope ends will not be able to start until the render thread finishes processing prior commands.

diff --git a/Engine/Plugins/Experimental/Water/Source/Runtime/Private/WaterInfoRendering.cpp b/Engine/Plugins/Experimental/Water/Source/Runtime/Private/WaterInfoRendering.cpp
--- a/Engine/Plugins/Experimental/Water/Source/Runtime/Private/WaterInfoRendering.cpp
+++ b/Engine/Plugins/Experimental/Water/Source/Runtime/Private/WaterInfoRendering.cpp
@@ -818,6 +818,8 @@ void UpdateWaterInfoRendering(
 	Params.VelocityBlurRadius = Context.ZoneToRender->GetVelocityBlurRadius();
 	Params.WaterZoneExtents = ZoneExtent;
 
+	UE::RenderCommandPipe::FSyncScope SyncScope;
+
 	ENQUEUE_RENDER_COMMAND(WaterInfoCommand)(
 	[Params, ZoneName = Context.ZoneToRender->GetActorNameOrLabel()](FRHICommandListImmediate& RHICmdList)
 		{
diff --git a/Engine/Plugins/Runtime/ComputeFramework/Source/ComputeFramework/Private/ComputeFramework.cpp b/Engine/Plugins/Runtime/ComputeFramework/Source/ComputeFramework/Private/ComputeFramework.cpp
--- a/Engine/Plugins/Runtime/ComputeFramework/Source/ComputeFramework/Private/ComputeFramework.cpp
+++ b/Engine/Plugins/Runtime/ComputeFramework/Source/ComputeFramework/Private/ComputeFramework.cpp
@@ -87,6 +87,8 @@ namespace ComputeFramework
 		FComputeGraphTaskWorker* ComputeGraphWorker = ComputeSystem != nullptr ? ComputeSystem->GetComputeWorker(InScene) : nullptr;
 		if (ensure(ComputeGraphWorker))
 		{
+			UE::RenderCommandPipe::FSyncScope SyncScope;
+
 			ENQUEUE_RENDER_COMMAND(ComputeFrameworkFlushCommand)(
 				[ComputeGraphWorker, InExecutionGroupName](FRHICommandListImmediate& RHICmdList)
 				{
diff --git a/Engine/Source/Editor/VirtualTexturingEditor/Private/RuntimeVirtualTextureBuildStreamingMips.cpp b/Engine/Source/Editor/VirtualTexturingEditor/Private/RuntimeVirtualTextureBuildStreamingMips.cpp
--- a/Engine/Source/Editor/VirtualTexturingEditor/Private/RuntimeVirtualTextureBuildStreamingMips.cpp
+++ b/Engine/Source/Editor/VirtualTexturingEditor/Private/RuntimeVirtualTextureBuildStreamingMips.cpp
@@ -213,6 +213,8 @@ namespace RuntimeVirtualTexture
 		TArray64<uint8> FinalPixels;
 		FinalPixels.SetNumUninitialized(RenderTileResourcesBytes);
 
+		UE::RenderCommandPipe::FSyncScope SyncScope;
+
 		// Iterate over all tiles and render/store each one to the final image
 		for (int32 TileY = 0; TileY < NumTilesY && !Task.ShouldCancel(); TileY++)
 		{
diff --git a/Engine/Source/Runtime/Engine/Private/ActorEditor.cpp b/Engine/Source/Runtime/Engine/Private/ActorEditor.cpp
--- a/Engine/Source/Runtime/Engine/Private/ActorEditor.cpp
+++ b/Engine/Source/Runtime/Engine/Private/ActorEditor.cpp
@@ -774,6 +774,8 @@ void AActor::PostEditUndo()
 	UWorld* World = GetWorld();
 	if (World && World->Scene && !FActorEditorUtils::IsABrush(this))
 	{
+		UE::RenderCommandPipe::FSyncScope SyncScope;
+
 		ENQUEUE_RENDER_COMMAND(UpdateAllPrimitiveSceneInfosCmd)([Scene = World->Scene](FRHICommandListImmediate& RHICmdList) {
 			Scene->UpdateAllPrimitiveSceneInfos(RHICmdList);
 		});
diff --git a/Engine/Source/Runtime/Engine/Private/Components/ActorComponent.cpp b/Engine/Source/Runtime/Engine/Private/Components/ActorComponent.cpp
--- a/Engine/Source/Runtime/Engine/Private/Components/ActorComponent.cpp
+++ b/Engine/Source/Runtime/Engine/Private/Components/ActorComponent.cpp
@@ -149,6 +149,8 @@ void UpdateAllPrimitiveSceneInfosForSingleComponent(UActorComponent* InComponent
 	{
 		if (InScenesToUpdateAllPrimitiveSceneInfosForBatching == nullptr)
 		{
+			UE::RenderCommandPipe::FSyncScope SyncScope;
+
 			// If no batching is available (this ComponentReregisterContext is not created by a FGlobalComponentReregisterContext), issue one update per component
 			ENQUEUE_RENDER_COMMAND(UpdateAllPrimitiveSceneInfosCmd)([Scene](FRHICommandListImmediate& RHICmdList) {
 				Scene->UpdateAllPrimitiveSceneInfos(RHICmdList);
@@ -166,6 +168,8 @@ void UpdateAllPrimitiveSceneInfosForScenes(TSet<FSceneInterface*> ScenesToUpdate
 {
 	if (ScenesToUpdateAllPrimitiveSceneInfos.Num())
 	{
+		UE::RenderCommandPipe::FSyncScope SyncScope;
+
 		ENQUEUE_RENDER_COMMAND(UpdateAllPrimitiveSceneInfosCmd)(
 			[ScenesToUpdateAllPrimitiveSceneInfos](FRHICommandListImmediate& RHICmdList)
 			{
diff --git a/Engine/Source/Runtime/Engine/Private/StaticMesh.cpp b/Engine/Source/Runtime/Engine/Private/StaticMesh.cpp
--- a/Engine/Source/Runtime/Engine/Private/StaticMesh.cpp
+++ b/Engine/Source/Runtime/Engine/Private/StaticMesh.cpp
@@ -1125,11 +1125,11 @@ void FStaticMeshVertexBuffers::InitModelBuffers(TArray<FModelVertex>& Vertices)
 	}
 }
 
-void FStaticMeshVertexBuffers::InitModelVF(FLocalVertexFactory* VertexFactory)
+void FStaticMeshVertexBuffers::InitModelVF(FRenderCommandPipe* RenderCommandPipe, FLocalVertexFactory* VertexFactory)
 {
 	FStaticMeshVertexBuffers* Self = this;
-	ENQUEUE_RENDER_COMMAND(StaticMeshVertexBuffersLegacyBspInit)(
-		[VertexFactory, Self](FRHICommandListImmediate& RHICmdList)
+	ENQUEUE_RENDER_COMMAND(StaticMeshVertexBuffersLegacyBspInit)(RenderCommandPipe,
+		[VertexFactory, Self](FRHICommandListBase& RHICmdList)
 	{
 		check(Self->PositionVertexBuffer.IsInitialized());
 		check(Self->StaticMeshVertexBuffer.IsInitialized());
@@ -1146,7 +1146,7 @@ void FStaticMeshVertexBuffers::InitModelVF(FLocalVertexFactory* VertexFactory)
 	});
 }
 
-void FStaticMeshVertexBuffers::InitWithDummyData(FLocalVertexFactory* VertexFactory, uint32 NumVerticies, uint32 NumTexCoords, uint32 LightMapIndex)
+void FStaticMeshVertexBuffers::InitWithDummyData(FRenderCommandPipe* RenderCommandPipe, FLocalVertexFactory* VertexFactory, uint32 NumVerticies, uint32 NumTexCoords, uint32 LightMapIndex)
 {
 	check(NumVerticies);
 	check(NumTexCoords < MAX_STATIC_TEXCOORDS && NumTexCoords > 0);
@@ -1157,8 +1157,8 @@ void FStaticMeshVertexBuffers::InitWithDummyData(FLocalVertexFactory* VertexFact
 	ColorVertexBuffer.Init(NumVerticies);
 
 	FStaticMeshVertexBuffers* Self = this;
-	ENQUEUE_RENDER_COMMAND(StaticMeshVertexBuffersLegacyInit)(
-		[VertexFactory, Self, LightMapIndex](FRHICommandListImmediate& RHICmdList)
+	ENQUEUE_RENDER_COMMAND(StaticMeshVertexBuffersLegacyInit)(RenderCommandPipe,
+		[VertexFactory, Self, LightMapIndex](FRHICommandListBase& RHICmdList)
 	{
 		InitOrUpdateResource(RHICmdList, &Self->PositionVertexBuffer);
 		InitOrUpdateResource(RHICmdList, &Self->StaticMeshVertexBuffer);
@@ -1176,7 +1176,7 @@ void FStaticMeshVertexBuffers::InitWithDummyData(FLocalVertexFactory* VertexFact
 	});
 }
 
-void FStaticMeshVertexBuffers::InitFromDynamicVertex(FLocalVertexFactory* VertexFactory, TArray<FDynamicMeshVertex>& Vertices, uint32 NumTexCoords, uint32 LightMapIndex)
+void FStaticMeshVertexBuffers::InitFromDynamicVertex(FRenderCommandPipe* RenderCommandPipe, FLocalVertexFactory* VertexFactory, TArray<FDynamicMeshVertex>& Vertices, uint32 NumTexCoords, uint32 LightMapIndex)
 {
 	check(NumTexCoords < MAX_STATIC_TEXCOORDS && NumTexCoords > 0);
 	check(LightMapIndex < NumTexCoords);
@@ -1215,8 +1215,8 @@ void FStaticMeshVertexBuffers::InitFromDynamicVertex(FLocalVertexFactory* Vertex
 	}
 
 	FStaticMeshVertexBuffers* Self = this;
-	ENQUEUE_RENDER_COMMAND(StaticMeshVertexBuffersLegacyInit)(
-		[VertexFactory, Self, LightMapIndex](FRHICommandListImmediate& RHICmdList)
+	ENQUEUE_RENDER_COMMAND(StaticMeshVertexBuffersLegacyInit)(RenderCommandPipe,
+		[VertexFactory, Self, LightMapIndex](FRHICommandListBase& RHICmdList)
 		{
 			InitOrUpdateResource(RHICmdList, &Self->PositionVertexBuffer);
 			InitOrUpdateResource(RHICmdList, &Self->StaticMeshVertexBuffer);
diff --git a/Engine/Source/Runtime/Engine/Private/UserInterface/Canvas.cpp b/Engine/Source/Runtime/Engine/Private/UserInterface/Canvas.cpp
--- a/Engine/Source/Runtime/Engine/Private/UserInterface/Canvas.cpp
+++ b/Engine/Source/Runtime/Engine/Private/UserInterface/Canvas.cpp
@@ -559,6 +559,8 @@ FCanvasRenderThreadScope::~FCanvasRenderThreadScope()
 {
 	RenderCommandFunctionArray* RenderCommandArray = RenderCommands;
 
+	UE::RenderCommandPipe::FSyncScope SyncScope;
+
 	ENQUEUE_RENDER_COMMAND(DispatchCanvasRenderCommands)(
 		[RenderCommandArray, RenderTarget = Canvas.RenderTarget, ViewRect = Canvas.ViewRect, ScissorRect = Canvas.ScissorRect, bScaledToRenderTarget = Canvas.bScaledToRenderTarget](FRHICommandListImmediate& RHICmdList) mutable
 	{
diff --git a/Engine/Source/Runtime/Engine/Public/StaticMeshResources.h b/Engine/Source/Runtime/Engine/Public/StaticMeshResources.h
--- a/Engine/Source/Runtime/Engine/Public/StaticMeshResources.h
+++ b/Engine/Source/Runtime/Engine/Public/StaticMeshResources.h
@@ -316,16 +316,31 @@ struct FStaticMeshVertexBuffers
 	FColorVertexBuffer ColorVertexBuffer;
 
 	/* This is a temporary function to refactor and convert old code, do not copy this as is and try to build your data as SoA from the beginning.*/
-	void ENGINE_API InitWithDummyData(FLocalVertexFactory* VertexFactory, uint32 NumVerticies, uint32 NumTexCoords = 1, uint32 LightMapIndex = 0);
+	void ENGINE_API InitWithDummyData(FRenderCommandPipe* RenderCommandPipe, FLocalVertexFactory* VertexFactory, uint32 NumVertices, uint32 NumTexCoords = 1, uint32 LightMapIndex = 0);
+
+	void inline InitWithDummyData(FLocalVertexFactory* VertexFactory, uint32 NumVertices, uint32 NumTexCoords = 1, uint32 LightMapIndex = 0)
+	{
+		InitWithDummyData(nullptr, VertexFactory, NumVertices, NumTexCoords, LightMapIndex);
+	}
 
 	/* This is a temporary function to refactor and convert old code, do not copy this as is and try to build your data as SoA from the beginning.*/
-	void ENGINE_API InitFromDynamicVertex(FLocalVertexFactory* VertexFactory, TArray<FDynamicMeshVertex>& Vertices, uint32 NumTexCoords = 1, uint32 LightMapIndex = 0);
+	void ENGINE_API InitFromDynamicVertex(FRenderCommandPipe* RenderCommandPipe, FLocalVertexFactory* VertexFactory, TArray<FDynamicMeshVertex>& Vertices, uint32 NumTexCoords = 1, uint32 LightMapIndex = 0);
+
+	inline void InitFromDynamicVertex(FLocalVertexFactory* VertexFactory, TArray<FDynamicMeshVertex>& Vertices, uint32 NumTexCoords = 1, uint32 LightMapIndex = 0)
+	{
+		InitFromDynamicVertex(nullptr, VertexFactory, Vertices, NumTexCoords, LightMapIndex);
+	}
 
 	/* This is a temporary function to refactor and convert old code, do not copy this as is and try to build your data as SoA from the beginning.*/
 	void ENGINE_API InitModelBuffers(TArray<FModelVertex>& Vertices);
 
 	/* This is a temporary function to refactor and convert old code, do not copy this as is and try to build your data as SoA from the beginning.*/
-	void ENGINE_API InitModelVF(FLocalVertexFactory* VertexFactory);
+	void ENGINE_API InitModelVF(FRenderCommandPipe* RenderCommandPipe, FLocalVertexFactory* VertexFactory);
+
+	inline void ENGINE_API InitModelVF(FLocalVertexFactory* VertexFactory)
+	{
+		InitModelVF(nullptr, VertexFactory);
+	}
 
 	void ENGINE_API SetOwnerName(const FName& OwnerName);
 };
diff --git a/Engine/Source/Runtime/Launch/Private/LaunchEngineLoop.cpp b/Engine/Source/Runtime/Launch/Private/LaunchEngineLoop.cpp
--- a/Engine/Source/Runtime/Launch/Private/LaunchEngineLoop.cpp
+++ b/Engine/Source/Runtime/Launch/Private/LaunchEngineLoop.cpp
@@ -4362,6 +4362,8 @@ int32 FEngineLoop::PreInitPostStartupScreen(const TCHAR* CmdLine)
 		PrecacheComputePipelineStatesForGlobalShaders(GShaderPlatformForFeatureLevel[GMaxRHIFeatureLevel], nullptr);
 	}
 
+	UE::RenderCommandPipe::Initialize();
+
 #else // WITH_ENGINE
 	InitEngineTextLocalization();
 	InitGameTextLocalization();
@@ -5688,6 +5690,8 @@ void FEngineLoop::Tick()
 			});
 		}
 
+		UE::RenderCommandPipe::StartRecording();
+
 #if !UE_SERVER && WITH_ENGINE
 		if (!GIsEditor && GEngine->GameViewport && GEngine->GameViewport->GetWorld() && GEngine->GameViewport->GetWorld()->IsCameraMoveable())
 		{
@@ -6009,6 +6013,8 @@ void FEngineLoop::Tick()
 		}
 #endif
 
+		UE::RenderCommandPipe::StopRecording();
+
 		for (FSceneInterface* Scene : GetRendererModule().GetAllocatedScenes())
 		{
 			ENQUEUE_RENDER_COMMAND(FScene_EndFrame)([Scene](FRHICommandListImmediate& RHICmdList)
diff --git a/Engine/Source/Runtime/RenderCore/Private/RenderResource.cpp b/Engine/Source/Runtime/RenderCore/Private/RenderResource.cpp
--- a/Engine/Source/Runtime/RenderCore/Private/RenderResource.cpp
+++ b/Engine/Source/Runtime/RenderCore/Private/RenderResource.cpp
@@ -291,19 +291,19 @@ FName FRenderResource::GetOwnerName() const
 #endif
 }
 
-void BeginInitResource(FRenderResource* Resource)
+void BeginInitResource(FRenderResource* Resource, FRenderCommandPipe* RenderCommandPipe)
 {
-	ENQUEUE_RENDER_COMMAND(InitCommand)(
-		[Resource](FRHICommandListImmediate& RHICmdList)
+	ENQUEUE_RENDER_COMMAND(InitCommand)(RenderCommandPipe,
+		[Resource](FRHICommandListBase& RHICmdList)
 		{
 			Resource->InitResource(RHICmdList);
 		});
 }
 
-void BeginUpdateResourceRHI(FRenderResource* Resource)
+void BeginUpdateResourceRHI(FRenderResource* Resource, FRenderCommandPipe* RenderCommandPipe)
 {
-	ENQUEUE_RENDER_COMMAND(UpdateCommand)(
-		[Resource](FRHICommandListImmediate& RHICmdList)
+	ENQUEUE_RENDER_COMMAND(UpdateCommand)(RenderCommandPipe,
+		[Resource](FRHICommandListBase& RHICmdList)
 		{
 			Resource->UpdateRHI(RHICmdList);
 		});
@@ -367,15 +367,15 @@ void EndBatchedRelease()
 	GBatchedReleaseIsActive = false;
 }
 
-void BeginReleaseResource(FRenderResource* Resource)
+void BeginReleaseResource(FRenderResource* Resource, FRenderCommandPipe* RenderCommandPipe)
 {
 	if (GBatchedReleaseIsActive && IsInGameThread())
 	{
 		GBatchedRelease.Add(Resource);
 		return;
 	}
-	ENQUEUE_RENDER_COMMAND(ReleaseCommand)(
-		[Resource](FRHICommandList& RHICmdList)
+	ENQUEUE_RENDER_COMMAND(ReleaseCommand)(RenderCommandPipe,
+		[Resource]
 		{
 			Resource->ReleaseResource();
 		});
diff --git a/Engine/Source/Runtime/RenderCore/Private/RenderingThread.cpp b/Engine/Source/Runtime/RenderCore/Private/RenderingThread.cpp
--- a/Engine/Source/Runtime/RenderCore/Private/RenderingThread.cpp
+++ b/Engine/Source/Runtime/RenderCore/Private/RenderingThread.cpp
@@ -29,7 +29,7 @@
 #include "ProfilingDebugging/MiscTrace.h"
 #include "ProfilingDebugging/CsvProfiler.h"
 #include "Async/TaskTrace.h"
-
+#include "DataDrivenShaderPlatformInfo.h"
 
 //
 // Globals
@@ -1276,6 +1276,8 @@ void FlushRenderingCommands()
 		FTaskGraphInterface::Get().ProcessThreadUntilIdle(ENamedThreads::GameThread_Local);
 	}
 
+	UE::RenderCommandPipe::FSyncScope SyncScope;
+
 	ENQUEUE_RENDER_COMMAND(FlushPendingDeleteRHIResourcesCmd)([](FRHICommandListImmediate& RHICmdList)
 	{
 		RHICmdList.ImmediateFlush(EImmediateFlushType::FlushRHIThreadFlushResources);
@@ -1487,4 +1489,304 @@ static FAutoConsoleCommand CVarRHIThreadEnable(
 	FConsoleCommandWithArgsDelegate::CreateStatic(&HandleRHIThreadEnableChanged)
 	);
 
+ERenderCommandPipeMode GRenderCommandPipeMode = ERenderCommandPipeMode::All;
+FAutoConsoleVariable CVarRenderCommandPipeMode(
+	TEXT("r.RenderCommandPipeMode"),
+	(int32)GRenderCommandPipeMode,
+	TEXT("Controls behavior of the main render thread command pipe.")
+	TEXT(" 0: Render commands are launched individually as tasks;\n")
+	TEXT(" 1: Render commands are enqueued into a render command pipe for the render thread only.;\n")
+	TEXT(" 2: Render commands are enqueued into a render command pipe for all declared pipes.;\n"),
+	FConsoleVariableDelegate::CreateLambda([](IConsoleVariable* Variable)
+	{
+		switch (Variable->GetInt())
+		{
+		case 1:
+			GRenderCommandPipeMode = ERenderCommandPipeMode::RenderThread;
+			break;
+		case 2:
+			GRenderCommandPipeMode = ERenderCommandPipeMode::All;
+			break;
+		default:
+			GRenderCommandPipeMode = ERenderCommandPipeMode::None;
+			break;
+		}
+
+		const bool bAllowThreading = !GRHICommandList.Bypass() && FApp::ShouldUseThreadingForPerformance() && GIsThreadedRendering;
+
+		if (GRenderCommandPipeMode == ERenderCommandPipeMode::All && !bAllowThreading)
+		{
+			GRenderCommandPipeMode = ERenderCommandPipeMode::RenderThread;
+		}
+
+		if (GRenderCommandPipeMode == ERenderCommandPipeMode::RenderThread && IsMobilePlatform(GMaxRHIShaderPlatform))
+		{
+			GRenderCommandPipeMode = ERenderCommandPipeMode::None;
+		}
+	}));
+
+FRenderThreadCommandPipe FRenderThreadCommandPipe::Instance;
+
+void FRenderThreadCommandPipe::EnqueueAndLaunch(const TCHAR* Name, uint32& SpecId, TStatId StatId, TUniqueFunction<void(FRHICommandListImmediate&)>&& Function)
+{
+	ensureMsgf(!FTaskTagScope::IsCurrentTag(ETaskTag::EParallelRenderingThread) && !FTaskTagScope::IsCurrentTag(ETaskTag::ERenderingThread),
+		TEXT("Attempted to enqueue a render command from the rendering thread timeline which a pipeline violation. Render commands can only be enqueued from the game thread timeline."));
+
+	Mutex.Lock();
+	bool bWasEmpty = Queues[ProduceIndex].IsEmpty();
+	Queues[ProduceIndex].Emplace(Name, SpecId, StatId, MoveTemp(Function));
+	Mutex.Unlock();
+
+	if (bWasEmpty)
+	{
+		TGraphTask<TFunctionGraphTaskImpl<void(), ESubsequentsMode::FireAndForget>>::CreateTask().ConstructAndDispatchWhenReady([this]
+		{
+			FRHICommandListImmediate& RHICmdList = GetImmediateCommandList_ForRenderCommand();
+
+			Mutex.Lock();
+			TArray<FCommand>& ConsumeCommands = Queues[ProduceIndex];
+			ProduceIndex ^= 1;
+			Mutex.Unlock();
+
+			for (const FCommand& Command : ConsumeCommands)
+			{
+				TRACE_CPUPROFILER_EVENT_SCOPE_USE_ON_CHANNEL(*Command.SpecId, Command.Name, EventScope, RenderCommandsChannel, true);
+				FScopeCycleCounter Scope(Command.StatId, true);
+				Command.Function(RHICmdList);
+			}
+			ConsumeCommands.Reset();
+
+		}, TStatId(), ENamedThreads::GetRenderThread());
+	}
+}
+
+class FRenderCommandPipeRegistry
+{
+public:
+	TLinkedList<FRenderCommandPipe*>* GlobalList = nullptr;
+
+	void Initialize()
+	{
+		Pipes.Reset();
+
+		for (TLinkedList<FRenderCommandPipe*>::TIterator PipeIt(GlobalList); PipeIt; PipeIt.Next())
+		{
+			Pipes.Emplace(*PipeIt);
+		}
+	}
+
+	void StartRecording()
+	{
+		check(IsInGameThread());
+		SCOPED_NAMED_EVENT(FRenderCommandPipe_StartRecording, FColor::Magenta);
+
+		if (bRecording || UE::RenderCommandPipe::GetMode() != ERenderCommandPipeMode::All || Pipes.IsEmpty() || !ShouldExecuteOnRenderThread())
+		{
+			return;
+		}
+
+		UE::Tasks::FTaskEvent TaskEvent{ UE_SOURCE_LOCATION };
+		TArray<FRenderCommandPipe::FFrame*> NextFrames;
+		NextFrames.SetNumZeroed(Pipes.Num());
+		bool bHasNextFrames = false;
+
+		for (int32 PipeIndex = 0; PipeIndex < Pipes.Num(); ++PipeIndex)
+		{
+			FRenderCommandPipe& Pipe = *Pipes[PipeIndex];
+
+			if (Pipe.bEnabled)
+			{
+				Pipe.Frame_GameThread = new FRenderCommandPipe::FFrame(Pipe.Name, TaskEvent);
+				NextFrames[PipeIndex] = Pipe.Frame_GameThread;
+				bHasNextFrames = true;
+			}
+		}
+
+		if (bHasNextFrames)
+		{
+			ENQUEUE_RENDER_COMMAND(RenderCommandPipe_Start)([this, TaskEvent, NextFrames = MoveTemp(NextFrames)](FRHICommandListImmediate&) mutable
+			{
+				RHIResourceLifetimeAddRef();
+
+				for (int32 PipeIndex = 0; PipeIndex < Pipes.Num(); ++PipeIndex)
+				{
+					Pipes[PipeIndex]->Frame_RenderThread = NextFrames[PipeIndex];
+				}
+
+				bReplaying = true;
+				TaskEvent.Trigger();
+			});
+
+			bRecording = true;
+		}
+		else
+		{
+			TaskEvent.Trigger();
+		}
+	}
+
+	void StopRecording()
+	{
+		check(IsInGameThread());
+		SCOPED_NAMED_EVENT(FRenderCommandPipe_StopRecording, FColor::Magenta);
+
+		if (!bRecording)
+		{
+			return;
+		}
+
+		for (FRenderCommandPipe* Pipe : Pipes)
+		{
+			Pipe->Frame_GameThread = nullptr;
+		}
+
+		ENQUEUE_RENDER_COMMAND(RenderCommandPipe_Stop)([this](FRHICommandListImmediate& RHICmdList)
+		{
+			TArray<FRHICommandListImmediate::FQueuedCommandList> QueuedCommandLists;
+			QueuedCommandLists.Reserve(Pipes.Num());
+
+			for (FRenderCommandPipe* Pipe : Pipes)
+			{
+				FRenderCommandPipe::FFrame*& Frame_RenderThread = Pipe->Frame_RenderThread;
+
+				if (Frame_RenderThread)
+				{
+					Frame_RenderThread->Pipe.WaitUntilEmpty();
+
+					if (Frame_RenderThread->RHICmdList)
+					{
+						Frame_RenderThread->RHICmdList->FinishRecording();
+						QueuedCommandLists.Emplace(Frame_RenderThread->RHICmdList);
+					}
+
+					delete Frame_RenderThread;
+					Frame_RenderThread = nullptr;
+				}
+			}
+
+			RHICmdList.QueueAsyncCommandListSubmit(QueuedCommandLists);
+			RHIResourceLifetimeReleaseRef(RHICmdList);
+			bReplaying = false;
+		});
+
+		bRecording = false;
+	}
+
+	bool IsRecording() const
+	{
+		ensureMsgf(!FTaskTagScope::IsCurrentTag(ETaskTag::EParallelRenderingThread) && !FTaskTagScope::IsCurrentTag(ETaskTag::ERenderingThread),
+			TEXT("IsRecording() is not valid from the render thread timeline."));
+
+		return bRecording;
+	}
+
+	bool IsReplaying() const
+	{
+		ensureMsgf(IsInParallelRenderingThread(), TEXT("IsReplaying() is only valid from the render thread timeline."));
+		return bReplaying;
+	}
+
+private:
+	TArray<FRenderCommandPipe*> Pipes;
+	bool bRecording = false;
+	bool bReplaying = false;
+};
+
+static FRenderCommandPipeRegistry GRenderCommandPipeRegistry;
+
+namespace UE::RenderCommandPipe
+{
+	void Initialize()
+	{
+		GRenderCommandPipeRegistry.Initialize();
+	}
+
+	ERenderCommandPipeMode GetMode()
+	{
+		return GRenderCommandPipeMode;
+	}
+
+	bool IsRecording()
+	{
+		return GRenderCommandPipeRegistry.IsRecording();
+	}
+
+	bool IsReplaying()
+	{
+		return GRenderCommandPipeRegistry.IsReplaying();
+	}
+
+	void StartRecording()
+	{
+		GRenderCommandPipeRegistry.StartRecording();
+	}
+
+	void StopRecording()
+	{
+		GRenderCommandPipeRegistry.StopRecording();
+	}
+}
+
+FRenderCommandPipe::FRenderCommandPipe(const TCHAR* InName)
+	: Name(InName)
+	, GlobalListLink(this)
+{
+	GlobalListLink.LinkHead(GRenderCommandPipeRegistry.GlobalList);
+}
+
+FRenderCommandPipe::~FRenderCommandPipe()
+{
+	delete Frame_GameThread;
+	Frame_GameThread = nullptr;
+	delete Frame_RenderThread;
+	Frame_RenderThread = nullptr;
+}
+
+void FRenderCommandPipe::EnqueueAndLaunch(FFunctionVariant&& FunctionVariant, const TCHAR* CommandName, uint32& CommandSpecId)
+{
+	Frame_GameThread->QueueMutex.Lock();
+	bool bWasEmpty = Frame_GameThread->Queue.IsEmpty();
+	Frame_GameThread->Queue.Emplace(MoveTemp(FunctionVariant), CommandName, CommandSpecId);
+	Frame_GameThread->QueueMutex.Unlock();
+
+	if (!bWasEmpty)
+	{
+		return;
+	}
+
+	Frame_GameThread->Pipe.Launch(Name, [this]
+	{
+		check(Frame_RenderThread);
+		SCOPED_NAMED_EVENT_TCHAR(Name, FColor::Magenta);
+		FOptionalTaskTagScope Scope(ETaskTag::EParallelRenderingThread);
+
+		TArray<FCommand> PoppedQueue;
+
+		Frame_RenderThread->QueueMutex.Lock();
+		PoppedQueue = MoveTemp(Frame_RenderThread->Queue);
+		Frame_RenderThread->Queue.Reserve(128);
+		Frame_RenderThread->QueueMutex.Unlock();
+
+		for (FCommand& Command : PoppedQueue)
+		{
+			TRACE_CPUPROFILER_EVENT_SCOPE_USE_ON_CHANNEL(*Command.SpecId, Command.Name, EventScope, RenderCommandsChannel, true);
+
+			if (FCommandListFunction* Function = Command.Function.TryGet<FCommandListFunction>())
+			{
+				if (!Frame_RenderThread->RHICmdList)
+				{
+					FRHICommandList* RHICmdList = new FRHICommandList(FRHIGPUMask::All());
+					RHICmdList->SwitchPipeline(ERHIPipeline::Graphics);
+					Frame_RenderThread->RHICmdList = RHICmdList;
+				}
+
+				(*Function)(*Frame_RenderThread->RHICmdList);
+			}
+			else
+			{
+				Command.Function.Get<FEmptyFunction>()();
+			}
+		}
 
+	}, Frame_GameThread->TaskEvent);
+}
\ No newline at end of file
diff --git a/Engine/Source/Runtime/RenderCore/Public/RenderResource.h b/Engine/Source/Runtime/RenderCore/Public/RenderResource.h
--- a/Engine/Source/Runtime/RenderCore/Public/RenderResource.h
+++ b/Engine/Source/Runtime/RenderCore/Public/RenderResource.h
@@ -36,6 +36,7 @@
 #include "DataDrivenShaderPlatformInfo.h"
 #endif
 
+class FRenderCommandPipe;
 class FRDGPooledBuffer;
 class FResourceArrayInterface;
 
@@ -225,19 +226,19 @@ public:
  * Sends a message to the rendering thread to initialize a resource.
  * This is called in the game thread.
  */
-extern RENDERCORE_API void BeginInitResource(FRenderResource* Resource);
+extern RENDERCORE_API void BeginInitResource(FRenderResource* Resource, FRenderCommandPipe* RenderCommandPipe = nullptr);
 
 /**
  * Sends a message to the rendering thread to update a resource.
  * This is called in the game thread.
  */
-extern RENDERCORE_API void BeginUpdateResourceRHI(FRenderResource* Resource);
+extern RENDERCORE_API void BeginUpdateResourceRHI(FRenderResource* Resource, FRenderCommandPipe* RenderCommandPipe = nullptr);
 
 /**
  * Sends a message to the rendering thread to release a resource.
  * This is called in the game thread.
  */
-extern RENDERCORE_API void BeginReleaseResource(FRenderResource* Resource);
+extern RENDERCORE_API void BeginReleaseResource(FRenderResource* Resource, FRenderCommandPipe* RenderCommandPipe = nullptr);
 
 /**
 * Enables the batching of calls to BeginReleaseResource
diff --git a/Engine/Source/Runtime/RenderCore/Public/RenderingThread.h b/Engine/Source/Runtime/RenderCore/Public/RenderingThread.h
--- a/Engine/Source/Runtime/RenderCore/Public/RenderingThread.h
+++ b/Engine/Source/Runtime/RenderCore/Public/RenderingThread.h
@@ -8,6 +8,7 @@
 
 #include "Async/TaskGraphInterfaces.h"
 #include "Containers/Array.h"
+#include "Containers/List.h"
 #include "CoreGlobals.h"
 #include "CoreMinimal.h"
 #include "CoreTypes.h"
@@ -26,6 +27,8 @@
 #include "Templates/UnrealTemplate.h"
 #include "Templates/UnrealTypeTraits.h"
 #include "Trace/Trace.h"
+#include "Async/Mutex.h"
+#include "Tasks/Pipe.h"
 #if UE_ENABLE_INCLUDE_ORDER_DEPRECATED_IN_5_2
 #include "RenderDeferredCleanup.h"
 #endif
@@ -142,24 +145,6 @@ private:
 
 UE_TRACE_CHANNEL_EXTERN(RenderCommandsChannel, RENDERCORE_API);
 
-/** The parent class of commands stored in the rendering command queue. */
-class FRenderCommand
-{
-public:
-	// All render commands run on the render thread
-	static ENamedThreads::Type GetDesiredThread()
-	{
-		check(!GIsThreadedRendering || ENamedThreads::GetRenderThread() != ENamedThreads::GameThread);
-		return ENamedThreads::GetRenderThread();
-	}
-
-	static ESubsequentsMode::Type GetSubsequentsMode()
-	{
-		// Don't support tasks having dependencies on us, reduces task graph overhead tracking and dealing with subsequents
-		return ESubsequentsMode::FireAndForget;
-	}
-};
-
 //
 // Macros for using render commands.
 //
@@ -183,121 +168,401 @@ DECLARE_STATS_GROUP(TEXT("Render Thread Commands"), STATGROUP_RenderThreadComman
 	#define	ShouldExecuteOnRenderThread()			(LIKELY(GIsThreadedRendering || !IsInGameThread()))
 #endif // UE_SERVER
 
-#define TASK_FUNCTION(Code) \
-		void DoTask(ENamedThreads::Type CurrentThread, const FGraphEventRef& MyCompletionGraphEvent) \
-		{ \
-			FRHICommandListImmediate& RHICmdList = GetImmediateCommandList_ForRenderCommand(); \
-			Code; \
-		} 
+/** Type that contains profiler data necessary to mark up render commands for various profilers. */
+template <typename TSTR>
+struct TRenderCommandTag
+{
+	static const TCHAR* GetName()
+	{
+		return TSTR::TStr();
+	}
 
-#define TASKNAME_FUNCTION(TypeName) \
-		FORCEINLINE TStatId GetStatId() const \
-		{ \
-			RETURN_QUICK_DECLARE_CYCLE_STAT(TypeName, STATGROUP_RenderThreadCommands); \
-		}
+	static uint32& GetSpecId()
+	{
+		static uint32 SpecId;
+		return SpecId;
+	}
+
+	static TStatId GetStatId()
+	{
+#if STATS
+		struct FStatData
+		{
+			typedef FStatGroup_STATGROUP_RenderThreadCommands TGroup;
+			static FORCEINLINE const char* GetStatName()
+			{
+				return TSTR::CStr();
+			}
+			static FORCEINLINE const TCHAR* GetDescription()
+			{
+				return TSTR::TStr();
+			}
+			static FORCEINLINE EStatDataType::Type GetStatType()
+			{
+				return EStatDataType::ST_int64;
+			}
+			static FORCEINLINE bool IsClearEveryFrame()
+			{
+				return true;
+			}
+			static FORCEINLINE bool IsCycleStat()
+			{
+				return true;
+			}
+			static FORCEINLINE FPlatformMemory::EMemoryCounterRegion GetMemoryRegion()
+			{
+				return FPlatformMemory::MCR_Invalid;
+			}
+		};
+		static struct FThreadSafeStaticStat<FStatData> Stat;
+		return Stat.GetStatId();
+#else
+		return TStatId();
+#endif
+	}
+};
 
-template<typename TSTR, typename LAMBDA>
+/** Declares a new render command tag type from a name. */
+#define DECLARE_RENDER_COMMAND_TAG(Type, Name) \
+	struct PREPROCESSOR_JOIN(TSTR_, PREPROCESSOR_JOIN(Name, __LINE__)) \
+	{  \
+		static const char* CStr() { return #Name; } \
+		static const TCHAR* TStr() { return TEXT(#Name); } \
+	}; \
+	using Type = TRenderCommandTag<PREPROCESSOR_JOIN(TSTR_, PREPROCESSOR_JOIN(Name, __LINE__))>;
+
+/** The parent class of commands stored in the rendering command queue. */
+class FRenderCommand
+{
+public:
+	// All render commands run on the render thread
+	static ENamedThreads::Type GetDesiredThread()
+	{
+		check(!GIsThreadedRendering || ENamedThreads::GetRenderThread() != ENamedThreads::GameThread);
+		return ENamedThreads::GetRenderThread();
+	}
+
+	static ESubsequentsMode::Type GetSubsequentsMode()
+	{
+		// Don't support tasks having dependencies on us, reduces task graph overhead tracking and dealing with subsequents
+		return ESubsequentsMode::FireAndForget;
+	}
+};
+
+template <typename TagType, typename LambdaType>
 class TEnqueueUniqueRenderCommandType : public FRenderCommand
 {
 public:
-	TEnqueueUniqueRenderCommandType(LAMBDA&& InLambda) : Lambda(Forward<LAMBDA>(InLambda)) {}
+	TEnqueueUniqueRenderCommandType(LambdaType&& InLambda) : Lambda(Forward<LambdaType>(InLambda)) {}
 
 	void DoTask(ENamedThreads::Type CurrentThread, const FGraphEventRef& MyCompletionGraphEvent)
 	{
-		TRACE_CPUPROFILER_EVENT_SCOPE_ON_CHANNEL_STR(TSTR::TStr(), RenderCommandsChannel);
+		TRACE_CPUPROFILER_EVENT_SCOPE_ON_CHANNEL_STR(TagType::GetName(), RenderCommandsChannel);
 		FRHICommandListImmediate& RHICmdList = GetImmediateCommandList_ForRenderCommand();
 		Lambda(RHICmdList);
 	}
 
 	FORCEINLINE_DEBUGGABLE TStatId GetStatId() const
 	{
-#if STATS
-		static struct FThreadSafeStaticStat<FStat_EnqueueUniqueRenderCommandType> StatPtr_EnqueueUniqueRenderCommandType;
-		return StatPtr_EnqueueUniqueRenderCommandType.GetStatId();
-#else
-		return TStatId();
-#endif
+		return TagType::GetStatId();
 	}
 
 private:
-#if STATS
-	struct FStat_EnqueueUniqueRenderCommandType
+	LambdaType Lambda;
+};
+
+/** Describes which pipes are configured to use the render command pipe system. */
+enum class ERenderCommandPipeMode
+{
+	/** Bypasses the render command pipe system altogether. Render commands are issued using tasks. */
+	None,
+
+	/** The render command pipe on the render thread pipe is active, and all other pipes forward to the render thread pipe. */
+	RenderThread,
+
+	/** All render command pipes are active. */
+	All
+};
+
+namespace UE::RenderCommandPipe
+{
+	// [Game Thread] Initializes all statically initialized render command pipes.
+	extern RENDERCORE_API void Initialize();
+
+	// [Game Thread] Returns the active render command pipe mode.
+	extern RENDERCORE_API ERenderCommandPipeMode GetMode();
+
+	// [Game Thread (Parallel)] Returns whether render command pipes are currently recording on the game thread timeline.
+	extern RENDERCORE_API bool IsRecording();
+
+	// [Render Thread (Parallel)] Returns whether render command pipes are currently replaying commands on the render thread timeline.
+	extern RENDERCORE_API bool IsReplaying();
+
+	// [Game Thread] Starts recording render commands into pipes.
+	extern RENDERCORE_API void StartRecording();
+
+	// [Game Thread] Stops recording commands into pipes and syncs all remaining pipe work to the render thread.
+	extern RENDERCORE_API void StopRecording();
+
+	// [Game Thread] Stops render command pipe recording during the duration of the scope and restarts recording once the scope is complete.
+	class FSyncScope
 	{
-		typedef FStatGroup_STATGROUP_RenderThreadCommands TGroup;
-		static FORCEINLINE const char* GetStatName()
-		{
-			return TSTR::CStr();
-		}
-		static FORCEINLINE const TCHAR* GetDescription()
-		{
-			return TSTR::TStr();
-		}
-		static FORCEINLINE EStatDataType::Type GetStatType()
+	public:
+		FSyncScope()
 		{
-			return EStatDataType::ST_int64;
+			if (IsRecording())
+			{
+				StopRecording();
+				bWasRecording = true;
+			}
 		}
-		static FORCEINLINE bool IsClearEveryFrame()
+
+		~FSyncScope()
 		{
-			return true;
+			if (bWasRecording)
+			{
+				StartRecording();
+			}
 		}
-		static FORCEINLINE bool IsCycleStat()
+
+	private:
+		bool bWasRecording = false;
+	};
+}
+
+class FRenderThreadCommandPipe
+{
+public:
+	template <typename RenderCommandTag, typename LambdaType>
+	FORCEINLINE_DEBUGGABLE static void Enqueue(LambdaType&& Lambda)
+	{
+		if (UE::RenderCommandPipe::GetMode() != ERenderCommandPipeMode::None)
 		{
-			return true;
+			Instance.EnqueueAndLaunch(RenderCommandTag::GetName(), RenderCommandTag::GetSpecId(), RenderCommandTag::GetStatId(), MoveTemp(Lambda));
 		}
-		static FORCEINLINE FPlatformMemory::EMemoryCounterRegion GetMemoryRegion()
+		else
 		{
-			return FPlatformMemory::MCR_Invalid;
+			TGraphTask<TEnqueueUniqueRenderCommandType<RenderCommandTag, LambdaType>>::CreateTask().ConstructAndDispatchWhenReady(MoveTemp(Lambda));
 		}
-	};
-#endif
+	}
 
 private:
-	LAMBDA Lambda;
+	static RENDERCORE_API FRenderThreadCommandPipe Instance;
+
+	RENDERCORE_API void EnqueueAndLaunch(const TCHAR* Name, uint32& SpecId, TStatId StatId, TUniqueFunction<void(FRHICommandListImmediate&)>&& Function);
+
+	struct FCommand
+	{
+		FCommand(const TCHAR* InName, uint32& OutSpecId, TStatId InStatId, TUniqueFunction<void(FRHICommandListImmediate&)>&& InFunction)
+			: Name(InName)
+			, SpecId(&OutSpecId)
+			, StatId(InStatId)
+			, Function(MoveTemp(InFunction))
+		{}
+
+		const TCHAR* Name;
+		uint32* SpecId;
+		TStatId StatId;
+		TUniqueFunction<void(FRHICommandListImmediate&)> Function;
+	};
+
+	int32 ProduceIndex = 0;
+	TStaticArray<TArray<FCommand>, 2> Queues;
+	UE::FMutex Mutex;
 };
 
-template<typename TSTR, typename LAMBDA>
-FORCEINLINE_DEBUGGABLE void EnqueueUniqueRenderCommand(LAMBDA&& Lambda)
+template<typename RenderCommandTag, typename LambdaType>
+FORCEINLINE_DEBUGGABLE void EnqueueUniqueRenderCommand(LambdaType&& Lambda)
 {
-	//QUICK_SCOPE_CYCLE_COUNTER(STAT_EnqueueUniqueRenderCommand);
-	typedef TEnqueueUniqueRenderCommandType<TSTR, LAMBDA> EURCType;
+	TRACE_CPUPROFILER_EVENT_SCOPE_USE_ON_CHANNEL(RenderCommandTag::GetSpecId(), RenderCommandTag::GetName(), EventScope, RenderCommandsChannel, true); 
 
-#if 0 // UE_SERVER && UE_BUILD_DEBUG
-	UE_LOG(LogRHI, Warning, TEXT("Render command '%s' is being executed on a dedicated server."), TSTR::TStr())
-#endif
 	if (IsInRenderingThread())
 	{
-		FRHICommandListImmediate& RHICmdList = GetImmediateCommandList_ForRenderCommand();
-		Lambda(RHICmdList);
+		Lambda(GetImmediateCommandList_ForRenderCommand());
+	}
+	else if (ShouldExecuteOnRenderThread())
+	{
+		CheckNotBlockedOnRenderThread();
+		FRenderThreadCommandPipe::Enqueue<RenderCommandTag, LambdaType>(MoveTemp(Lambda));
 	}
 	else
 	{
-		if (ShouldExecuteOnRenderThread())
+		FScopeCycleCounter CycleScope(RenderCommandTag::GetStatId());
+		Lambda(GetImmediateCommandList_ForRenderCommand());
+	}
+}
+
+template<typename RenderCommandTag, typename LambdaType>
+FORCEINLINE_DEBUGGABLE void EnqueueUniqueRenderCommand(LambdaType& Lambda)
+{
+	static_assert(sizeof(LambdaType) == 0, "EnqueueUniqueRenderCommand enforces use of rvalue and therefore move to avoid an extra copy of the Lambda");
+}
+
+class FRenderCommandPipe
+{
+public:
+	using FCommandListFunction = TUniqueFunction<void(FRHICommandList&)>;
+	using FEmptyFunction = TUniqueFunction<void()>;
+
+	RENDERCORE_API FRenderCommandPipe(const TCHAR* Name);
+	RENDERCORE_API ~FRenderCommandPipe();
+
+	FORCEINLINE const TCHAR* GetName() const
+	{
+		return Name;
+	}
+
+	void SetEnabled(bool bInIsEnabled)
+	{
+		check(IsInGameThread());
+		bEnabled = bInIsEnabled;
+	}
+
+	template <typename RenderCommandTag>
+	static void Enqueue(FRenderCommandPipe* Pipe, FCommandListFunction&& Function)
+	{
+		if (Pipe && Pipe->IsRecording())
 		{
-			CheckNotBlockedOnRenderThread();
-			TGraphTask<EURCType>::CreateTask().ConstructAndDispatchWhenReady(Forward<LAMBDA>(Lambda));
+			Pipe->EnqueueAndLaunch(MoveTemp(Function), RenderCommandTag::GetName(), RenderCommandTag::GetSpecId());
 		}
 		else
 		{
-			EURCType TempCommand(Forward<LAMBDA>(Lambda));
-			FScopeCycleCounter EURCMacro_Scope(TempCommand.GetStatId());
-			TempCommand.DoTask(ENamedThreads::GameThread, FGraphEventRef());
+			EnqueueUniqueRenderCommand<RenderCommandTag>([Function = MoveTemp(Function)](FRHICommandListImmediate& RHICmdList) { Function(RHICmdList); });
 		}
 	}
-}
 
-#define ENQUEUE_RENDER_COMMAND(Type) \
-	struct Type##Name \
-	{  \
-		static const char* CStr() { return #Type; } \
-		static const TCHAR* TStr() { return TEXT(#Type); } \
-	}; \
-	EnqueueUniqueRenderCommand<Type##Name>
+	template <typename RenderCommandTag>
+	static void Enqueue(FRenderCommandPipe& Pipe, FCommandListFunction&& Function)
+	{
+		if (Pipe.IsRecording())
+		{
+			Pipe.EnqueueAndLaunch(MoveTemp(Function), RenderCommandTag::GetName(), RenderCommandTag::GetSpecId());
+		}
+		else
+		{
+			EnqueueUniqueRenderCommand<RenderCommandTag>([Function = MoveTemp(Function)](FRHICommandListImmediate& RHICmdList) { Function(RHICmdList); });
+		}
+	}
 
-template<typename TSTR, typename LAMBDA>
-FORCEINLINE_DEBUGGABLE void EnqueueUniqueRenderCommand(LAMBDA& Lambda)
-{
-	static_assert(sizeof(LAMBDA) == 0, "EnqueueUniqueRenderCommand enforces use of rvalue and therefore move to avoid an extra copy of the Lambda");
-}
+	template <typename RenderCommandTag>
+	static void Enqueue(FRenderCommandPipe* Pipe, FEmptyFunction&& Function)
+	{
+		if (Pipe && Pipe->IsRecording())
+		{
+			Pipe->EnqueueAndLaunch(MoveTemp(Function), RenderCommandTag::GetName(), RenderCommandTag::GetSpecId());
+		}
+		else
+		{
+			EnqueueUniqueRenderCommand<RenderCommandTag>([Function = MoveTemp(Function)](FRHICommandListImmediate&) { Function(); });
+		}
+	}
+
+	template <typename RenderCommandTag>
+	static void Enqueue(FRenderCommandPipe& Pipe, FEmptyFunction&& Function)
+	{
+		if (Pipe.IsRecording())
+		{
+			Pipe.EnqueueAndLaunch(MoveTemp(Function), RenderCommandTag::GetName(), RenderCommandTag::GetSpecId());
+		}
+		else
+		{
+			EnqueueUniqueRenderCommand<RenderCommandTag>([Function = MoveTemp(Function)](FRHICommandListImmediate&) { Function(); });
+		}
+	}
+
+	template <typename RenderCommandTag, typename LambdaType>
+	FORCEINLINE static void Enqueue(LambdaType&& Lambda)
+	{
+		EnqueueUniqueRenderCommand<RenderCommandTag>(MoveTemp(Lambda));
+	}
+
+private:
+	friend class FRenderCommandPipeRegistry;
+
+	using FFunctionVariant = TVariant<FEmptyFunction, FCommandListFunction>;
+
+	RENDERCORE_API void EnqueueAndLaunch(FFunctionVariant&& FunctionVariant, const TCHAR* Name, uint32& SpecId);
+
+	void EnqueueAndLaunch(FCommandListFunction&& Function, const TCHAR* CommandName, uint32& CommandSpecId)
+	{
+		EnqueueAndLaunch(FFunctionVariant(TInPlaceType<FCommandListFunction>(), MoveTemp(Function)), CommandName, CommandSpecId);
+	}
+
+	void EnqueueAndLaunch(FEmptyFunction&& Function, const TCHAR* CommandName, uint32& CommandSpecId)
+	{
+		EnqueueAndLaunch(FFunctionVariant(TInPlaceType<FEmptyFunction>(), MoveTemp(Function)), CommandName, CommandSpecId);
+	}
+
+	FORCEINLINE bool IsRecording() const
+	{
+		return Frame_GameThread != nullptr;
+	}
+
+	struct FCommand
+	{
+		FCommand(FFunctionVariant&& InFunction, const TCHAR* InName, uint32& InOutSpecId)
+			: Function(MoveTemp(InFunction))
+			, Name(InName)
+			, SpecId(&InOutSpecId)
+		{}
+
+		FFunctionVariant Function;
+		const TCHAR* Name;
+		uint32* SpecId;
+	};
+
+	struct FFrame
+	{
+		FFrame(const TCHAR* Name, const UE::Tasks::FTaskEvent& InTaskEvent)
+			: Pipe(Name)
+			, TaskEvent(InTaskEvent)
+		{}
+
+		UE::Tasks::FPipe Pipe;
+		UE::Tasks::FTaskEvent TaskEvent;
+		TArray<FCommand> Queue;
+		UE::FMutex QueueMutex;
+		FRHICommandList* RHICmdList = nullptr;
+	};
+
+	const TCHAR* Name;
+	FFrame* Frame_GameThread = nullptr;
+	FFrame* Frame_RenderThread = nullptr;
+	TLinkedList<FRenderCommandPipe*> GlobalListLink;
+	bool bEnabled = true;
+};
+
+/** Declares an extern reference to a render command pipe. */
+#define DECLARE_RENDER_COMMAND_PIPE(Name, PrefixKeywords) \
+	namespace UE::RenderCommandPipe { extern PrefixKeywords FRenderCommandPipe Name; }
+
+/** Defines a render command pipe. */
+#define DEFINE_RENDER_COMMAND_PIPE(Name) \
+	namespace UE::RenderCommandPipe \
+	{ \
+		FRenderCommandPipe Name(TEXT(#Name)); \
+		namespace \
+		{ \
+			FAutoConsoleVariable CVar_##Name( \
+				TEXT("r.RenderCommandPipe." #Name), true, \
+				TEXT("Whether to enable the " #Name " Render Command Pipe") \
+				TEXT(" 0: off;") \
+				TEXT(" 1: on (default)"), \
+				FConsoleVariableDelegate::CreateLambda([](IConsoleVariable* Variable) \
+			{ \
+				UE::RenderCommandPipe::Name.SetEnabled(Variable->GetBool()); \
+			}), \
+			ECVF_Default); \
+		} \
+	}
+
+/** Enqueues a render command to a render pipe. The default implementation takes a lambda and schedules on the render thread.
+ *  Alternative implementations accept either a reference or pointer to an FRenderCommandPipe instance to schedule on an async
+ *  pipe, if enabled.
+ */
+#define ENQUEUE_RENDER_COMMAND(Type) \
+	DECLARE_RENDER_COMMAND_TAG(PREPROCESSOR_JOIN(FRenderCommandTag_, PREPROCESSOR_JOIN(Type, __LINE__)), Type) \
+	FRenderCommandPipe::Enqueue<PREPROCESSOR_JOIN(FRenderCommandTag_, PREPROCESSOR_JOIN(Type, __LINE__))>
 
 ////////////////////////////////////
 // RenderThread scoped work
diff --git a/Engine/Source/Runtime/Renderer/Private/PlanarReflectionRendering.cpp b/Engine/Source/Runtime/Renderer/Private/PlanarReflectionRendering.cpp
--- a/Engine/Source/Runtime/Renderer/Private/PlanarReflectionRendering.cpp
+++ b/Engine/Source/Runtime/Renderer/Private/PlanarReflectionRendering.cpp
@@ -706,6 +706,8 @@ void FScene::UpdatePlanarReflectionContents(UPlanarReflectionComponent* CaptureC
 			FPlanarReflectionSceneProxy* SceneProxyPtr = CaptureComponent->SceneProxy;
 			FPlanarReflectionRenderTarget* RenderTargetPtr = CaptureComponent->RenderTarget;
 
+			UE::RenderCommandPipe::FSyncScope SyncScope;
+
 			if (bIsMobilePixelProjectedReflectionEnabled)
 			{
 				ENQUEUE_RENDER_COMMAND(CaptureCommand)(
diff --git a/Engine/Source/Runtime/Renderer/Private/ReflectionEnvironmentCapture.cpp b/Engine/Source/Runtime/Renderer/Private/ReflectionEnvironmentCapture.cpp
--- a/Engine/Source/Runtime/Renderer/Private/ReflectionEnvironmentCapture.cpp
+++ b/Engine/Source/Runtime/Renderer/Private/ReflectionEnvironmentCapture.cpp
@@ -1343,6 +1343,8 @@ void CaptureSceneIntoScratchCubemap(
 
 		FSceneRenderer* SceneRenderer = FSceneRenderer::CreateSceneRenderer(&ViewFamily, nullptr);
 
+		UE::RenderCommandPipe::FSyncScope SyncScope;
+
 		ENQUEUE_RENDER_COMMAND(CaptureCommand)(
 			[SceneRenderer, &ReflectionCubemapTexture, CubeFace, CubemapSize, bCapturingForSkyLight, bLowerHemisphereIsBlack, LowerHemisphereColor, bCapturingForMobile, bInsideTick](FRHICommandListImmediate& RHICmdList)
 		{
diff --git a/Engine/Source/Runtime/Renderer/Private/RendererScene.cpp b/Engine/Source/Runtime/Renderer/Private/RendererScene.cpp
--- a/Engine/Source/Runtime/Renderer/Private/RendererScene.cpp
+++ b/Engine/Source/Runtime/Renderer/Private/RendererScene.cpp
@@ -4321,6 +4321,8 @@ void FScene::UpdateStaticDrawLists_RenderThread(FRHICommandListImmediate& RHICmd
 
 void FScene::UpdateStaticDrawLists()
 {
+	UE::RenderCommandPipe::FSyncScope SyncScope;
+
 	FScene* Scene = this;
 	ENQUEUE_RENDER_COMMAND(FUpdateDrawLists)(
 		[Scene](FRHICommandListImmediate& RHICmdList)
@@ -4391,6 +4393,8 @@ void FScene::Release()
 
 	GetRendererModule().RemoveScene(this);
 
+	UE::RenderCommandPipe::FSyncScope SyncScope;
+
 	// Send a command to the rendering thread to release the scene.
 	FScene* Scene = this;
 	ENQUEUE_RENDER_COMMAND(FReleaseCommand)(
@@ -4543,6 +4547,8 @@ void FScene::Export( FArchive& Ar ) const
 
 void FScene::ApplyWorldOffset(const FVector& InOffset)
 {
+	UE::RenderCommandPipe::FSyncScope SyncScope;
+
 	// Send a command to the rendering thread to shift scene data
 	FScene* Scene = this;
 	FVector Offset = InOffset;
@@ -5212,6 +5218,7 @@ void FScene::UpdateAllPrimitiveSceneInfos(FRDGBuilder& GraphBuilder, EUpdateAllP
 	SCOPE_CYCLE_COUNTER(STAT_UpdateScenePrimitiveRenderThreadTime);
 
 	check(IsInRenderingThread());
+	check(!UE::RenderCommandPipe::IsReplaying());
 
 	UE::Tasks::FTask GPUSkinCacheTask;
 
@@ -6727,6 +6734,9 @@ void UpdateStaticMeshesForMaterials(const TArray<const FMaterial*>& MaterialReso
 			}
 		}
 	}
+
+	UE::RenderCommandPipe::FSyncScope SyncScope;
+
 	ENQUEUE_RENDER_COMMAND(FUpdateStaticMeshesForMaterials)(
 		[UsedPrimitives = MoveTemp(UsedPrimitives)](FRHICommandListImmediate& RHICmdList) mutable
 		{
diff --git a/Engine/Source/Runtime/Renderer/Private/SceneCaptureRendering.cpp b/Engine/Source/Runtime/Renderer/Private/SceneCaptureRendering.cpp
--- a/Engine/Source/Runtime/Renderer/Private/SceneCaptureRendering.cpp
+++ b/Engine/Source/Runtime/Renderer/Private/SceneCaptureRendering.cpp
@@ -1005,6 +1005,8 @@ void FScene::UpdateSceneCaptureContents(USceneCaptureComponent2D* CaptureCompone
 			Extension->BeginRenderViewFamily(SceneRenderer->ViewFamily);
 		}
 
+		UE::RenderCommandPipe::FSyncScope SyncScope;
+
 		ENQUEUE_RENDER_COMMAND(CaptureCommand)(
 			[SceneRenderer, TextureRenderTargetResource, TexturePtrNotDeferenced, EventName, TargetName, bGenerateMips, GenerateMipsParams, GameViewportRT, bEnableOrthographicTiling, bIsCompositing, bOrthographicCamera, NumXTiles, NumYTiles, TileID, CaptureMemorySize](FRHICommandListImmediate& RHICmdList)
 			{
@@ -1178,6 +1180,8 @@ void FScene::UpdateSceneCaptureContents(USceneCaptureComponentCube* CaptureCompo
 				Extension->BeginRenderViewFamily(SceneRenderer->ViewFamily);
 			}
 
+			UE::RenderCommandPipe::FSyncScope SyncScope;
+
 			ENQUEUE_RENDER_COMMAND(CaptureCommand)(
 				[SceneRenderer, TextureRenderTarget, EventName, TargetFace](FRHICommandListImmediate& RHICmdList)
 				{
diff --git a/Engine/Source/Runtime/Renderer/Private/SceneRendering.cpp b/Engine/Source/Runtime/Renderer/Private/SceneRendering.cpp
--- a/Engine/Source/Runtime/Renderer/Private/SceneRendering.cpp
+++ b/Engine/Source/Runtime/Renderer/Private/SceneRendering.cpp
@@ -4727,6 +4727,8 @@ void FRendererModule::BeginRenderingViewFamilies(FCanvas* Canvas, TArrayView<FSc
 			GFastVRamConfig.Update();
 		});
 
+	UE::RenderCommandPipe::FSyncScope SyncScope;
+
 	// Flush the canvas first.
 	Canvas->Flush_GameThread();
 
