Fixed edge case in RDG builder where buffer would fail the transient allocator and not fall back to a pooled allocation.

diff --git a/Engine/Source/Runtime/RenderCore/Private/RenderGraphBuilder.cpp b/Engine/Source/Runtime/RenderCore/Private/RenderGraphBuilder.cpp
--- a/Engine/Source/Runtime/RenderCore/Private/RenderGraphBuilder.cpp
+++ b/Engine/Source/Runtime/RenderCore/Private/RenderGraphBuilder.cpp
@@ -3688,7 +3688,7 @@ void FRDGBuilder::BeginResourceRHI(FRDGPassHandle PassHandle, FRDGBufferRef Buff
 		}
 	}
 
-	if (!Buffer->bTransient)
+	if (!Buffer->ResourceRHI)
 	{
 		const ERDGPooledBufferAlignment Alignment = Buffer->bQueuedForUpload ? ERDGPooledBufferAlignment::PowerOfTwo : ERDGPooledBufferAlignment::Page;
 
@@ -3696,7 +3696,7 @@ void FRDGBuilder::BeginResourceRHI(FRDGPassHandle PassHandle, FRDGBufferRef Buff
 	}
 
 #if RHI_ENABLE_RESOURCE_INFO
-	if (Buffer->HasRHI())
+	if (Buffer->ResourceRHI)
 	{
 		Buffer->ResourceRHI->SetOwnerName(Buffer->OwnerName);
 	}
