Temporarily disabled async render command pipes to investigate crash.

diff --git a/Engine/Source/Runtime/RenderCore/Private/RenderingThread.cpp b/Engine/Source/Runtime/RenderCore/Private/RenderingThread.cpp
--- a/Engine/Source/Runtime/RenderCore/Private/RenderingThread.cpp
+++ b/Engine/Source/Runtime/RenderCore/Private/RenderingThread.cpp
@@ -1521,7 +1521,7 @@ inline ERenderCommandPipeMode GetValidatedRenderCommandPipeMode(int32 CVarValue)
 ERenderCommandPipeMode GRenderCommandPipeMode = ERenderCommandPipeMode::None;
 FAutoConsoleVariable CVarRenderCommandPipeMode(
 	TEXT("r.RenderCommandPipeMode"),
-	2,
+	1,
 	TEXT("Controls behavior of the main render thread command pipe.")
 	TEXT(" 0: Render commands are launched individually as tasks;\n")
 	TEXT(" 1: Render commands are enqueued into a render command pipe for the render thread only.;\n")
