RenderCore: Fixed crashes that could occur while using DumpGPU console command if writes to the drive are slow.

#rnx

diff --git a/Engine/Source/Runtime/RenderCore/Internal/RenderThreadTimeoutControl.h b/Engine/Source/Runtime/RenderCore/Internal/RenderThreadTimeoutControl.h
new file mode 100644
--- /dev/null
+++ b/Engine/Source/Runtime/RenderCore/Internal/RenderThreadTimeoutControl.h
@@ -0,0 +1,6 @@
+// Copyright Epic Games, Inc. All Rights Reserved.
+
+#pragma once
+
+extern RENDERCORE_API void SuspendRenderThreadTimeout();
+extern RENDERCORE_API void ResumeRenderThreadTimeout();
diff --git a/Engine/Source/Runtime/RenderCore/Private/DumpGPU.cpp b/Engine/Source/Runtime/RenderCore/Private/DumpGPU.cpp
--- a/Engine/Source/Runtime/RenderCore/Private/DumpGPU.cpp
+++ b/Engine/Source/Runtime/RenderCore/Private/DumpGPU.cpp
@@ -4,6 +4,7 @@
 #include "HAL/PlatformFileManager.h"
 #include "HAL/PlatformMisc.h"
 #include "HAL/PlatformOutputDevices.h"
+#include "HAL/ThreadHeartBeat.h"
 #include "Async/AsyncWork.h"
 #include "Misc/App.h"
 #include "Misc/FileHelper.h"
@@ -23,6 +24,7 @@
 #include "RHIUtilities.h"
 #include "RHIValidation.h"
 #include "RenderGraphPrivate.h"
+#include "RenderThreadTimeoutControl.h"
 
 DEFINE_LOG_CATEGORY_STATIC(LogDumpGPU, Log, All);
 
@@ -2149,6 +2151,11 @@ FString FRDGBuilder::BeginResourceDump(const TCHAR* Cmd)
 		GNextRDGResourceDumpContext = nullptr;
 	}
 
+	// Dumping resource may take a while and we don't want to get 'hang' crashes in the process.
+	// We resume from EndResourceDump after the dump has finished.
+	FThreadHeartBeat::Get().SuspendHeartBeat(true);
+	SuspendRenderThreadTimeout();
+
 	UE_LOG(LogDumpGPU, Display, TEXT("DumpGPU to %s starting this frame"), *NewResourceDumpContext->DumpingDirectoryPath);
 	NewResourceDumpContext->MemoryConstants = FPlatformMemory::GetConstants();
 	NewResourceDumpContext->MemoryStats = FPlatformMemory::GetStats();
@@ -2436,6 +2443,10 @@ void FRDGBuilder::EndResourceDump()
 	delete GRDGResourceDumpContext;
 	GRDGResourceDumpContext = nullptr;
 	DumpingFrameCounter_GameThread = 0;
+
+	// It matches SuspendHeartBeat from BeginResourceDump.
+	FThreadHeartBeat::Get().ResumeHeartBeat(true);
+	ResumeRenderThreadTimeout();
 }
 
 static const TCHAR* GetPassEventNameWithGPUMask(const FRDGPass* Pass, FString& OutNameStorage)
diff --git a/Engine/Source/Runtime/RenderCore/Private/RenderingThread.cpp b/Engine/Source/Runtime/RenderCore/Private/RenderingThread.cpp
--- a/Engine/Source/Runtime/RenderCore/Private/RenderingThread.cpp
+++ b/Engine/Source/Runtime/RenderCore/Private/RenderingThread.cpp
@@ -1001,6 +1001,25 @@ void StopRenderCommandFenceBundler()
 	BundledCompletionEvent.Reset();
 }
 
+std::atomic<int> GTimeoutSuspendCount;
+
+void SuspendRenderThreadTimeout()
+{
+	++GTimeoutSuspendCount;
+}
+
+void ResumeRenderThreadTimeout()
+{
+	--GTimeoutSuspendCount;
+
+	check(GTimeoutSuspendCount >= 0);
+}
+
+bool IsRenderThreadTimeoutSuspended()
+{
+	return GTimeoutSuspendCount > 0;
+}
+
 TAutoConsoleVariable<int32> CVarGTSyncType(
 	TEXT("r.GTSyncType"),
 	0,
@@ -1224,7 +1243,7 @@ static void GameThreadWaitForTask(const UE::Tasks::FTask& Task, bool bEmptyGameT
 				// editor threads can block for quite a while... 
 				if (!bDone && !bRenderThreadEnsured)
 				{
-					if (bOverdue && !bDisabled && !FPlatformMisc::IsDebuggerPresent())
+					if (bOverdue && !bDisabled && !IsRenderThreadTimeoutSuspended() && !FPlatformMisc::IsDebuggerPresent())
 					{
 						UE_LOG(LogRendererCore, Fatal, TEXT("GameThread timed out waiting for RenderThread after %.02f secs"), RenderThreadTimeoutClock.Seconds() - StartTime);
 					}
