Shader Compiler Job Cache:  Fix for bugs caused by sharing incompatible results generated by single and pipeline jobs with one shader that happen to have the same hash, due to the pipeline job hash just being the sum of the individual shader hashes.  Solution is to set the bCompilingForShaderPipeline on the first pipeline stage job on construction (rather than later during compile), so it produces a different input hash.  Also need to include bCompilingForShaderPipeline as part of the input hash when generating the hash from preprocessed results.

Fixing this bug allows enabling the parallel SubmitJobs optimization once again.  Parallel SubmitJobs, combined with the bug of incompatible sharing of results, meant that single shader jobs would sometimes use pipeline results, leading to runtime crashes.  Single threaded job submit hid the bug, because single shader jobs were always issued first, and took precedence over pipeline jobs in the cache.  Other minor changes:  validate input hash in the output data returned by shader compiler tasks or workers, and use a reference pointer in code that attempts to change job priority, to ensure it can't be deleted by another thread during that time.

#rnx

diff --git a/Engine/Source/Runtime/Core/Private/UObject/DevObjectVersion.cpp b/Engine/Source/Runtime/Core/Private/UObject/DevObjectVersion.cpp
--- a/Engine/Source/Runtime/Core/Private/UObject/DevObjectVersion.cpp
+++ b/Engine/Source/Runtime/Core/Private/UObject/DevObjectVersion.cpp
@@ -118,11 +118,11 @@ FGuid FDevSystemGuidRegistry::GetSystemGuid(FGuid System)
 }
 
 FDevSystemGuids::FDevSystemGuids()
-	: GLOBALSHADERMAP_DERIVEDDATA_VER(0x7912A706, 0x52B8450A, 0x9CD66FA7, 0xEFBBFA0A)
+	: GLOBALSHADERMAP_DERIVEDDATA_VER(0x7BB10A3C, 0xAC4E46F3, 0xAC78F4C1, 0xEFB3E34F)
 	, GROOM_BINDING_DERIVED_DATA_VERSION(0x30769E53, 0x0C574C7B, 0xA15C56F2, 0x24A64E32)
 	, GROOM_DERIVED_DATA_VERSION(0x05A37379, 0xF8A049D5, 0x986824E9, 0xAAA83F41)
 	, LANDSCAPE_MOBILE_COOK_VERSION(0x0E9ADF72, 0xD6B64E0D, 0x81C4A92B, 0x081A37AB)
-	, MATERIALSHADERMAP_DERIVEDDATA_VER(0x2579AAFE, 0x8F1D4E4F, 0xB04278C6, 0x35917C17)
+	, MATERIALSHADERMAP_DERIVEDDATA_VER(0x9131A169, 0x3E004B0D, 0x864390EF, 0xA1F12934)
 	, NANITE_DERIVEDDATA_VER(0xBEB0226A, 0x070E4ECA, 0x972C1E7D, 0xD8599E68)
 	, NIAGARASHADERMAP_DERIVEDDATA_VER(0x7BBD9913, 0xC1554D20, 0xADAE9F17, 0xB006299E)
 	, Niagara_LatestScriptCompileVersion(0x6D32B8EE, 0x909FCA7E, 0xA5CE4F17, 0x066A5F25)
diff --git a/Engine/Source/Runtime/Engine/Private/ShaderCompiler/ShaderCompiler.cpp b/Engine/Source/Runtime/Engine/Private/ShaderCompiler/ShaderCompiler.cpp
--- a/Engine/Source/Runtime/Engine/Private/ShaderCompiler/ShaderCompiler.cpp
+++ b/Engine/Source/Runtime/Engine/Private/ShaderCompiler/ShaderCompiler.cpp
@@ -241,7 +241,7 @@ static FAutoConsoleVariableRef CVarShaderCompilerDebugDiscardCacheOutputs(
 	ECVF_Default
 );
 
-int32 GShaderCompilerParallelSubmitJobs = 0;
+int32 GShaderCompilerParallelSubmitJobs = 1;
 static FAutoConsoleVariableRef CVarShaderCompilerParallelSubmitJobs(
 	TEXT("r.ShaderCompiler.ParallelSubmitJobs"),
 	GShaderCompilerParallelSubmitJobs,
@@ -481,7 +481,7 @@ public:
 	static const int32 STRIPE_SHIFT = 32 - NUM_STRIPE_BITS;
 
 	template<typename JobType, typename KeyType>
-	JobType* PrepareJob(uint32 InId, const KeyType& InKey, EShaderCompileJobPriority InPriority, bool& bOutNewJob)
+	FShaderCommonCompileJobPtr PrepareJob(uint32 InId, const KeyType& InKey, EShaderCompileJobPriority InPriority, bool& bOutNewJob)
 	{
 		const uint32 Hash = InKey.MakeHash(InId);
 		FLockStripeData& Stripe = GetStripe(JobType::Type, Hash);
@@ -611,12 +611,14 @@ public:
 	JobType* PrepareJob(uint32 InId, const KeyType& InKey, EShaderCompileJobPriority InPriority)
 	{
 		bool bNewJob;
-		JobType* Result = JobTable.PrepareJob<JobType>(InId, InKey, InPriority, bNewJob);
+		FShaderCommonCompileJobPtr Result = JobTable.PrepareJob<JobType>(InId, InKey, InPriority, bNewJob);
 
 		if (bNewJob)
 		{
-			// If it's a new job, return it
-			return Result;
+			// If it's a new job, return it -- it's OK to cast the ref-counted pointer to a raw pointer, because JobTable
+			// itself has a reference to the job, and a newly added job hasn't been submitted yet, so it can't make a
+			// round trip through the pipeline and be released until that happens.
+			return (JobType*)Result.GetReference();
 		}
 		else if (InPriority > Result->Priority)
 		{
@@ -7074,6 +7076,8 @@ void GlobalBeginCompileShader(
 
 	FShaderCompileUtilities::GenerateBrdfHeaders(ShaderPlatform);
 
+	// NOTE:  Input.bCompilingForShaderPipeline is initialized by the constructor for single versus pipeline jobs, do not initialize again here!
+
 	Input.Target = Target;
 	Input.ShaderPlatformName = FDataDrivenShaderPlatformInfo::GetName(ShaderPlatform);
 	Input.ShaderFormat = ShaderFormatName;
@@ -7081,7 +7085,6 @@ void GlobalBeginCompileShader(
 	GetShaderCompressionOodleSettings(Input.OodleCompressor, Input.OodleLevel);
 	Input.VirtualSourceFilePath = SourceFilename;
 	Input.EntryPointName = FunctionName;
-	Input.bCompilingForShaderPipeline = false;
 	Input.bIncludeUsedOutputs = false;
 	Input.DumpDebugInfoRootPath = GShaderCompilingManager->GetAbsoluteShaderDebugInfoDirectory() / Input.ShaderPlatformName.ToString();
 	Input.DebugInfoFlags = GShaderCompilingManager->GetDumpShaderDebugInfoFlags();
diff --git a/Engine/Source/Runtime/RenderCore/Private/ShaderCore.cpp b/Engine/Source/Runtime/RenderCore/Private/ShaderCore.cpp
--- a/Engine/Source/Runtime/RenderCore/Private/ShaderCore.cpp
+++ b/Engine/Source/Runtime/RenderCore/Private/ShaderCore.cpp
@@ -1291,6 +1291,8 @@ public:
 				Compiler->CompileShader(Job.Input.ShaderFormat, Job.Input, Job.Output, WorkingDirectory);
 			}
 
+			// set the input hash on the output to allow sanity checking the job results on completion/reading from cache
+			Job.Output.InputHash = Job.Input.Hash;
 		}
 #if PLATFORM_WINDOWS
 		__except(HandleShaderCompileException(GetExceptionInformation(), OutExceptionMsg, OutExceptionCallstack))
@@ -1458,7 +1460,8 @@ void CompileShaderPipeline(const TArray<const IShaderFormat*>& ShaderFormats, FS
 	checkf(PipelineJob->StageJobs.Num() > 0, TEXT("Pipeline %s has zero jobs!"), PipelineJob->Key.ShaderPipeline->GetName());
 	FShaderCompileJob* CurrentJob = PipelineJob->StageJobs[0]->GetSingleShaderJob();
 
-	CurrentJob->Input.bCompilingForShaderPipeline = true;
+	// Flag should be set on the first job when the FShaderPipelineCompileJob was constructed, to ensure the flag is included when computing the input hash.
+	check(CurrentJob->Input.bCompilingForShaderPipeline == true);
 
 	// First job doesn't have to trim outputs
 	CurrentJob->Input.bIncludeUsedOutputs = false;
@@ -3126,6 +3129,10 @@ FShaderCommonCompileJob::FInputHash FShaderCompileJob::GetInputHash()
 		FShaderTarget Target = Input.Target;
 		Hasher << Target;
 		Hasher << Input.EntryPointName;
+
+		// Include this flag, so shader pipeline jobs get a different hash from single shader jobs, even if the preprocessed shader is otherwise the same.
+		Hasher << Input.bCompilingForShaderPipeline;
+
 		FShaderCompilerEnvironment MergedEnvironment = Input.Environment;
 		if (Input.SharedEnvironment)
 		{
@@ -3300,6 +3307,8 @@ void FShaderCompileJob::SerializeOutput(FArchive& Ar)
 		Output.CompileTime = ActualCompileTime;
 		Output.PreprocessTime = ActualPreprocessTime;
 	}
+
+	checkf(Output.InputHash == Input.Hash, TEXT("Failed sanity check: InputHash reported from compile output does not match InputHash from compile input"));
 }
 
 void FShaderCompileJob::OnComplete()
@@ -3413,6 +3422,12 @@ FShaderPipelineCompileJob::FShaderPipelineCompileJob(int32 NumStages)
 	{
 		StageJobs.Add(new FShaderCompileJob());
 	}
+
+	if (StageJobs.Num())
+	{
+		// Set this flag on first job in constructor, so it's included during input hash computation.  Flag is set conditionally for other stage jobs in CompileShaderPipeline.
+		StageJobs[0]->Input.bCompilingForShaderPipeline = true;
+	}
 }
 
 FShaderPipelineCompileJob::FShaderPipelineCompileJob(uint32 InHash, uint32 InId, EShaderCompileJobPriority InPriroity, const FShaderPipelineCompileJobKey& InKey) :
@@ -3426,6 +3441,12 @@ FShaderPipelineCompileJob::FShaderPipelineCompileJob(uint32 InHash, uint32 InId,
 		const FShaderCompileJobKey StageKey(ShaderType, InKey.VFType, InKey.PermutationId);
 		StageJobs.Add(new FShaderCompileJob(StageKey.MakeHash(InId), InId, InPriroity, StageKey));
 	}
+	
+	if (StageJobs.Num())
+	{
+		// Set this flag on first job in constructor, so it's included during input hash computation.  Flag is set conditionally for other stage jobs in CompileShaderPipeline.
+		StageJobs[0]->Input.bCompilingForShaderPipeline = true;
+	}
 }
 
 void FShaderPipelineCompileJob::SerializeOutput(FArchive& Ar)
diff --git a/Engine/Source/Runtime/RenderCore/Public/ShaderCompilerCore.h b/Engine/Source/Runtime/RenderCore/Public/ShaderCompilerCore.h
--- a/Engine/Source/Runtime/RenderCore/Public/ShaderCompilerCore.h
+++ b/Engine/Source/Runtime/RenderCore/Public/ShaderCompilerCore.h
@@ -29,7 +29,7 @@ class FShaderPipelineCompileJob;
 // this is for the protocol, not the data, bump if FShaderCompilerInput or ProcessInputFromArchive changes.
 inline const int32 ShaderCompileWorkerInputVersion = 21;
 // this is for the protocol, not the data, bump if FShaderCompilerOutput or WriteToOutputArchive changes.
-inline const int32 ShaderCompileWorkerOutputVersion = 16;
+inline const int32 ShaderCompileWorkerOutputVersion = 17;
 // this is for the protocol, not the data.
 inline const int32 ShaderCompileWorkerSingleJobHeader = 'S';
 // this is for the protocol, not the data.
@@ -565,6 +565,7 @@ struct FShaderCompilerOutput
 	TArray<FString> PragmaDirectives;
 	FShaderTarget Target;
 	FShaderCode ShaderCode;
+	FBlake3Hash InputHash;
 	FSHAHash OutputHash;
 	uint32 NumInstructions;
 	uint32 NumTextureSamplers;
@@ -629,7 +630,7 @@ struct FShaderCompilerOutput
 	friend FArchive& operator<<(FArchive& Ar, FShaderCompilerOutput& Output)
 	{
 		// Note: this serialize is used to pass between UE and the shader compile worker, recompile both when modifying
-		Ar << Output.ParameterMap << Output.Errors << Output.Target << Output.ShaderCode << Output.OutputHash << Output.NumInstructions << Output.NumTextureSamplers << Output.bSucceeded;
+		Ar << Output.ParameterMap << Output.Errors << Output.Target << Output.ShaderCode << Output.InputHash << Output.OutputHash << Output.NumInstructions << Output.NumTextureSamplers << Output.bSucceeded;
 		Ar << Output.bSupportsQueryingUsedAttributes << Output.UsedAttributes;
 		Ar << Output.CompileTime;
 		Ar << Output.PreprocessTime;
