Shader Compiler:  Fix for serialization of old DDC data for ODSC -- add version for FShaderCompilerOutput structure, which is included in the input hash for shaders, to prevent old DDC data from being used.

#rnx

diff --git a/Engine/Source/Runtime/RenderCore/Private/ShaderCore.cpp b/Engine/Source/Runtime/RenderCore/Private/ShaderCore.cpp
--- a/Engine/Source/Runtime/RenderCore/Private/ShaderCore.cpp
+++ b/Engine/Source/Runtime/RenderCore/Private/ShaderCore.cpp
@@ -3126,6 +3126,10 @@ FShaderCommonCompileJob::FInputHash FShaderCompileJob::GetInputHash()
 	if (Input.bCachePreprocessed)
 	{
 		FMemoryHasherBlake3 Hasher;
+
+		int32 FShaderCompilerOutputStructVersionLocal = FShaderCompilerOutputStructVersion;
+		Hasher << FShaderCompilerOutputStructVersionLocal;
+
 		FShaderTarget Target = Input.Target;
 		Hasher << Target;
 		Hasher << Input.EntryPointName;
@@ -3154,6 +3158,9 @@ FShaderCommonCompileJob::FInputHash FShaderCompileJob::GetInputHash()
 		{
 			checkf(Archive.IsSaving() && !Archive.IsLoading(), TEXT("A loading archive is passed to FShaderCompileJob::GetInputHash(), this is not supported as it may corrupt its data"));
 
+			int32 FShaderCompilerOutputStructVersionLocal = FShaderCompilerOutputStructVersion;
+			Archive << FShaderCompilerOutputStructVersionLocal;
+
 			// Don't include debug group name in the hashing; this drastically worsens our cache hit rate
 			FString DebugGroupNameTmp(MoveTemp(Input.DebugGroupName));
 			Archive << Input;
diff --git a/Engine/Source/Runtime/RenderCore/Public/ShaderCompilerCore.h b/Engine/Source/Runtime/RenderCore/Public/ShaderCompilerCore.h
--- a/Engine/Source/Runtime/RenderCore/Public/ShaderCompilerCore.h
+++ b/Engine/Source/Runtime/RenderCore/Public/ShaderCompilerCore.h
@@ -35,6 +35,9 @@ inline const int32 ShaderCompileWorkerSingleJobHeader = 'S';
 // this is for the protocol, not the data.
 inline const int32 ShaderCompileWorkerPipelineJobHeader = 'P';
 
+// modify this for changes to the FShaderCompilerOutput data structure (in addition to ShaderCompileWorkerOutputVersion)
+inline const int32 FShaderCompilerOutputStructVersion = 1;
+
 namespace UE::ShaderCompiler
 {
 	RENDERCORE_API ERHIBindlessConfiguration GetBindlessResourcesConfiguration(FName ShaderFormat);
@@ -545,7 +548,7 @@ struct FShaderCompilerError
 
 /**
  *	The output of the shader compiler.
- *	Bump ShaderCompileWorkerOutputVersion if FShaderCompilerOutput changes
+ *	Bump FShaderCompilerOutputStructVersion and ShaderCompileWorkerOutputVersion if FShaderCompilerOutput changes
  */
 struct FShaderCompilerOutput
 {
