[shaders] fix remapping of already-parsed error messages; attempting to index past the end of the line number string

diff --git a/Engine/Source/Runtime/RenderCore/Private/ShaderPreprocessTypes.cpp b/Engine/Source/Runtime/RenderCore/Private/ShaderPreprocessTypes.cpp
--- a/Engine/Source/Runtime/RenderCore/Private/ShaderPreprocessTypes.cpp
+++ b/Engine/Source/Runtime/RenderCore/Private/ShaderPreprocessTypes.cpp
@@ -21,7 +21,7 @@ void FShaderDiagnosticRemapper::Remap(FShaderCompilerError& Diagnostic) const
 		const FString& ErrorLineStr = Diagnostic.ErrorLineString;
 		check(FChar::IsDigit(ErrorLineStr[0]));
 		int32 StrippedLineNum = 0, LineNumberEnd = 0;
-		while (FChar::IsDigit(ErrorLineStr[LineNumberEnd]))
+		while (LineNumberEnd < ErrorLineStr.Len() && FChar::IsDigit(ErrorLineStr[LineNumberEnd]))
 		{
 			StrippedLineNum = StrippedLineNum * 10 + ErrorLineStr[LineNumberEnd++] - TEXT('0');
 		}
