[shaders] fix sanity check of input hash value to correctly account for failed jobs (it won't be set in that case, but we don't want the check to fire as that will prevent the compile failure from being properly reported)

diff --git a/Engine/Source/Runtime/RenderCore/Private/ShaderCore.cpp b/Engine/Source/Runtime/RenderCore/Private/ShaderCore.cpp
--- a/Engine/Source/Runtime/RenderCore/Private/ShaderCore.cpp
+++ b/Engine/Source/Runtime/RenderCore/Private/ShaderCore.cpp
@@ -3315,7 +3315,7 @@ void FShaderCompileJob::SerializeOutput(FArchive& Ar)
 		Output.PreprocessTime = ActualPreprocessTime;
 	}
 
-	checkf(Output.InputHash == Input.Hash, TEXT("Failed sanity check: InputHash reported from compile output does not match InputHash from compile input"));
+	checkf(!Output.bSucceeded || Output.InputHash == Input.Hash, TEXT("Failed sanity check: InputHash reported from compile output does not match InputHash from compile input"));
 }
 
 void FShaderCompileJob::OnComplete()
