The majority of FRenderCommandPipeRegistry is not needed in server builds.

diff --git a/Engine/Source/Runtime/RenderCore/Private/RenderingThread.cpp b/Engine/Source/Runtime/RenderCore/Private/RenderingThread.cpp
--- a/Engine/Source/Runtime/RenderCore/Private/RenderingThread.cpp
+++ b/Engine/Source/Runtime/RenderCore/Private/RenderingThread.cpp
@@ -1595,18 +1595,21 @@ public:
 
 	void Initialize()
 	{
+#if !UE_SERVER
 		Pipes.Reset();
 
 		for (TLinkedList<FRenderCommandPipe*>::TIterator PipeIt(GlobalList); PipeIt; PipeIt.Next())
 		{
 			Pipes.Emplace(*PipeIt);
 		}
+#endif // !UE_SERVER
 
 		GRenderCommandPipeMode = GetValidatedRenderCommandPipeMode(CVarRenderCommandPipeMode->GetInt());
 	}
 
 	void StartRecording()
 	{
+#if !UE_SERVER
 		check(IsInGameThread());
 		SCOPED_NAMED_EVENT(FRenderCommandPipe_StartRecording, FColor::Magenta);
 
@@ -1653,10 +1656,12 @@ public:
 		{
 			TaskEvent.Trigger();
 		}
+#endif // !UE_SERVER
 	}
 
 	void StopRecording()
 	{
+#if !UE_SERVER
 		check(IsInGameThread());
 		SCOPED_NAMED_EVENT(FRenderCommandPipe_StopRecording, FColor::Magenta);
 
@@ -1700,26 +1705,37 @@ public:
 		});
 
 		bRecording = false;
+#endif // !UE_SERVER
 	}
 
 	bool IsRecording() const
 	{
+#if !UE_SERVER
 		ensureMsgf(!FTaskTagScope::IsCurrentTag(ETaskTag::EParallelRenderingThread) && !FTaskTagScope::IsCurrentTag(ETaskTag::ERenderingThread),
 			TEXT("IsRecording() is not valid from the render thread timeline."));
 
 		return bRecording;
+#else
+		return false;
+#endif // !UE_SERVER
 	}
 
 	bool IsReplaying() const
 	{
+#if !UE_SERVER
 		ensureMsgf(IsInParallelRenderingThread(), TEXT("IsReplaying() is only valid from the render thread timeline."));
 		return bReplaying;
+#else
+		return false;
+#endif // !UE_SERVER
 	}
 
 private:
+#if !UE_SERVER
 	TArray<FRenderCommandPipe*> Pipes;
 	bool bRecording = false;
 	bool bReplaying = false;
+#endif // !UE_SERVER
 };
 
 static FRenderCommandPipeRegistry GRenderCommandPipeRegistry;
