Fix an issue where intermediate shader dirs were not being cleand up at the start of a full cook.
Fixed an issue where chunked shader libraries were not loaded from a previous cook resulting in shaders missing.

diff --git a/Engine/Source/Runtime/RenderCore/Private/ShaderCodeLibrary.cpp b/Engine/Source/Runtime/RenderCore/Private/ShaderCodeLibrary.cpp
--- a/Engine/Source/Runtime/RenderCore/Private/ShaderCodeLibrary.cpp
+++ b/Engine/Source/Runtime/RenderCore/Private/ShaderCodeLibrary.cpp
@@ -68,6 +68,10 @@ static FString ShaderAssetInfoExtension = TEXT(".assetinfo.json");
 static FString StableExtension = TEXT(".shk");
 static FString PipelineExtension = TEXT(".ushaderpipelines");
 
+static FString ShaderCodePatternStr = TEXT("^ShaderCode\\-(\\w*)\\-([\\w\\-]*)\\.ushaderbytecode$");
+static FString ShaderArchivePatternStr = TEXT("^ShaderArchive\\-(\\w*)\\-([\\w\\-]*)\\.ushaderbytecode$");
+static FString StableInfoPatternStr = TEXT("^ShaderStableInfo\\-(\\w*)\\-([\\w\\-]*)\\.shk$");
+
 int32 GShaderCodeLibrarySeparateLoadingCache = 0;
 static FAutoConsoleVariableRef CVarShaderCodeLibrarySeparateLoadingCache(
 	TEXT("r.ShaderCodeLibrary.SeparateLoadingCache"),
@@ -1751,9 +1755,9 @@ struct FEditorShaderCodeArchive
 		return bOK;
 	}
 
-	void AddShaderCodeLibraryFromDirectory(const FString& BaseDir)
+	void AddShaderCodeLibraryByName(const FString& BaseDir, const FString& InLibraryName)
 	{
-		if (FArchive* PrevCookedAr = IFileManager::Get().CreateFileReader(*GetCodeArchiveFilename(BaseDir, LibraryName, FormatName)))
+		if (FArchive* PrevCookedAr = IFileManager::Get().CreateFileReader(*GetCodeArchiveFilename(BaseDir, InLibraryName, FormatName)))
 		{
 			uint32 Version = 0;
 			*PrevCookedAr << Version;
@@ -1765,19 +1769,19 @@ struct FEditorShaderCodeArchive
 				*PrevCookedAr << PrevCookedShaders;
 
 				// check if it also contains the asset info file
-				if (PrevCookedShaders.LoadAssetInfo(GetShaderAssetInfoFilename(BaseDir, LibraryName, FormatName)))
+				if (PrevCookedShaders.LoadAssetInfo(GetShaderAssetInfoFilename(BaseDir, InLibraryName, FormatName)))
 				{
 					UE_LOG(LogShaderLibrary, Display, TEXT("Loaded asset info %s for the shader library %s: %d entries"),
-						*GetShaderAssetInfoFilename(BaseDir, LibraryName, FormatName),
-						*GetCodeArchiveFilename(BaseDir, LibraryName, FormatName),
+						*GetShaderAssetInfoFilename(BaseDir, InLibraryName, FormatName),
+						*GetCodeArchiveFilename(BaseDir, InLibraryName, FormatName),
 						PrevCookedShaders.ShaderCodeToAssets.Num()
 					);
 				}
 				else
 				{
 					UE_LOG(LogShaderLibrary, Warning, TEXT("Could not find or load asset info %s for the shader library %s"),
-						*GetShaderAssetInfoFilename(BaseDir, LibraryName, FormatName),
-						*GetCodeArchiveFilename(BaseDir, LibraryName, FormatName)
+						*GetShaderAssetInfoFilename(BaseDir, InLibraryName, FormatName),
+						*GetCodeArchiveFilename(BaseDir, InLibraryName, FormatName)
 					);
 				}
 
@@ -1793,20 +1797,34 @@ struct FEditorShaderCodeArchive
 		}
 	}
 
+	void AddShaderCodeLibraryFromDirectory(const FString& BaseDir)
+	{
+		AddShaderCodeLibraryByName(BaseDir, LibraryName);
+	}
+
 	void AddExistingShaderCodeLibrary(FString const& OutputDir)
 	{
 		check(LibraryName.Len() > 0);
+		static const FRegexPattern ShaderCodePattern(ShaderCodePatternStr);
 
-		const FString ShaderIntermediateLocation = FPaths::ProjectSavedDir() / TEXT("Shaders") / FormatName.ToString();
+		const FString FormatNameStr = FormatName.ToString();
+		const FString ShaderIntermediateLocation = FPaths::ProjectSavedDir() / TEXT("Shaders") / FormatNameStr;
 
 		TArray<FString> ShaderFiles;
 		IFileManager::Get().FindFiles(ShaderFiles, *ShaderIntermediateLocation, *ShaderExtension);
 
 		for (const FString& ShaderFileName : ShaderFiles)
 		{
-			if (ShaderFileName.Contains(LibraryName + TEXT("-") + FormatName.ToString() + TEXT(".")))
+			FRegexMatcher FindShaderFormat(ShaderCodePattern, ShaderFileName);
+			if (FindShaderFormat.FindNext())
 			{
-				AddShaderCodeLibraryFromDirectory(OutputDir);
+				const FString& FoundLibraryName = FindShaderFormat.GetCaptureGroup(1);
+				const FString& FoundFormatName = FindShaderFormat.GetCaptureGroup(2);
+
+				if (FoundLibraryName.StartsWith(LibraryName) && FoundFormatName.Equals(FormatNameStr, ESearchCase::IgnoreCase))
+				{
+					AddShaderCodeLibraryByName(OutputDir, FoundLibraryName);
+				}
 			}
 		}
 	}
@@ -2182,7 +2200,7 @@ struct FEditorShaderStableInfo
 
 		const FString ShaderIntermediateLocation = FPaths::ProjectSavedDir() / TEXT("Shaders") / FormatName.ToString();
 		TArray<FString> ShaderFiles;
-		IFileManager::Get().FindFiles(ShaderFiles, *ShaderIntermediateLocation, *ShaderExtension);
+		IFileManager::Get().FindFiles(ShaderFiles, *ShaderIntermediateLocation, *StableExtension);
 
 		for (const FString& ShaderFileName : ShaderFiles)
 		{
@@ -2888,12 +2906,23 @@ public:
 	}
 
 #if WITH_EDITOR
+
+	FString GetFormatAndPlatformName(const FName& Format)
+	{
+		EShaderPlatform Platform = ShaderFormatToLegacyShaderPlatform(Format);
+		FName PossiblyAdjustedFormat = LegacyShaderPlatformToShaderFormat(Platform);	// Vulkan and GL switch between name variants depending on CVars 
+
+		return PossiblyAdjustedFormat.ToString() + TEXT("-") + FDataDrivenShaderPlatformInfo::GetName(Platform).ToString();
+	}
+
 	void CleanDirectories(TArray<FName> const& ShaderFormats)
 	{
 		for (FName const& Format : ShaderFormats)
 		{
 			FString ShaderIntermediateLocation = FPaths::ProjectSavedDir() / TEXT("Shaders") / Format.ToString();
+			FString ShaderPlatformIntermediateLocation = FPaths::ProjectSavedDir() / TEXT("Shaders") / GetFormatAndPlatformName(Format);
 			IFileManager::Get().DeleteDirectory(*ShaderIntermediateLocation, false, true);
+			IFileManager::Get().DeleteDirectory(*ShaderPlatformIntermediateLocation, false, true);
 		}
 	}
 
@@ -2906,9 +2935,7 @@ public:
 			FName const& Format = Descriptor.ShaderFormat;
 
 			EShaderPlatform Platform = ShaderFormatToLegacyShaderPlatform(Format);
-			FName PossiblyAdjustedFormat = LegacyShaderPlatformToShaderFormat(Platform);	// Vulkan and GL switch between name variants depending on CVars 
-
-			FString FormatAndPlatformName = PossiblyAdjustedFormat.ToString() + TEXT("-") + FDataDrivenShaderPlatformInfo::GetName(Platform).ToString();
+			FString FormatAndPlatformName = GetFormatAndPlatformName(Format);
 			FEditorShaderCodeArchive* CodeArchive = EditorShaderCodeArchive[Platform];
 			if (!CodeArchive)
 			{
@@ -3829,11 +3856,12 @@ bool FShaderLibraryCooker::MergeShaderCodeArchive(const TArray<FString>& CookedM
 	bool bNeedsDeterministicOrder = false;
 	GConfig->GetBool(TEXT("/Script/UnrealEd.ProjectPackagingSettings"), TEXT("bDeterministicShaderCodeOrder"), bNeedsDeterministicOrder, GGameIni);
 
+	static const FRegexPattern ShaderArchivePattern(ShaderArchivePatternStr);
+	static const FRegexPattern StableInfoPattern(StableInfoPatternStr);
+
 	// Map of file name to the archive we are unioning
 	TMap<FString, FEditorShaderCodeArchive> ShaderCodeArchives;
 	TMap<FString, FEditorShaderStableInfo> ShaderStableInfos;
-	static const FRegexPattern ShaderArchivePattern(TEXT("^ShaderArchive\\-(\\w*)\\-([\\w\\-]*)\\.ushaderbytecode$"));
-	static const FRegexPattern StableInfoPattern(TEXT("^ShaderStableInfo\\-(\\w*)\\-([\\w\\-]*)\\.shk$"));
 
 	for (const FString& MetadataDir : CookedMetadataDirs)
 	{
