[shaders] GetFinalSource should return an empty string if bExtractShaderSource is false; this prevents source from being unnecessarily added to/serialized with shader map results.

diff --git a/Engine/Source/Runtime/RenderCore/Private/ShaderCore.cpp b/Engine/Source/Runtime/RenderCore/Private/ShaderCore.cpp
--- a/Engine/Source/Runtime/RenderCore/Private/ShaderCore.cpp
+++ b/Engine/Source/Runtime/RenderCore/Private/ShaderCore.cpp
@@ -3410,9 +3410,19 @@ const FString& FShaderCompileJob::GetFinalSource() const
 	// done as part of the compile step will be written to the "ModifiedShaderSource" field
 	if (Input.bIndependentPreprocessed)
 	{
-		// if there are no such modifications, return the "unstripped" version of the source code (with comments & line directives maintained),
-		// otherwise return whatever the final modified source is as input to the compiler by the backend.
-		return Output.ModifiedShaderSource.IsEmpty() ? PreprocessOutput.GetUnstrippedSource() : Output.ModifiedShaderSource;
+		// always return empty string if source extraction was not requested; this will prevent bloat of material DDC data in the case where debug info is enabled 
+		// or Output.ModifiedShaderSource is unset (since the preprocess output unstripped source will always be set)
+		if (Input.ExtraSettings.bExtractShaderSource)
+		{
+			// if there are no such modifications, return the "unstripped" version of the source code (with comments & line directives maintained),
+			// otherwise return whatever the final modified source is as input to the compiler by the backend.
+			return Output.ModifiedShaderSource.IsEmpty() ? PreprocessOutput.GetUnstrippedSource() : Output.ModifiedShaderSource;
+		}
+		else
+		{
+			static FString Empty;
+			return Empty;
+		}
 	}
 	else
 	{
