ShaderLibraries in iterative cooking:
Suppress warnings about stable shaders cooked previously that were also cooked in the current cook but have a modified value. Instead, silently ignore the previously cooked value for the shaders.
#rnx

diff --git a/Engine/Source/Runtime/RenderCore/Private/ShaderCodeLibrary.cpp b/Engine/Source/Runtime/RenderCore/Private/ShaderCodeLibrary.cpp
--- a/Engine/Source/Runtime/RenderCore/Private/ShaderCodeLibrary.cpp
+++ b/Engine/Source/Runtime/RenderCore/Private/ShaderCodeLibrary.cpp
@@ -2168,28 +2168,52 @@ struct FEditorShaderStableInfo
 		LibraryName = TEXT("");
 	}
 
-	void AddShader(FStableShaderKeyAndValue& StableKeyValue)
+	enum class EMergeRule
+	{
+		/** If the key already exists, do not modify it, keep the existing value. */
+		KeepExisting,
+		/**
+		 * If the key already exists, compare whether the output hash is different. If different,
+		 * log a warning and keep the existing value. If the same, overwrite the existing value
+		 * with the new value. 
+		 */
+		OverwriteUnmodifiedWarnModified,
+	};
+	void AddShader(FStableShaderKeyAndValue& StableKeyValue, EMergeRule MergeRule)
 	{
 		FStableShaderKeyAndValue* Existing = StableMap.Find(StableKeyValue);
-		if (Existing && Existing->OutputHash != StableKeyValue.OutputHash)
+		if (Existing)
 		{
-			UE_LOG(LogShaderLibrary, Warning, TEXT("Duplicate key in stable shader library, but different keys, skipping new item:"));
-			UE_LOG(LogShaderLibrary, Warning, TEXT("    Existing: %s"), *Existing->ToString());
-			UE_LOG(LogShaderLibrary, Warning, TEXT("    New     : %s"), *StableKeyValue.ToString());
-			return;
+			switch (MergeRule)
+			{
+			case EMergeRule::KeepExisting:
+				return;
+			case EMergeRule::OverwriteUnmodifiedWarnModified:
+				if (Existing->OutputHash != StableKeyValue.OutputHash)
+				{
+					UE_LOG(LogShaderLibrary, Warning, TEXT("Duplicate key in stable shader library, but different output, skipping new item:"));
+					UE_LOG(LogShaderLibrary, Warning, TEXT("    Existing: %s"), *Existing->ToString());
+					UE_LOG(LogShaderLibrary, Warning, TEXT("    New     : %s"), *StableKeyValue.ToString());
+					return;
+				}
+				break; // Otherwise fall through to overwrite
+			default:
+				checkNoEntry();
+				break;
+			}
 		}
 		StableMap.Add(StableKeyValue);
 		MarkKeyValueDirty(StableKeyValue);
 	}
 
-	void AddShaderCodeLibraryFromDirectory(const FString& BaseDir)
+	void AddShaderCodeLibraryFromDirectory(const FString& BaseDir, EMergeRule MergeRule)
 	{
 		TArray<FStableShaderKeyAndValue> StableKeys;
 		if (UE::PipelineCacheUtilities::LoadStableKeysFile(GetStableInfoArchiveFilename(BaseDir, LibraryName, FormatName), StableKeys))
 		{
 			for (FStableShaderKeyAndValue& Item : StableKeys)
 			{
-				AddShader(Item);
+				AddShader(Item, MergeRule);
 			}
 		}
 	}
@@ -2198,17 +2222,26 @@ struct FEditorShaderStableInfo
 	{
 		check(LibraryName.Len() > 0);
 
-		const FString ShaderIntermediateLocation = FPaths::ProjectSavedDir() / TEXT("Shaders") / FormatName.ToString();
-		TArray<FString> ShaderFiles;
-		IFileManager::Get().FindFiles(ShaderFiles, *ShaderIntermediateLocation, *StableExtension);
-
-		for (const FString& ShaderFileName : ShaderFiles)
+		bool bLibraryExistsInSavedShadersDir = false;
 		{
-			if (ShaderFileName.Contains(LibraryName + TEXT("-") + FormatName.ToString() + TEXT(".")))
+			const FString ShaderIntermediateLocation = FPaths::ProjectSavedDir() / TEXT("Shaders") / FormatName.ToString();
+			TArray<FString> ShaderFiles;
+			IFileManager::Get().FindFiles(ShaderFiles, *ShaderIntermediateLocation, *StableExtension);
+			FString ExpectedFileNameText = LibraryName + TEXT("-") + FormatName.ToString() + TEXT(".");
+
+			for (const FString& ShaderFileName : ShaderFiles)
 			{
-				AddShaderCodeLibraryFromDirectory(OutputDir);
+				if (ShaderFileName.Contains(ExpectedFileNameText))
+				{
+					bLibraryExistsInSavedShadersDir = true;
+					break;
+				}
 			}
 		}
+		if (bLibraryExistsInSavedShadersDir)
+		{
+			AddShaderCodeLibraryFromDirectory(OutputDir, EMergeRule::KeepExisting);
+		}
 	}
 
 	void FinishPopulate(FString const& OutputDir)
@@ -2366,7 +2399,7 @@ struct FEditorShaderStableInfo
 			FStableShaderKeyAndValue StableInfo;
 			if (LoadFromCompactBinary(InfoView, StableInfo, Hashes))
 			{
-				AddShader(StableInfo);
+				AddShader(StableInfo, EMergeRule::OverwriteUnmodifiedWarnModified);
 			}
 			else
 			{
@@ -3131,7 +3164,7 @@ public:
 		FScopeLock ScopeLock(&ShaderCodeCS);
 
 		StableKeyValue.ComputeKeyHash();
-		StableArchive->AddShader(StableKeyValue);
+		StableArchive->AddShader(StableKeyValue, FEditorShaderStableInfo::EMergeRule::OverwriteUnmodifiedWarnModified);
 	}
  
 	void FinishPopulateShaderCode(const FString& ShaderCodeDir, const FString& MetaOutputDir, const TArray<FName>& ShaderFormats)
@@ -3897,7 +3930,8 @@ bool FShaderLibraryCooker::MergeShaderCodeArchive(const TArray<FString>& CookedM
 		{
 			if (ShaderStableInfos.Contains(StableInfoFile))
 			{
-				ShaderStableInfos[StableInfoFile].AddShaderCodeLibraryFromDirectory(ShaderStableInfoDir);
+				ShaderStableInfos[StableInfoFile].AddShaderCodeLibraryFromDirectory(ShaderStableInfoDir,
+					FEditorShaderStableInfo::EMergeRule::OverwriteUnmodifiedWarnModified);
 			}
 			else
 			{
@@ -3907,7 +3941,8 @@ bool FShaderLibraryCooker::MergeShaderCodeArchive(const TArray<FString>& CookedM
 					const FName ShaderFormat(*FindShaderInfo.GetCaptureGroup(2));
 					FEditorShaderStableInfo NewStableInfo(ShaderFormat);
 					NewStableInfo.OpenLibrary(FindShaderInfo.GetCaptureGroup(1));
-					NewStableInfo.AddShaderCodeLibraryFromDirectory(ShaderStableInfoDir);
+					NewStableInfo.AddShaderCodeLibraryFromDirectory(ShaderStableInfoDir,
+						FEditorShaderStableInfo::EMergeRule::OverwriteUnmodifiedWarnModified);
 
 					ShaderStableInfos.FindOrAdd(StableInfoFile, MoveTemp(NewStableInfo));
 				}
