Making FCopyTextureCS and FResolveDepth*PS shaders bindful on platforms that need that.

diff --git a/Engine/Source/Runtime/RenderCore/Private/ResolveShader.cpp b/Engine/Source/Runtime/RenderCore/Private/ResolveShader.cpp
--- a/Engine/Source/Runtime/RenderCore/Private/ResolveShader.cpp
+++ b/Engine/Source/Runtime/RenderCore/Private/ResolveShader.cpp
@@ -1,6 +1,7 @@
 // Copyright Epic Games, Inc. All Rights Reserved.
 
 #include "ResolveShader.h"
+#include "ShaderCompilerCore.h"
 #include "ShaderParameterUtils.h"
 #include "DataDrivenShaderPlatformInfo.h"
 #include "StereoRenderUtils.h"
@@ -29,6 +30,11 @@ FResolveDepthPS::FResolveDepthPS(const ShaderMetaType::CompiledShaderInitializer
 void FResolveDepthPS::ModifyCompilationEnvironment(const FGlobalShaderPermutationParameters& Parameters, FShaderCompilerEnvironment& OutEnvironment)
 {
 	FGlobalShader::ModifyCompilationEnvironment(Parameters, OutEnvironment);
+
+	if (FDataDrivenShaderPlatformInfo::GetRequiresBindfulUtilityShaders(Parameters.Platform))
+	{
+		OutEnvironment.CompilerFlags.Add(CFLAG_ForceBindful);
+	}
 }
 
 void FResolveDepthPS::SetParameters(FRHIBatchedShaderParameters& BatchedParameters, FParameter)
diff --git a/Engine/Source/Runtime/RenderCore/Public/CopyTextureShaders.h b/Engine/Source/Runtime/RenderCore/Public/CopyTextureShaders.h
--- a/Engine/Source/Runtime/RenderCore/Public/CopyTextureShaders.h
+++ b/Engine/Source/Runtime/RenderCore/Public/CopyTextureShaders.h
@@ -10,6 +10,7 @@
 #include "RHIDefinitions.h"
 #include "Serialization/MemoryLayout.h"
 #include "Shader.h"
+#include "ShaderCompilerCore.h"
 #include "ShaderCore.h"
 #include "ShaderParameterUtils.h"
 #include "ShaderParameters.h"
@@ -65,6 +66,16 @@ protected:
 	}
 
 public:
+	static inline void ModifyCompilationEnvironment(const FGlobalShaderPermutationParameters& Parameters, FShaderCompilerEnvironment& OutEnvironment)
+	{
+		FGlobalShader::ModifyCompilationEnvironment(Parameters, OutEnvironment);
+
+		if (FDataDrivenShaderPlatformInfo::GetRequiresBindfulUtilityShaders(Parameters.Platform))
+		{
+			OutEnvironment.CompilerFlags.Add(CFLAG_ForceBindful);
+		}
+	}
+
 	inline const FShaderResourceParameter& GetSrcResourceParam() { return SrcResourceParam; }
 	inline const FShaderResourceParameter& GetDstResourceParam() { return DstResourceParam; }
 
@@ -140,7 +151,7 @@ public:
 
 	static inline void ModifyCompilationEnvironment(const FGlobalShaderPermutationParameters& Parameters, FShaderCompilerEnvironment& OutEnvironment)
 	{
-		FGlobalShader::ModifyCompilationEnvironment(Parameters, OutEnvironment);
+		FCopyTextureCS::ModifyCompilationEnvironment(Parameters, OutEnvironment);
 
 		OutEnvironment.SetDefine(TEXT("THREADGROUPSIZE_X"), ThreadGroupSizeX);
 		OutEnvironment.SetDefine(TEXT("THREADGROUPSIZE_Y"), ThreadGroupSizeY);
