Refactored global dynamic index buffer to share thread-safe implementation with vertex buffer.

diff --git a/Engine/Source/Runtime/RenderCore/Private/GlobalRenderResources.cpp b/Engine/Source/Runtime/RenderCore/Private/GlobalRenderResources.cpp
--- a/Engine/Source/Runtime/RenderCore/Private/GlobalRenderResources.cpp
+++ b/Engine/Source/Runtime/RenderCore/Private/GlobalRenderResources.cpp
@@ -6,6 +6,7 @@
 #include "RenderCore.h"
 #include "RenderUtils.h"
 #include "Algo/Reverse.h"
+#include "Async/Mutex.h"
 
 // The maximum number of transient vertex buffer bytes to allocate before we start panic logging who is doing the allocations
 int32 GMaxVertexBytesAllocatedPerFrame = 32 * 1024 * 1024;
@@ -660,7 +661,8 @@ TGlobalResource<FTwoTrianglesIndexBuffer, FRenderResource::EInitPhase::Pre> GTwo
 /**
  * An individual dynamic vertex buffer.
  */
-class FDynamicVertexBuffer final : public FVertexBuffer
+template <typename BufferType>
+class TDynamicBuffer final : public BufferType
 {
 public:
 	/** The aligned size of all dynamic vertex buffers. */
@@ -671,14 +673,17 @@ public:
 	uint32 BufferSize;
 	/** Number of bytes currently allocated from the buffer. */
 	uint32 AllocatedByteCount;
+	/** Stride of the buffer in bytes. */
+	uint32 Stride;
 	/** Last render thread frame this resource was used in. */
 	uint64 LastUsedFrame = 0;
 
 	/** Default constructor. */
-	explicit FDynamicVertexBuffer(uint32 InMinBufferSize)
+	explicit TDynamicBuffer(uint32 InMinBufferSize, uint32 InStride)
 		: MappedBuffer(NULL)
 		, BufferSize(FMath::Max<uint32>(Align(InMinBufferSize, ALIGNMENT), ALIGNMENT))
 		, AllocatedByteCount(0)
+		, Stride(InStride)
 	{
 	}
 
@@ -689,8 +694,8 @@ public:
 	{
 		check(MappedBuffer == NULL);
 		check(AllocatedByteCount == 0);
-		check(IsValidRef(VertexBufferRHI));
-		MappedBuffer = (uint8*)RHICmdList.LockBuffer(VertexBufferRHI, 0, BufferSize, RLM_WriteOnly);
+		check(IsValidRef(BufferType::GetRHI()));
+		MappedBuffer = (uint8*)RHICmdList.LockBuffer(BufferType::GetRHI(), 0, BufferSize, RLM_WriteOnly);
 	}
 
 	/**
@@ -699,8 +704,8 @@ public:
 	void Unlock(FRHICommandListBase& RHICmdList)
 	{
 		check(MappedBuffer != NULL);
-		check(IsValidRef(VertexBufferRHI));
-		RHICmdList.UnlockBuffer(VertexBufferRHI);
+		check(IsValidRef(BufferType::GetRHI()));
+		RHICmdList.UnlockBuffer(BufferType::GetRHI());
 		MappedBuffer = NULL;
 		AllocatedByteCount = 0;
 	}
@@ -708,56 +713,52 @@ public:
 	// FRenderResource interface.
 	virtual void InitRHI(FRHICommandListBase& RHICmdList) override
 	{
-		check(!IsValidRef(VertexBufferRHI));
-		FRHIResourceCreateInfo CreateInfo(TEXT("FDynamicVertexBuffer"));
-		VertexBufferRHI = RHICmdList.CreateBuffer(BufferSize, EBufferUsageFlags::VertexBuffer, 0, ERHIAccess::VertexOrIndexBuffer, CreateInfo);
+		check(!IsValidRef(BufferType::GetRHI()));
+		FRHIResourceCreateInfo CreateInfo(TEXT("FDynamicBuffer"));
+		BufferType::SetRHI(RHICmdList.CreateBuffer(BufferSize, Stride == 0 ? EBufferUsageFlags::VertexBuffer : EBufferUsageFlags::IndexBuffer, Stride, ERHIAccess::VertexOrIndexBuffer, CreateInfo));
 		MappedBuffer = NULL;
 		AllocatedByteCount = 0;
 	}
 
 	virtual void ReleaseRHI() override
 	{
-		FVertexBuffer::ReleaseRHI();
+		BufferType::ReleaseRHI();
 		MappedBuffer = NULL;
 		AllocatedByteCount = 0;
 	}
-
-	virtual FString GetFriendlyName() const override
-	{
-		return TEXT("FDynamicVertexBuffer");
-	}
 };
 
 /**
- * A pool of dynamic vertex buffers.
+ * A pool of dynamic buffers.
  */
-struct FDynamicVertexBufferPool : public FRenderResource
-{
-	FCriticalSection CriticalSection;
-	TArray<FDynamicVertexBuffer*> LiveList;
-	TArray<FDynamicVertexBuffer*> FreeList;
-	TArray<FDynamicVertexBuffer*> LockList;
-	TArray<FDynamicVertexBuffer*> ReclaimList;
+template <typename DynamicBufferType>
+struct TDynamicBufferPool : public FRenderResource
+{
+	UE::FMutex Mutex;
+	TArray<DynamicBufferType*> LiveList;
+	TArray<DynamicBufferType*> FreeList;
+	TArray<DynamicBufferType*> LockList;
+	TArray<DynamicBufferType*> ReclaimList;
 	std::atomic_uint32_t TotalAllocatedMemory{0};
 	uint32 CurrentCycle = 0;
 
-	FDynamicVertexBuffer* Acquire(FRHICommandListBase& RHICmdList, uint32 SizeInBytes)
+	DynamicBufferType* Acquire(FRHICommandListBase& RHICmdList, uint32 SizeInBytes, uint32 Stride)
 	{
 		const uint32 MinimumBufferSize = 65536u;
 		SizeInBytes = FMath::Max(SizeInBytes, MinimumBufferSize);
 
-		FDynamicVertexBuffer* FoundBuffer = nullptr;
+		DynamicBufferType* FoundBuffer = nullptr;
 		bool bInitializeBuffer = false;
 
 		{
-			FScopeLock Lock(&CriticalSection);
+			UE::TScopeLock Lock(Mutex);
 
 			// Traverse the free list like a stack, starting from the top, so recently reclaimed items are allocated first.
 			for (int32 Index = FreeList.Num() - 1; Index >= 0; --Index)
 			{
-				FDynamicVertexBuffer* Buffer = FreeList[Index];
+				DynamicBufferType* Buffer = FreeList[Index];
 
-				if (SizeInBytes <= Buffer->BufferSize)
+				if (SizeInBytes <= Buffer->BufferSize && Buffer->Stride == Stride)
 				{
 					FreeList.RemoveAt(Index, 1, false);
 					FoundBuffer = Buffer;
@@ -767,7 +768,7 @@ struct FDynamicVertexBufferPool : public FRenderResource
 
 			if (!FoundBuffer)
 			{
-				FoundBuffer = new FDynamicVertexBuffer(SizeInBytes);
+				FoundBuffer = new DynamicBufferType(SizeInBytes, Stride);
 				LiveList.Emplace(FoundBuffer);
 				bInitializeBuffer = true;
 			}
@@ -792,26 +793,29 @@ struct FDynamicVertexBufferPool : public FRenderResource
 		return FoundBuffer;
 	}
 
-	void Forfeit(FRHICommandListBase& RHICmdList, TConstArrayView<FDynamicVertexBuffer*> BuffersToForfeit)
+	void Forfeit(FRHICommandListBase& RHICmdList, TConstArrayView<DynamicBufferType*> BuffersToForfeit)
 	{
-		for (FDynamicVertexBuffer* Buffer : BuffersToForfeit)
+		if (!BuffersToForfeit.IsEmpty())
 		{
-			Buffer->Unlock(RHICmdList);
-		}
+			for (DynamicBufferType* Buffer : BuffersToForfeit)
+			{
+				Buffer->Unlock(RHICmdList);
+			}
 
-		FScopeLock Lock(&CriticalSection);
-		ReclaimList.Append(BuffersToForfeit);
+			UE::TScopeLock Lock(Mutex);
+			ReclaimList.Append(BuffersToForfeit);
+		}
 	}
 
 	void GarbageCollect()
 	{
-		FScopeLock Lock(&CriticalSection);
+		UE::TScopeLock Lock(Mutex);
 		FreeList.Append(ReclaimList);
 		ReclaimList.Reset();
 
 		for (int32 Index = 0; Index < LiveList.Num(); ++Index)
 		{
-			FDynamicVertexBuffer* Buffer = LiveList[Index];
+			DynamicBufferType* Buffer = LiveList[Index];
 
 			if (GGlobalBufferNumFramesUnusedThreshold > 0 && Buffer->LastUsedFrame + GGlobalBufferNumFramesUnusedThreshold <= GFrameCounterRenderThread)
 			{
@@ -835,22 +839,22 @@ private:
 		check(LockList.IsEmpty());
 		check(FreeList.Num() == LiveList.Num());
 
-		for (FDynamicVertexBuffer* VertexBuffer : LiveList)
+		for (DynamicBufferType* Buffer : LiveList)
 		{
-			TotalAllocatedMemory -= VertexBuffer->BufferSize;
-			VertexBuffer->ReleaseResource();
-			delete VertexBuffer;
+			TotalAllocatedMemory -= Buffer->BufferSize;
+			Buffer->ReleaseResource();
+			delete Buffer;
 		}
 		LiveList.Empty();
 		FreeList.Empty();
 	}
 };
 
-static TGlobalResource<FDynamicVertexBufferPool, FRenderResource::EInitPhase::Pre> GDynamicVertexBufferPool;
+static TGlobalResource<TDynamicBufferPool<FDynamicVertexBuffer>, FRenderResource::EInitPhase::Pre> GDynamicVertexBufferPool;
 
 FGlobalDynamicVertexBuffer::FAllocation FGlobalDynamicVertexBuffer::Allocate(uint32 SizeInBytes)
 {
-	return Allocate(FRHICommandListExecutor::GetImmediateCommandList(), SizeInBytes);
+	return Allocate(FRHICommandListImmediate::Get(), SizeInBytes);
 }
 
 FGlobalDynamicVertexBuffer::FAllocation FGlobalDynamicVertexBuffer::Allocate(FRHICommandListBase& RHICmdList, uint32 SizeInBytes)
@@ -859,7 +863,7 @@ FGlobalDynamicVertexBuffer::FAllocation FGlobalDynamicVertexBuffer::Allocate(FRH
 
 	if (VertexBuffers.IsEmpty() || VertexBuffers.Last()->AllocatedByteCount + SizeInBytes > VertexBuffers.Last()->BufferSize)
 	{
-		VertexBuffers.Emplace(GDynamicVertexBufferPool.Acquire(RHICmdList, SizeInBytes));
+		VertexBuffers.Emplace(GDynamicVertexBufferPool.Acquire(RHICmdList, SizeInBytes, 0));
 	}
 
 	FDynamicVertexBuffer* VertexBuffer = VertexBuffers.Last();
@@ -885,149 +889,17 @@ void FGlobalDynamicVertexBuffer::Commit(FRHICommandListBase& RHICmdList)
 
 void FGlobalDynamicVertexBuffer::Commit()
 {
-	Commit(FRHICommandListExecutor::GetImmediateCommandList());
+	Commit(FRHICommandListImmediate::Get());
 }
 
-void FGlobalDynamicVertexBuffer::GarbageCollect()
-{
-	GDynamicVertexBufferPool.GarbageCollect();
-}
-
-/*------------------------------------------------------------------------------
-	FGlobalDynamicIndexBuffer implementation.
-------------------------------------------------------------------------------*/
-
-/**
- * An individual dynamic index buffer.
- */
-class FDynamicIndexBuffer : public FIndexBuffer
-{
-public:
-	/** The aligned size of all dynamic index buffers. */
-	enum { ALIGNMENT = (1 << 16) }; // 64KB
-	/** Pointer to the index buffer mapped in main memory. */
-	uint8* MappedBuffer;
-	/** Size of the index buffer in bytes. */
-	uint32 BufferSize;
-	/** Number of bytes currently allocated from the buffer. */
-	uint32 AllocatedByteCount;
-	/** Stride of the buffer in bytes. */
-	uint32 Stride;
-	/** Number of successive frames for which AllocatedByteCount == 0. Used as a metric to decide when to free the allocation. */
-	int32 NumFramesUnused = 0;
-
-	/** Initialization constructor. */
-	explicit FDynamicIndexBuffer(uint32 InMinBufferSize, uint32 InStride)
-		: MappedBuffer(NULL)
-		, BufferSize(FMath::Max<uint32>(Align(InMinBufferSize, ALIGNMENT), ALIGNMENT))
-		, AllocatedByteCount(0)
-		, Stride(InStride)
-	{
-	}
-
-	/**
-	 * Locks the vertex buffer so it may be written to.
-	 */
-	void Lock(FRHICommandListBase& RHICmdList)
-	{
-		check(MappedBuffer == NULL);
-		check(AllocatedByteCount == 0);
-		check(IsValidRef(IndexBufferRHI));
-		MappedBuffer = (uint8*)RHICmdList.LockBuffer(IndexBufferRHI, 0, BufferSize, RLM_WriteOnly);
-	}
-
-	void Lock()
-	{
-		Lock(FRHICommandListImmediate::Get());
-	}
-
-	/**
-	 * Unocks the buffer so the GPU may read from it.
-	 */
-	void Unlock(FRHICommandListBase& RHICmdList)
-	{
-		check(MappedBuffer != NULL);
-		check(IsValidRef(IndexBufferRHI));
-		RHICmdList.UnlockBuffer(IndexBufferRHI);
-		MappedBuffer = NULL;
-		AllocatedByteCount = 0;
-		NumFramesUnused = 0;
-	}
-
-	void Unlock()
-	{
-		Unlock(FRHICommandListImmediate::Get());
-	}
-
-	// FRenderResource interface.
-	virtual void InitRHI(FRHICommandListBase& RHICmdList) override
-	{
-		check(!IsValidRef(IndexBufferRHI));
-		FRHIResourceCreateInfo CreateInfo(TEXT("FDynamicIndexBuffer"));
-		IndexBufferRHI = RHICmdList.CreateIndexBuffer(Stride, BufferSize, BUF_Volatile, CreateInfo);
-		MappedBuffer = NULL;
-		AllocatedByteCount = 0;
-	}
-
-	virtual void ReleaseRHI() override
-	{
-		FIndexBuffer::ReleaseRHI();
-		MappedBuffer = NULL;
-		AllocatedByteCount = 0;
-	}
-
-	virtual FString GetFriendlyName() const override
-	{
-		return TEXT("FDynamicIndexBuffer");
-	}
-};
-
-/**
- * A pool of dynamic index buffers.
- */
-struct FDynamicIndexBufferPool
+FGlobalDynamicIndexBuffer::FAllocation FGlobalDynamicIndexBuffer::Allocate(uint32 SizeInBytes, uint32 IndexStride)
 {
-	/** List of index buffers. */
-	TIndirectArray<FDynamicIndexBuffer> IndexBuffers;
-	/** The current buffer from which allocations are being made. */
-	FDynamicIndexBuffer* CurrentIndexBuffer;
-	/** Stride of buffers in this pool. */
-	uint32 BufferStride;
-
-	/** Initialization constructor. */
-	explicit FDynamicIndexBufferPool(uint32 InBufferStride)
-		: CurrentIndexBuffer(NULL)
-		, BufferStride(InBufferStride)
-	{
-	}
-
-	/** Destructor. */
-	~FDynamicIndexBufferPool()
-	{
-		int32 NumIndexBuffers = IndexBuffers.Num();
-		for (int32 BufferIndex = 0; BufferIndex < NumIndexBuffers; ++BufferIndex)
-		{
-			IndexBuffers[BufferIndex].ReleaseResource();
-		}
-	}
-};
-
-FGlobalDynamicIndexBuffer::FGlobalDynamicIndexBuffer()
-{
-	Pools[0] = new FDynamicIndexBufferPool(sizeof(uint16));
-	Pools[1] = new FDynamicIndexBufferPool(sizeof(uint32));
+	return Allocate(FRHICommandListImmediate::Get(), SizeInBytes, IndexStride);
 }
 
-FGlobalDynamicIndexBuffer::~FGlobalDynamicIndexBuffer()
-{
-	for (int32 i = 0; i < 2; ++i)
-	{
-		delete Pools[i];
-		Pools[i] = NULL;
-	}
-}
+static TGlobalResource<TDynamicBufferPool<FDynamicIndexBuffer>, FRenderResource::EInitPhase::Pre> GDynamicIndexBufferPool;
 
-FGlobalDynamicIndexBuffer::FAllocation FGlobalDynamicIndexBuffer::Allocate(uint32 NumIndices, uint32 IndexStride)
+FGlobalDynamicIndexBuffer::FAllocation FGlobalDynamicIndexBuffer::Allocate(FRHICommandListBase& RHICmdList, uint32 SizeInBytes, uint32 IndexStride)
 {
 	FAllocation Allocation;
 
@@ -1036,78 +908,43 @@ FGlobalDynamicIndexBuffer::FAllocation FGlobalDynamicIndexBuffer::Allocate(uint3
 		return Allocation;
 	}
 
-	FDynamicIndexBufferPool* Pool = Pools[IndexStride >> 2]; // 2 -> 0, 4 -> 1
+	TArray<FDynamicIndexBuffer*>& IndexBuffers = (IndexStride == 2)
+		? IndexBuffers16
+		: IndexBuffers32;
 
-	uint32 SizeInBytes = NumIndices * IndexStride;
-	FDynamicIndexBuffer* IndexBuffer = Pool->CurrentIndexBuffer;
-	if (IndexBuffer == NULL || IndexBuffer->AllocatedByteCount + SizeInBytes > IndexBuffer->BufferSize)
+	if (IndexBuffers.IsEmpty() || IndexBuffers.Last()->AllocatedByteCount + SizeInBytes > IndexBuffers.Last()->BufferSize)
 	{
-		// Find a buffer in the pool big enough to service the request.
-		IndexBuffer = NULL;
-		for (int32 BufferIndex = 0, NumBuffers = Pool->IndexBuffers.Num(); BufferIndex < NumBuffers; ++BufferIndex)
-		{
-			FDynamicIndexBuffer& IndexBufferToCheck = Pool->IndexBuffers[BufferIndex];
-			if (IndexBufferToCheck.AllocatedByteCount + SizeInBytes <= IndexBufferToCheck.BufferSize)
-			{
-				IndexBuffer = &IndexBufferToCheck;
-				break;
-			}
-		}
-
-		// Create a new index buffer if needed.
-		if (IndexBuffer == NULL)
-		{
-			IndexBuffer = new FDynamicIndexBuffer(SizeInBytes, Pool->BufferStride);
-			Pool->IndexBuffers.Add(IndexBuffer);
-			IndexBuffer->InitResource(FRHICommandListImmediate::Get());
-		}
-
-		// Lock the buffer if needed.
-		if (IndexBuffer->MappedBuffer == NULL)
-		{
-			IndexBuffer->Lock();
-		}
-		Pool->CurrentIndexBuffer = IndexBuffer;
+		IndexBuffers.Emplace(GDynamicIndexBufferPool.Acquire(RHICmdList, SizeInBytes, IndexStride));
 	}
 
-	check(IndexBuffer != NULL);
-	checkf(IndexBuffer->AllocatedByteCount + SizeInBytes <= IndexBuffer->BufferSize, TEXT("Global index buffer allocation failed: BufferSize=%d AllocatedByteCount=%d SizeInBytes=%d"), IndexBuffer->BufferSize, IndexBuffer->AllocatedByteCount, SizeInBytes);
+	FDynamicIndexBuffer* IndexBuffer = IndexBuffers.Last();
 
+	checkf(IndexBuffer->AllocatedByteCount + SizeInBytes <= IndexBuffer->BufferSize, TEXT("Global index buffer allocation failed: BufferSize=%d BufferStride=%d AllocatedByteCount=%d SizeInBytes=%d"), IndexBuffer->BufferSize, IndexBuffer->Stride, IndexBuffer->AllocatedByteCount, SizeInBytes);
 	Allocation.Buffer = IndexBuffer->MappedBuffer + IndexBuffer->AllocatedByteCount;
 	Allocation.IndexBuffer = IndexBuffer;
 	Allocation.FirstIndex = IndexBuffer->AllocatedByteCount / IndexStride;
 	IndexBuffer->AllocatedByteCount += SizeInBytes;
-
 	return Allocation;
 }
 
-void FGlobalDynamicIndexBuffer::Commit()
+void FGlobalDynamicIndexBuffer::Commit(FRHICommandListBase& RHICmdList)
 {
-	for (int32 i = 0; i < 2; ++i)
-	{
-		FDynamicIndexBufferPool* Pool = Pools[i];
+	GDynamicIndexBufferPool.Forfeit(RHICmdList, IndexBuffers16);
+	GDynamicIndexBufferPool.Forfeit(RHICmdList, IndexBuffers32);
+	IndexBuffers16.Reset();
+	IndexBuffers32.Reset();
+}
 
-		for (int32 BufferIndex = 0, NumBuffers = Pool->IndexBuffers.Num(); BufferIndex < NumBuffers; ++BufferIndex)
-		{
-			FDynamicIndexBuffer& IndexBuffer = Pool->IndexBuffers[BufferIndex];
-			if (IndexBuffer.MappedBuffer != NULL)
-			{
-				IndexBuffer.Unlock();
-			}
-			else if (GGlobalBufferNumFramesUnusedThreshold && !IndexBuffer.AllocatedByteCount)
-			{
-				++IndexBuffer.NumFramesUnused;
-				if (IndexBuffer.NumFramesUnused >= GGlobalBufferNumFramesUnusedThreshold)
-				{
-					// Remove the buffer, assumes they are unordered.
-					IndexBuffer.ReleaseResource();
-					Pool->IndexBuffers.RemoveAtSwap(BufferIndex);
-					--BufferIndex;
-					--NumBuffers;
-				}
-			}
-		}
-		Pool->CurrentIndexBuffer = NULL;
-	}
+void FGlobalDynamicIndexBuffer::Commit()
+{
+	Commit(FRHICommandListImmediate::Get());
 }
 
+namespace GlobalDynamicBuffer
+{
+	void GarbageCollect()
+	{
+		GDynamicVertexBufferPool.GarbageCollect();
+		GDynamicIndexBufferPool.GarbageCollect();
+	}
+}
\ No newline at end of file
diff --git a/Engine/Source/Runtime/RenderCore/Public/GlobalRenderResources.h b/Engine/Source/Runtime/RenderCore/Public/GlobalRenderResources.h
--- a/Engine/Source/Runtime/RenderCore/Public/GlobalRenderResources.h
+++ b/Engine/Source/Runtime/RenderCore/Public/GlobalRenderResources.h
@@ -156,6 +156,12 @@ extern RENDERCORE_API TGlobalResource<FTwoTrianglesIndexBuffer, FRenderResource:
 //////////////////////////////////////////////////////////////////////////////////////////////////////////////////////
 // FGlobalDynamicVertexBuffer
 
+template <typename BufferType>
+class TDynamicBuffer;
+
+using FDynamicVertexBuffer = TDynamicBuffer<FVertexBuffer>;
+using FDynamicIndexBuffer  = TDynamicBuffer<FIndexBuffer>;
+
 struct FGlobalDynamicVertexBufferAllocation
 {
 	/** The location of the buffer in main memory. */
@@ -196,15 +202,18 @@ public:
 	 *		remain valid only until the next call to Allocate!
 	 */
 	RENDERCORE_API void Commit(FRHICommandListBase& RHICmdList);
+
+	UE_DEPRECATED(5.4, "Commit requires a command list.")
 	RENDERCORE_API void Commit();
 
-	static RENDERCORE_API void GarbageCollect();
+	UE_DEPRECATED(5.4, "Use GlobalDynamicBuffer::GarbageCollect instead.")
+	void GarbageCollect() {}
 
 	/** Returns true if log statements should be made because we exceeded GMaxVertexBytesAllocatedPerFrame */
 	RENDERCORE_API bool IsRenderAlarmLoggingEnabled() const;
 
 private:
-	TArray<class FDynamicVertexBuffer*> VertexBuffers;
+	TArray<FDynamicVertexBuffer*> VertexBuffers;
 };
 
 struct FGlobalDynamicIndexBufferAllocation
@@ -250,18 +259,13 @@ public:
 	using FAllocation = FGlobalDynamicIndexBufferAllocation;
 	using FAllocationEx = FGlobalDynamicIndexBufferAllocationEx;
 
-	/** Default constructor. */
-	RENDERCORE_API FGlobalDynamicIndexBuffer();
-
-	/** Destructor. */
-	RENDERCORE_API ~FGlobalDynamicIndexBuffer();
-
 	/**
 	 * Allocates space in the global index buffer.
 	 * @param NumIndices - The number of indices to allocate.
 	 * @param IndexStride - The size of an index (2 or 4 bytes).
 	 * @returns An FAllocation with information regarding the allocated memory.
 	 */
+	RENDERCORE_API FAllocation Allocate(FRHICommandListBase& RHICmdList, uint32 NumIndices, uint32 IndexStride);
 	RENDERCORE_API FAllocation Allocate(uint32 NumIndices, uint32 IndexStride);
 
 	/**
@@ -275,14 +279,28 @@ public:
 		return FAllocationEx(Allocate(NumIndices, sizeof(IndexType)), NumIndices, sizeof(IndexType));
 	}
 
+	template <typename IndexType>
+	FORCEINLINE FAllocationEx Allocate(FRHICommandListBase& RHICmdList, uint32 NumIndices)
+	{
+		return FAllocationEx(Allocate(RHICmdList, NumIndices, sizeof(IndexType)), NumIndices, sizeof(IndexType));
+	}
+
 	/**
 	 * Commits allocated memory to the GPU.
 	 *		WARNING: Once this buffer has been committed to the GPU, allocations
 	 *		remain valid only until the next call to Allocate!
 	 */
+	RENDERCORE_API void Commit(FRHICommandListBase& RHICmdList);
+
+	UE_DEPRECATED(5.4, "Commit requires a command list.")
 	RENDERCORE_API void Commit();
 
 private:
-	/** The pool of vertex buffers from which allocations are made. */
-	struct FDynamicIndexBufferPool* Pools[2];
+	TArray<FDynamicIndexBuffer*> IndexBuffers16;
+	TArray<FDynamicIndexBuffer*> IndexBuffers32;
 };
+
+namespace GlobalDynamicBuffer
+{
+	RENDERCORE_API void GarbageCollect();
+}
\ No newline at end of file
diff --git a/Engine/Source/Runtime/RenderCore/Public/RenderResource.h b/Engine/Source/Runtime/RenderCore/Public/RenderResource.h
--- a/Engine/Source/Runtime/RenderCore/Public/RenderResource.h
+++ b/Engine/Source/Runtime/RenderCore/Public/RenderResource.h
@@ -474,6 +474,10 @@ public:
 	RENDERCORE_API virtual void ReleaseRHI() override;
 	RENDERCORE_API virtual FString GetFriendlyName() const override;
 
+	const FBufferRHIRef& GetRHI() const { return VertexBufferRHI; }
+
+	void SetRHI(const FBufferRHIRef& BufferRHI) { VertexBufferRHI = BufferRHI; }
+
 	FBufferRHIRef VertexBufferRHI;
 };
 
@@ -503,6 +507,10 @@ public:
 	RENDERCORE_API virtual void ReleaseRHI() override;
 	RENDERCORE_API virtual FString GetFriendlyName() const override;
 
+	const FBufferRHIRef& GetRHI() const { return IndexBufferRHI; }
+
+	void SetRHI(const FBufferRHIRef& BufferRHI) { IndexBufferRHI = BufferRHI; }
+
 	FBufferRHIRef IndexBufferRHI;
 };
 
diff --git a/Engine/Source/Runtime/Renderer/Private/DeferredShadingRenderer.cpp b/Engine/Source/Runtime/Renderer/Private/DeferredShadingRenderer.cpp
--- a/Engine/Source/Runtime/Renderer/Private/DeferredShadingRenderer.cpp
+++ b/Engine/Source/Runtime/Renderer/Private/DeferredShadingRenderer.cpp
@@ -2390,8 +2390,8 @@ void FDeferredShadingSceneRenderer::FinishInitDynamicShadows(FRDGBuilder& GraphB
 		}
 
 		SCOPE_CYCLE_COUNTER(STAT_FDeferredShadingSceneRenderer_FGlobalDynamicVertexBuffer_Commit);
-		DynamicVertexBufferForInitShadows.Commit();
-		DynamicIndexBufferForInitShadows.Commit();
+		DynamicVertexBufferForInitShadows.Commit(GraphBuilder.RHICmdList);
+		DynamicIndexBufferForInitShadows.Commit(GraphBuilder.RHICmdList);
 		DynamicReadBufferForInitShadows.Commit(GraphBuilder.RHICmdList);
 	}
 }
@@ -3000,8 +3000,8 @@ void FDeferredShadingSceneRenderer::Render(FRDGBuilder& GraphBuilder)
 		// Dynamic vertex and index buffers need to be committed before rendering.
 		{
 			SCOPE_CYCLE_COUNTER(STAT_FDeferredShadingSceneRenderer_FGlobalDynamicVertexBuffer_Commit);
-			DynamicIndexBufferForInitViews.Commit();
-			DynamicVertexBufferForInitViews.Commit();
+			DynamicIndexBufferForInitViews.Commit(GraphBuilder.RHICmdList);
+			DynamicVertexBufferForInitViews.Commit(GraphBuilder.RHICmdList);
 			DynamicReadBufferForInitViews.Commit(GraphBuilder.RHICmdList);
 		}
 	}
diff --git a/Engine/Source/Runtime/Renderer/Private/MobileShadingRenderer.cpp b/Engine/Source/Runtime/Renderer/Private/MobileShadingRenderer.cpp
--- a/Engine/Source/Runtime/Renderer/Private/MobileShadingRenderer.cpp
+++ b/Engine/Source/Runtime/Renderer/Private/MobileShadingRenderer.cpp
@@ -917,8 +917,8 @@ void FMobileSceneRenderer::Render(FRDGBuilder& GraphBuilder)
 	GEngine->GetPreRenderDelegateEx().Broadcast(GraphBuilder);
 
 	// Global dynamic buffers need to be committed before rendering.
-	DynamicIndexBuffer.Commit();
-	DynamicVertexBuffer.Commit();
+	DynamicIndexBuffer.Commit(GraphBuilder.RHICmdList);
+	DynamicVertexBuffer.Commit(GraphBuilder.RHICmdList);
 	DynamicReadBuffer.Commit(GraphBuilder.RHICmdList);
 	
 	GraphBuilder.SetCommandListStat(GET_STATID(STAT_CLMM_SceneSim));
diff --git a/Engine/Source/Runtime/Renderer/Private/SceneHitProxyRendering.cpp b/Engine/Source/Runtime/Renderer/Private/SceneHitProxyRendering.cpp
--- a/Engine/Source/Runtime/Renderer/Private/SceneHitProxyRendering.cpp
+++ b/Engine/Source/Runtime/Renderer/Private/SceneHitProxyRendering.cpp
@@ -578,8 +578,8 @@ void FMobileSceneRenderer::RenderHitProxies(FRDGBuilder& GraphBuilder)
 	GEngine->GetPreRenderDelegateEx().Broadcast(GraphBuilder);
 
 	// Global dynamic buffers need to be committed before rendering.
-	DynamicIndexBuffer.Commit();
-	DynamicVertexBuffer.Commit();
+	DynamicIndexBuffer.Commit(GraphBuilder.RHICmdList);
+	DynamicVertexBuffer.Commit(GraphBuilder.RHICmdList);
 	DynamicReadBuffer.Commit(GraphBuilder.RHICmdList);
 
 	InstanceCullingManager.FlushRegisteredViews(GraphBuilder);
@@ -661,8 +661,8 @@ void FDeferredShadingSceneRenderer::RenderHitProxies(FRDGBuilder& GraphBuilder)
 	GEngine->GetPreRenderDelegateEx().Broadcast(GraphBuilder);
 
 	// Global dynamic buffers need to be committed before rendering.
-	DynamicIndexBufferForInitViews.Commit();
-	DynamicVertexBufferForInitViews.Commit();
+	DynamicIndexBufferForInitViews.Commit(GraphBuilder.RHICmdList);
+	DynamicVertexBufferForInitViews.Commit(GraphBuilder.RHICmdList);
 	DynamicReadBufferForInitViews.Commit(GraphBuilder.RHICmdList);
 
 	// Notify the FX system that the scene is about to be rendered.
diff --git a/Engine/Source/Runtime/Renderer/Private/SceneRendering.cpp b/Engine/Source/Runtime/Renderer/Private/SceneRendering.cpp
--- a/Engine/Source/Runtime/Renderer/Private/SceneRendering.cpp
+++ b/Engine/Source/Runtime/Renderer/Private/SceneRendering.cpp
@@ -4221,7 +4221,7 @@ static void FinishCleanUp(FRHICommandListImmediate& RHICmdList)
 	TRACE_CPUPROFILER_EVENT_SCOPE(FinishCleanUp);
 
 	// Can release only after all mesh pass tasks are finished.
-	FGlobalDynamicVertexBuffer::GarbageCollect();
+	GlobalDynamicBuffer::GarbageCollect();
 	GPrimitiveIdVertexBufferPool.DiscardAll();
 	FGraphicsMinimalPipelineStateId::ResetLocalPipelineIdTableSize();
 }
