Modified render target and buffer pools to be thread safe.

diff --git a/Engine/Source/Runtime/RenderCore/Private/RenderGraphBuilder.cpp b/Engine/Source/Runtime/RenderCore/Private/RenderGraphBuilder.cpp
--- a/Engine/Source/Runtime/RenderCore/Private/RenderGraphBuilder.cpp
+++ b/Engine/Source/Runtime/RenderCore/Private/RenderGraphBuilder.cpp
@@ -1765,6 +1765,10 @@ void FRDGBuilder::Execute()
 			CSV_SCOPED_TIMING_STAT_EXCLUSIVE(RDG_CollectResources);
 			SCOPED_NAMED_EVENT_TEXT("FRDGBuilder::CollectResources", FColor::Magenta);
 
+			// Take locks on the buffer / texture pools as nothing else should be allocating from them right now and granular locking is expensive.
+			UE::TScopeLock RenderTargetPoolLock(GRenderTargetPool.Mutex);
+			UE::TScopeLock BufferPoolLock(GRenderGraphResourcePool.Mutex);
+
 			UniformBuffersToCreate.Reserve(UniformBuffers.Num());
 
 			EnumerateExtendedLifetimeResources(Textures, [](FRDGTexture* Texture)
@@ -3606,7 +3610,7 @@ void FRDGBuilder::BeginResourceRHI(FRDGPassHandle PassHandle, FRDGTextureRef Tex
 
 	if (!Texture->ResourceRHI)
 	{
-		SetRHI(Texture, GRenderTargetPool.FindFreeElement(Texture->Desc, Texture->Name), PassHandle);
+		SetRHI(Texture, GRenderTargetPool.FindFreeElement(RHICmdList, Texture->Desc, Texture->Name), PassHandle);
 	}
 }
 
@@ -3692,7 +3696,7 @@ void FRDGBuilder::BeginResourceRHI(FRDGPassHandle PassHandle, FRDGBufferRef Buff
 	{
 		const ERDGPooledBufferAlignment Alignment = Buffer->bQueuedForUpload ? ERDGPooledBufferAlignment::PowerOfTwo : ERDGPooledBufferAlignment::Page;
 
-		SetRHI(Buffer, GRenderGraphResourcePool.FindFreeBuffer(Buffer->Desc, Buffer->Name, Alignment), PassHandle);
+		SetRHI(Buffer, GRenderGraphResourcePool.FindFreeBuffer(RHICmdList, Buffer->Desc, Buffer->Name, Alignment), PassHandle);
 	}
 
 #if RHI_ENABLE_RESOURCE_INFO
diff --git a/Engine/Source/Runtime/RenderCore/Private/RenderGraphResourcePool.cpp b/Engine/Source/Runtime/RenderCore/Private/RenderGraphResourcePool.cpp
--- a/Engine/Source/Runtime/RenderCore/Private/RenderGraphResourcePool.cpp
+++ b/Engine/Source/Runtime/RenderCore/Private/RenderGraphResourcePool.cpp
@@ -32,7 +32,9 @@ void FRDGBufferPool::DumpMemoryUsage(FOutputDevice& OutputDevice)
 {
 	OutputDevice.Logf(TEXT("Pooled Buffers:"));
 
+	Mutex.Lock();
 	TArray<TRefCountPtr<FRDGPooledBuffer>> BuffersBySize = AllocatedBuffers;
+	Mutex.Unlock();
 
 	Algo::Sort(BuffersBySize, [](const TRefCountPtr<FRDGPooledBuffer>& LHS, const TRefCountPtr<FRDGPooledBuffer>& RHS)
 	{
@@ -55,10 +57,8 @@ void FRDGBufferPool::DumpMemoryUsage(FOutputDevice& OutputDevice)
 	}
 }
 
-TRefCountPtr<FRDGPooledBuffer> FRDGBufferPool::FindFreeBuffer(const FRDGBufferDesc& Desc, const TCHAR* InDebugName, ERDGPooledBufferAlignment Alignment)
+TRefCountPtr<FRDGPooledBuffer> FRDGBufferPool::FindFreeBuffer(FRHICommandListBase& RHICmdList, const FRDGBufferDesc& Desc, const TCHAR* InDebugName, ERDGPooledBufferAlignment Alignment)
 {
-	FRHICommandListBase& RHICmdList = FRHICommandListImmediate::Get();
-
 	const uint64 BufferPageSize = 64 * 1024;
 
 	FRDGBufferDesc AlignedDesc = Desc;
@@ -80,7 +80,9 @@ TRefCountPtr<FRDGPooledBuffer> FRDGBufferPool::FindFreeBuffer(const FRDGBufferDe
 	}
 
 	const uint32 BufferHash = GetTypeHash(AlignedDesc);
-	
+
+	UE::TScopeLock Lock(Mutex);
+
 	// First find if available.
 	for (int32 Index = 0; Index < AllocatedBufferHashes.Num(); ++Index)
 	{
@@ -160,6 +162,8 @@ void FRDGBufferPool::TickPoolElements()
 	int32 NumReleasedBuffers = 0;
 	int64 NumReleasedBufferBytes = 0;
 
+	UE::TScopeLock Lock(Mutex);
+
 	while (BufferIndex < AllocatedBuffers.Num())
 	{
 		TRefCountPtr<FRDGPooledBuffer>& Buffer = AllocatedBuffers[BufferIndex];
diff --git a/Engine/Source/Runtime/RenderCore/Private/RenderGraphResourcePool.h b/Engine/Source/Runtime/RenderCore/Private/RenderGraphResourcePool.h
--- a/Engine/Source/Runtime/RenderCore/Private/RenderGraphResourcePool.h
+++ b/Engine/Source/Runtime/RenderCore/Private/RenderGraphResourcePool.h
@@ -9,6 +9,7 @@
 #include "RenderResource.h"
 #include "RendererInterface.h"
 #include "RenderGraphResources.h"
+#include "Async/RecursiveMutex.h"
 
 class FRDGBufferPool : public FRenderResource
 {
@@ -18,13 +19,20 @@ public:
 	/** Call once per frame to trim elements from the pool. */
 	RENDERCORE_API void TickPoolElements();
 
-	RENDERCORE_API TRefCountPtr<FRDGPooledBuffer> FindFreeBuffer(const FRDGBufferDesc& Desc, const TCHAR* InDebugName, ERDGPooledBufferAlignment Alignment = ERDGPooledBufferAlignment::Page);
+	RENDERCORE_API TRefCountPtr<FRDGPooledBuffer> FindFreeBuffer(FRHICommandListBase& RHICmdList, const FRDGBufferDesc& Desc, const TCHAR* InDebugName, ERDGPooledBufferAlignment Alignment = ERDGPooledBufferAlignment::Page);
+
+	TRefCountPtr<FRDGPooledBuffer> FindFreeBuffer(const FRDGBufferDesc& Desc, const TCHAR* InDebugName, ERDGPooledBufferAlignment Alignment = ERDGPooledBufferAlignment::Page)
+	{
+		return FindFreeBuffer(FRHICommandListImmediate::Get(), Desc, InDebugName, Alignment);
+	}
 
 	RENDERCORE_API void DumpMemoryUsage(FOutputDevice& OutputDevice);
 
 private:
 	RENDERCORE_API void ReleaseRHI() override;
 
+	mutable UE::FRecursiveMutex Mutex;
+
 	/** Elements can be 0, we compact the buffer later. */
 	TArray<TRefCountPtr<FRDGPooledBuffer>> AllocatedBuffers;
 	TArray<uint32> AllocatedBufferHashes;
diff --git a/Engine/Source/Runtime/RenderCore/Private/RenderTargetPool.cpp b/Engine/Source/Runtime/RenderCore/Private/RenderTargetPool.cpp
--- a/Engine/Source/Runtime/RenderCore/Private/RenderTargetPool.cpp
+++ b/Engine/Source/Runtime/RenderCore/Private/RenderTargetPool.cpp
@@ -65,7 +65,7 @@ static uint32 ComputeSizeInKB(FPooledRenderTarget& Element)
 	return (Element.ComputeMemorySize() + 1023) / 1024;
 }
 
-TRefCountPtr<IPooledRenderTarget> FRenderTargetPool::FindFreeElement(FRHITextureCreateInfo Desc, const TCHAR* Name)
+TRefCountPtr<IPooledRenderTarget> FRenderTargetPool::FindFreeElement(FRHICommandListBase& RHICmdList, FRHITextureCreateInfo Desc, const TCHAR* Name)
 {
 	FPooledRenderTarget* Found = 0;
 	uint32 FoundIndex = -1;
@@ -113,15 +113,13 @@ TRefCountPtr<IPooledRenderTarget> FRenderTargetPool::FindFreeElement(FRHITexture
 			<< FRenderTargetPool_CreateTexture.Name(Name);
 #endif
 
-		FRHICommandListBase& RHICmdList = FRHICommandListImmediate::Get();
-
 		const ERHIAccess AccessInitial = ERHIAccess::SRVMask;
 		FRHITextureCreateDesc CreateDesc(Desc, AccessInitial, Name);
 		const static FLazyName ClassName(TEXT("FPooledRenderTarget"));
 		CreateDesc.SetClassName(ClassName);
 
 		Found = new FPooledRenderTarget(
-			RHICreateTexture(CreateDesc),
+			RHICmdList.CreateTexture(CreateDesc),
 			Translate(CreateDesc),
 			this);
 
@@ -166,10 +164,8 @@ TRefCountPtr<IPooledRenderTarget> FRenderTargetPool::FindFreeElement(FRHITexture
 	return TRefCountPtr<IPooledRenderTarget>(MoveTemp(Found));
 }
 
-bool FRenderTargetPool::FindFreeElement(const FRHITextureCreateInfo& Desc, TRefCountPtr<IPooledRenderTarget>& Out, const TCHAR* Name)
+bool FRenderTargetPool::FindFreeElement(FRHICommandListBase& RHICmdList, const FRHITextureCreateInfo& Desc, TRefCountPtr<IPooledRenderTarget>& Out, const TCHAR* Name)
 {
-	check(IsInRenderingThread());
-
 	if (!Desc.IsValid())
 	{
 		// no need to do anything
@@ -204,6 +200,8 @@ bool FRenderTargetPool::FindFreeElement(const FRHITextureCreateInfo& Desc, TRefC
 
 			if (Current->IsFree())
 			{
+				UE::TScopeLock Lock(Mutex);
+
 				AllocationLevelInKB -= ComputeSizeInKB(*Current);
 				TRACE_COUNTER_SUBTRACT(RenderTargetPoolCount, 1);
 				TRACE_COUNTER_SET(RenderTargetPoolSize, (int64)AllocationLevelInKB * 1024);
@@ -214,7 +212,7 @@ bool FRenderTargetPool::FindFreeElement(const FRHITextureCreateInfo& Desc, TRefC
 		}
 	}
 
-	Out = FindFreeElement(Desc, Name);
+	Out = FindFreeElement(RHICmdList, Desc, Name);
 	return false;
 }
 
@@ -227,6 +225,7 @@ void FRenderTargetPool::CreateUntrackedElement(const FPooledRenderTargetDesc& De
 
 void FRenderTargetPool::GetStats(uint32& OutWholeCount, uint32& OutWholePoolInKB, uint32& OutUsedInKB) const
 {
+	UE::TScopeLock Lock(Mutex);
 	OutWholeCount = (uint32)PooledRenderTargets.Num();
 	OutUsedInKB = 0;
 	OutWholePoolInKB = 0;
@@ -254,13 +253,14 @@ void FRenderTargetPool::GetStats(uint32& OutWholeCount, uint32& OutWholePoolInKB
 
 void FRenderTargetPool::TickPoolElements()
 {
+	UE::TScopeLock Lock(Mutex);
+
 	uint32 DeferredAllocationLevelInKB = 0;
 	for (FPooledRenderTarget* Element : DeferredDeleteArray)
 	{
 		DeferredAllocationLevelInKB += ComputeSizeInKB(*Element);
 	}
 
-	check(IsInRenderingThread());
 	DeferredDeleteArray.Reset();
 
 	uint32 MinimumPoolSizeInKB;
@@ -370,7 +370,7 @@ void FRenderTargetPool::TickPoolElements()
 
 int32 FRenderTargetPool::FindIndex(IPooledRenderTarget* In) const
 {
-	check(IsInRenderingThread());
+	UE::TScopeLock Lock(Mutex);
 
 	if (In)
 	{
@@ -398,7 +398,7 @@ void FRenderTargetPool::FreeElementAtIndex(int32 Index)
 
 void FRenderTargetPool::FreeUnusedResource(TRefCountPtr<IPooledRenderTarget>& In)
 {
-	check(IsInRenderingThread());
+	UE::TScopeLock Lock(Mutex);
 
 	int32 Index = FindIndex(In);
 
@@ -425,7 +425,7 @@ void FRenderTargetPool::FreeUnusedResource(TRefCountPtr<IPooledRenderTarget>& In
 
 void FRenderTargetPool::FreeUnusedResources()
 {
-	check(IsInRenderingThread());
+	UE::TScopeLock Lock(Mutex);
 
 	for (uint32 i = 0, Num = (uint32)PooledRenderTargets.Num(); i < Num; ++i)
 	{
@@ -447,6 +447,8 @@ void FRenderTargetPool::FreeUnusedResources()
 
 void FRenderTargetPool::DumpMemoryUsage(FOutputDevice& OutputDevice)
 {
+	UE::TScopeLock Lock(Mutex);
+
 	uint32 UnusedAllocationInKB = 0;
 
 	OutputDevice.Logf(TEXT("Pooled Render Targets:"));
@@ -510,7 +512,7 @@ void FRenderTargetPool::DumpMemoryUsage(FOutputDevice& OutputDevice)
 
 void FRenderTargetPool::ReleaseRHI()
 {
-	check(IsInRenderingThread());
+	UE::TScopeLock Lock(Mutex);
 	DeferredDeleteArray.Empty();
 	PooledRenderTargets.Empty();
 }
@@ -549,8 +551,6 @@ void FRenderTargetPool::CompactPool()
 
 bool FPooledRenderTarget::OnFrameStart()
 {
-	check(IsInRenderingThread());
-
 	// If there are any references to the pooled render target other than the pool itself, then it may not be freed.
 	if (!IsFree())
 	{
diff --git a/Engine/Source/Runtime/RenderCore/Public/RenderTargetPool.h b/Engine/Source/Runtime/RenderCore/Public/RenderTargetPool.h
--- a/Engine/Source/Runtime/RenderCore/Public/RenderTargetPool.h
+++ b/Engine/Source/Runtime/RenderCore/Public/RenderTargetPool.h
@@ -16,6 +16,7 @@
 #include "RenderResource.h"
 #include "RendererInterface.h"
 #include "Templates/RefCounting.h"
+#include "Async/RecursiveMutex.h"
 
 class FOutputDevice;
 class FRHICommandList;
@@ -95,9 +96,19 @@ class FRenderTargetPool : public FRenderResource
 public:
 	FRenderTargetPool() = default;
 
-	RENDERCORE_API TRefCountPtr<IPooledRenderTarget> FindFreeElement(FRHITextureCreateInfo Desc, const TCHAR* Name);
+	RENDERCORE_API TRefCountPtr<IPooledRenderTarget> FindFreeElement(FRHICommandListBase& RHICmdList, FRHITextureCreateInfo Desc, const TCHAR* Name);
 
-	RENDERCORE_API bool FindFreeElement(const FRHITextureCreateInfo& Desc, TRefCountPtr<IPooledRenderTarget>& Out, const TCHAR* Name);
+	RENDERCORE_API bool FindFreeElement(FRHICommandListBase& RHICmdList, const FRHITextureCreateInfo& Desc, TRefCountPtr<IPooledRenderTarget>& Out, const TCHAR* Name);
+
+	TRefCountPtr<IPooledRenderTarget> FindFreeElement(FRHITextureCreateInfo Desc, const TCHAR* Name)
+	{
+		return FindFreeElement(FRHICommandListImmediate::Get(), Desc, Name);
+	}
+
+	bool FindFreeElement(const FRHITextureCreateInfo& Desc, TRefCountPtr<IPooledRenderTarget>& Out, const TCHAR* Name)
+	{
+		return FindFreeElement(FRHICommandListImmediate::Get(), Desc, Out, Name);
+	}
 
 	/**
 	 * @param DebugName must not be 0, we only store the pointer
@@ -106,12 +117,12 @@ public:
 	 * @return true if the old element was still valid, false if a new one was assigned
 	 */
 	bool FindFreeElement(
-		FRHICommandList& RHICmdList,
+		FRHICommandListBase& RHICmdList,
 		const FPooledRenderTargetDesc& Desc,
 		TRefCountPtr<IPooledRenderTarget>& Out,
 		const TCHAR* InDebugName)
 	{
-		return FindFreeElement(Translate(Desc), Out, InDebugName);
+		return FindFreeElement(RHICmdList, Translate(Desc), Out, InDebugName);
 	}
 
 	RENDERCORE_API void CreateUntrackedElement(const FPooledRenderTargetDesc& Desc, TRefCountPtr<IPooledRenderTarget>& Out, const FSceneRenderTargetItem& Item);
@@ -147,6 +158,8 @@ public:
 private:
 	RENDERCORE_API void FreeElementAtIndex(int32 Index);
 
+	mutable UE::FRecursiveMutex Mutex;
+
 	/** Elements can be 0, we compact the buffer later. */
 	TArray<uint32> PooledRenderTargetHashes;
 	TArray< TRefCountPtr<FPooledRenderTarget> > PooledRenderTargets;
