Fixed unity build break.

diff --git a/Engine/Source/Runtime/RenderCore/Private/RenderResource.cpp b/Engine/Source/Runtime/RenderCore/Private/RenderResource.cpp
--- a/Engine/Source/Runtime/RenderCore/Private/RenderResource.cpp
+++ b/Engine/Source/Runtime/RenderCore/Private/RenderResource.cpp
@@ -558,6 +558,11 @@ FString FVertexBuffer::GetFriendlyName() const
 	return TEXT("FVertexBuffer");
 }
 
+void FVertexBuffer::SetRHI(const FBufferRHIRef& BufferRHI)
+{
+	VertexBufferRHI = BufferRHI;
+}
+
 //////////////////////////////////////////////////////////////////////////////////////////////////////////////////////
 // FVertexBufferWithSRV
 
@@ -587,6 +592,10 @@ FString FIndexBuffer::GetFriendlyName() const
 	return TEXT("FIndexBuffer");
 }
 
+void FIndexBuffer::SetRHI(const FBufferRHIRef& BufferRHI)
+{
+	IndexBufferRHI = BufferRHI;
+}
 //////////////////////////////////////////////////////////////////////////////////////////////////////////////////////
 // FBufferWithRDG
 
diff --git a/Engine/Source/Runtime/RenderCore/Public/RenderResource.h b/Engine/Source/Runtime/RenderCore/Public/RenderResource.h
--- a/Engine/Source/Runtime/RenderCore/Public/RenderResource.h
+++ b/Engine/Source/Runtime/RenderCore/Public/RenderResource.h
@@ -476,7 +476,7 @@ public:
 
 	const FBufferRHIRef& GetRHI() const { return VertexBufferRHI; }
 
-	void SetRHI(const FBufferRHIRef& BufferRHI) { VertexBufferRHI = BufferRHI; }
+	RENDERCORE_API void SetRHI(const FBufferRHIRef& BufferRHI);
 
 	FBufferRHIRef VertexBufferRHI;
 };
@@ -509,7 +509,7 @@ public:
 
 	const FBufferRHIRef& GetRHI() const { return IndexBufferRHI; }
 
-	void SetRHI(const FBufferRHIRef& BufferRHI) { IndexBufferRHI = BufferRHI; }
+	RENDERCORE_API void SetRHI(const FBufferRHIRef& BufferRHI);
 
 	FBufferRHIRef IndexBufferRHI;
 };
