Allow 0 sized RDG buffers to CreateBuffer when a size callback is supplied.

diff --git a/Engine/Source/Runtime/RenderCore/Private/RenderGraphValidation.cpp b/Engine/Source/Runtime/RenderCore/Private/RenderGraphValidation.cpp
--- a/Engine/Source/Runtime/RenderCore/Private/RenderGraphValidation.cpp
+++ b/Engine/Source/Runtime/RenderCore/Private/RenderGraphValidation.cpp
@@ -226,6 +226,9 @@ void FRDGUserValidation::ValidateCreateBuffer(FRDGBufferRef Buffer)
 	}
 
 	ValidateCreateViewableResource(Buffer);
+
+	checkf(Buffer->Desc.GetSize() > 0 || Buffer->NumElementsCallback, TEXT("Creating buffer '%s' is zero bytes in size."), Buffer->Name);
+
 	Buffer->BufferDebugData = Allocator.Alloc<FRDGBufferDebugData>();
 	if (GRDGDebug)
 	{
@@ -366,8 +369,6 @@ void FRDGUserValidation::ValidateCreateBuffer(const FRDGBufferDesc& Desc, const
 	checkf(Name, TEXT("Creating a buffer requires a valid debug name."));
 	ExecuteGuard(TEXT("CreateBuffer"), Name);
 
-	checkf(Desc.GetSize() > 0, TEXT("Creating buffer '%s' is zero bytes in size."), Name);
-
 	if (EnumHasAllFlags(Desc.Usage, EBufferUsageFlags::StructuredBuffer | EBufferUsageFlags::ByteAddressBuffer))
 	{
 		checkf(Desc.BytesPerElement == 4, TEXT("Creating buffer '%s' as a structured buffer that is also byte addressable, BytesPerElement must be 4! Instead it is %d"), Name, Desc.BytesPerElement);
