Added variant of QueueUploadBuffer which avoids an additional memcpy.

diff --git a/Engine/Source/Runtime/RenderCore/Private/RenderGraphBuilder.cpp b/Engine/Source/Runtime/RenderCore/Private/RenderGraphBuilder.cpp
--- a/Engine/Source/Runtime/RenderCore/Private/RenderGraphBuilder.cpp
+++ b/Engine/Source/Runtime/RenderCore/Private/RenderGraphBuilder.cpp
@@ -2526,6 +2526,15 @@ UE::Tasks::FTask FRDGBuilder::SubmitBufferUploads()
 					UploadedBuffer.DataFreeCallback(UploadedBuffer.Data);
 				}
 			}
+			else if (UploadedBuffer.DataFillCallback)
+			{
+				const uint32 BufferSize = UploadedBuffer.Buffer->Desc.GetSize();
+				check(BufferSize > 0);
+
+				void* DestPtr = RHICmdListUpload.LockBuffer(UploadedBuffer.Buffer->GetRHIUnchecked(), 0, BufferSize, RLM_WriteOnly);
+				UploadedBuffer.DataFillCallback(DestPtr, BufferSize);
+				RHICmdListUpload.UnlockBuffer(UploadedBuffer.Buffer->GetRHIUnchecked());
+			}
 		}
 
 		UploadedBuffers.Reset();
diff --git a/Engine/Source/Runtime/RenderCore/Private/RenderGraphValidation.cpp b/Engine/Source/Runtime/RenderCore/Private/RenderGraphValidation.cpp
--- a/Engine/Source/Runtime/RenderCore/Private/RenderGraphValidation.cpp
+++ b/Engine/Source/Runtime/RenderCore/Private/RenderGraphValidation.cpp
@@ -467,6 +467,18 @@ void FRDGUserValidation::ValidateUploadBuffer(FRDGBufferRef Buffer, const void*
 	check((InitialData || InitialDataSize == 0) && InitialDataFreeCallback);
 }
 
+void FRDGUserValidation::ValidateUploadBuffer(FRDGBufferRef Buffer, const FRDGBufferInitialDataFillCallback& InitialDataFillCallback)
+{
+	if (!GRDGValidation)
+	{
+		return;
+	}
+
+	check(Buffer);
+	checkf(!Buffer->bQueuedForUpload, TEXT("Buffer %s already has an upload queued. Only one upload can be done for each graph."), Buffer->Name);
+	check(InitialDataFillCallback);
+}
+
 void FRDGUserValidation::ValidateUploadBuffer(FRDGBufferRef Buffer, const FRDGBufferInitialDataCallback& InitialDataCallback, const FRDGBufferInitialDataSizeCallback& InitialDataSizeCallback)
 {
 	if (!GRDGValidation)
diff --git a/Engine/Source/Runtime/RenderCore/Public/RenderGraphBuilder.h b/Engine/Source/Runtime/RenderCore/Public/RenderGraphBuilder.h
--- a/Engine/Source/Runtime/RenderCore/Public/RenderGraphBuilder.h
+++ b/Engine/Source/Runtime/RenderCore/Public/RenderGraphBuilder.h
@@ -259,7 +259,7 @@ public:
 	RENDERCORE_API void SetFlushResourcesRHI();
 
 	/** Queues a buffer upload operation prior to execution. The resource lifetime is extended and the data is uploaded prior to executing passes. */
-	RENDERCORE_API void QueueBufferUpload(FRDGBufferRef Buffer, const void* InitialData, uint64 InitialDataSize, ERDGInitialDataFlags InitialDataFlags = ERDGInitialDataFlags::None);
+	void QueueBufferUpload(FRDGBufferRef Buffer, const void* InitialData, uint64 InitialDataSize, ERDGInitialDataFlags InitialDataFlags = ERDGInitialDataFlags::None);
 
 	template <typename ElementType>
 	inline void QueueBufferUpload(FRDGBufferRef Buffer, TArrayView<ElementType, int32> Container, ERDGInitialDataFlags InitialDataFlags = ERDGInitialDataFlags::None)
@@ -268,7 +268,7 @@ public:
 	}
 
 	/** Queues a buffer upload operation prior to execution. The resource lifetime is extended and the data is uploaded prior to executing passes. */
-	RENDERCORE_API void QueueBufferUpload(FRDGBufferRef Buffer, const void* InitialData, uint64 InitialDataSize, FRDGBufferInitialDataFreeCallback&& InitialDataFreeCallback);
+	void QueueBufferUpload(FRDGBufferRef Buffer, const void* InitialData, uint64 InitialDataSize, FRDGBufferInitialDataFreeCallback&& InitialDataFreeCallback);
 
 	template <typename ElementType>
 	inline void QueueBufferUpload(FRDGBufferRef Buffer, TArrayView<ElementType, int32> Container, FRDGBufferInitialDataFreeCallback&& InitialDataFreeCallback)
@@ -276,11 +276,14 @@ public:
 		QueueBufferUpload(Buffer, Container.GetData(), Container.Num() * sizeof(ElementType), InitialDataFreeCallback);
 	}
 
+	/** A variant where the buffer is mapped and the pointer / size is provided to the callback to fill the buffer pointer. */
+	void QueueBufferUpload(FRDGBufferRef Buffer, FRDGBufferInitialDataFillCallback&& InitialDataFillCallback);
+
 	/** A variant where InitialData and InitialDataSize are supplied through callbacks. This allows queuing an upload with information unknown at
 	 *  creation time. The callbacks are called before RDG pass execution so data must be ready before that.
 	 */
-	RENDERCORE_API void QueueBufferUpload(FRDGBufferRef Buffer, FRDGBufferInitialDataCallback&& InitialDataCallback, FRDGBufferInitialDataSizeCallback&& InitialDataSizeCallback);
-	RENDERCORE_API void QueueBufferUpload(FRDGBufferRef Buffer, FRDGBufferInitialDataCallback&& InitialDataCallback, FRDGBufferInitialDataSizeCallback&& InitialDataSizeCallback, FRDGBufferInitialDataFreeCallback&& InitialDataFreeCallback);
+	void QueueBufferUpload(FRDGBufferRef Buffer, FRDGBufferInitialDataCallback&& InitialDataCallback, FRDGBufferInitialDataSizeCallback&& InitialDataSizeCallback);
+	void QueueBufferUpload(FRDGBufferRef Buffer, FRDGBufferInitialDataCallback&& InitialDataCallback, FRDGBufferInitialDataSizeCallback&& InitialDataSizeCallback, FRDGBufferInitialDataFreeCallback&& InitialDataFreeCallback);
 
 	/** Queues a pooled render target extraction to happen at the end of graph execution. For graph-created textures, this extends
 	 *  the lifetime of the GPU resource until execution, at which point the pointer is filled. If specified, the texture is transitioned
@@ -603,16 +606,18 @@ private:
 		FUploadedBuffer() = default;
 
 		FUploadedBuffer(FRDGBuffer* InBuffer, const void* InData, uint64 InDataSize)
-			: bUseDataCallbacks(false)
-			, bUseFreeCallbacks(false)
-			, Buffer(InBuffer)
+			: Buffer(InBuffer)
 			, Data(InData)
 			, DataSize(InDataSize)
 		{}
 
+		FUploadedBuffer(FRDGBuffer* InBuffer, FRDGBufferInitialDataFillCallback&& InDataFillCallback)
+			: Buffer(InBuffer)
+			, DataFillCallback(MoveTemp(InDataFillCallback))
+		{}
+
 		FUploadedBuffer(FRDGBuffer* InBuffer, const void* InData, uint64 InDataSize, FRDGBufferInitialDataFreeCallback&& InDataFreeCallback)
-			: bUseDataCallbacks(false)
-			, bUseFreeCallbacks(true)
+			: bUseFreeCallbacks(true)
 			, Buffer(InBuffer)
 			, Data(InData)
 			, DataSize(InDataSize)
@@ -621,7 +626,6 @@ private:
 
 		FUploadedBuffer(FRDGBuffer* InBuffer, FRDGBufferInitialDataCallback&& InDataCallback, FRDGBufferInitialDataSizeCallback&& InDataSizeCallback)
 			: bUseDataCallbacks(true)
-			, bUseFreeCallbacks(false)
 			, Buffer(InBuffer)
 			, DataCallback(MoveTemp(InDataCallback))
 			, DataSizeCallback(MoveTemp(InDataSizeCallback))
@@ -641,9 +645,14 @@ private:
 		FRDGBuffer* Buffer{};
 		const void* Data{};
 		uint64 DataSize{};
+
+		// User provided data callbacks
 		FRDGBufferInitialDataCallback DataCallback;
 		FRDGBufferInitialDataSizeCallback DataSizeCallback;
 		FRDGBufferInitialDataFreeCallback DataFreeCallback;
+
+		// RDG provided buffer pointer callback.
+		FRDGBufferInitialDataFillCallback DataFillCallback;
 	};
 
 	TArray<FUploadedBuffer, FRDGArrayAllocator> UploadedBuffers;
diff --git a/Engine/Source/Runtime/RenderCore/Public/RenderGraphBuilder.inl b/Engine/Source/Runtime/RenderCore/Public/RenderGraphBuilder.inl
--- a/Engine/Source/Runtime/RenderCore/Public/RenderGraphBuilder.inl
+++ b/Engine/Source/Runtime/RenderCore/Public/RenderGraphBuilder.inl
@@ -347,6 +347,15 @@ inline void FRDGBuilder::QueueBufferUpload(FRDGBufferRef Buffer, const void* Ini
 	Buffer->bForceNonTransient = 1;
 }
 
+inline void FRDGBuilder::QueueBufferUpload(FRDGBufferRef Buffer, FRDGBufferInitialDataFillCallback&& InitialDataFillCallback)
+{
+	IF_RDG_ENABLE_DEBUG(UserValidation.ValidateUploadBuffer(Buffer, InitialDataFillCallback));
+
+	UploadedBuffers.Emplace(Buffer, MoveTemp(InitialDataFillCallback));
+	Buffer->bQueuedForUpload = 1;
+	Buffer->bForceNonTransient = 1;
+}
+
 inline void FRDGBuilder::QueueBufferUpload(FRDGBufferRef Buffer, FRDGBufferInitialDataCallback&& InitialDataCallback, FRDGBufferInitialDataSizeCallback&& InitialDataSizeCallback)
 {
 	IF_RDG_ENABLE_DEBUG(UserValidation.ValidateUploadBuffer(Buffer, InitialDataCallback, InitialDataSizeCallback));
diff --git a/Engine/Source/Runtime/RenderCore/Public/RenderGraphDefinitions.h b/Engine/Source/Runtime/RenderCore/Public/RenderGraphDefinitions.h
--- a/Engine/Source/Runtime/RenderCore/Public/RenderGraphDefinitions.h
+++ b/Engine/Source/Runtime/RenderCore/Public/RenderGraphDefinitions.h
@@ -736,4 +736,5 @@ template <typename ArrayType,
 	typename ArrayTypeNoRef = std::remove_reference_t<ArrayType>,
 	typename = typename TEnableIf<TIsTArray_V<ArrayTypeNoRef>>::Type> using TRDGBufferArrayCallback = TFunction<const ArrayType&()>;
 using FRDGBufferInitialDataFreeCallback = TFunction<void(const void* InData)>;
+using FRDGBufferInitialDataFillCallback = TFunction<void(void* InData, uint32 InDataSize)>;
 using FRDGDispatchGroupCountCallback = TFunction<FIntVector()>;
diff --git a/Engine/Source/Runtime/RenderCore/Public/RenderGraphValidation.h b/Engine/Source/Runtime/RenderCore/Public/RenderGraphValidation.h
--- a/Engine/Source/Runtime/RenderCore/Public/RenderGraphValidation.h
+++ b/Engine/Source/Runtime/RenderCore/Public/RenderGraphValidation.h
@@ -77,6 +77,7 @@ public:
 
 	RENDERCORE_API void ValidateUploadBuffer(FRDGBufferRef Buffer, const void* InitialData, uint64 InitialDataSize);
 	RENDERCORE_API void ValidateUploadBuffer(FRDGBufferRef Buffer, const void* InitialData, uint64 InitialDataSize, const FRDGBufferInitialDataFreeCallback& InitialDataFreeCallback);
+	RENDERCORE_API void ValidateUploadBuffer(FRDGBufferRef Buffer, const FRDGBufferInitialDataFillCallback& InitialDataFillCallback);
 	RENDERCORE_API void ValidateUploadBuffer(FRDGBufferRef Buffer, const FRDGBufferInitialDataCallback& InitialDataCallback, const FRDGBufferInitialDataSizeCallback& InitialDataSizeCallback);
 	RENDERCORE_API void ValidateUploadBuffer(FRDGBufferRef Buffer, const FRDGBufferInitialDataCallback& InitialDataCallback, const FRDGBufferInitialDataSizeCallback& InitialDataSizeCallback, const FRDGBufferInitialDataFreeCallback& InitialDataFreeCallback);
 
