Re-enabled render command pipes and fixed Skinning / Niagara race conditions.
 - Added flags to render command pipe definitions to allow for disabling in code.
 - Release the TFunction after each command to match behavior of the render command tasks.

diff --git a/Engine/Plugins/FX/Niagara/Source/Niagara/Private/NiagaraComponent.cpp b/Engine/Plugins/FX/Niagara/Source/Niagara/Private/NiagaraComponent.cpp
--- a/Engine/Plugins/FX/Niagara/Source/Niagara/Private/NiagaraComponent.cpp
+++ b/Engine/Plugins/FX/Niagara/Source/Niagara/Private/NiagaraComponent.cpp
@@ -43,7 +43,7 @@ DECLARE_CYCLE_STAT(TEXT("Get Dynamic Mesh Elements (RT)"), STAT_NiagaraComponent
 
 DEFINE_LOG_CATEGORY(LogNiagara);
 
-DEFINE_RENDER_COMMAND_PIPE(NiagaraDynamicData);
+DEFINE_RENDER_COMMAND_PIPE(NiagaraDynamicData, ERenderCommandPipeFlags::None);
 
 static int GNiagaraSoloTickEarly = 1;
 static FAutoConsoleVariableRef CVarNiagaraSoloTickEarly(
diff --git a/Engine/Plugins/FX/Niagara/Source/Niagara/Private/NiagaraSystemRenderData.cpp b/Engine/Plugins/FX/Niagara/Source/Niagara/Private/NiagaraSystemRenderData.cpp
--- a/Engine/Plugins/FX/Niagara/Source/Niagara/Private/NiagaraSystemRenderData.cpp
+++ b/Engine/Plugins/FX/Niagara/Source/Niagara/Private/NiagaraSystemRenderData.cpp
@@ -8,6 +8,8 @@
 #include "PrimitiveViewRelevance.h"
 #include "SceneView.h"
 
+DECLARE_RENDER_COMMAND_PIPE(NiagaraDynamicData, );
+
 void FNiagaraSystemRenderData::ExecuteDynamicDataCommands_RenderThread(const FSetDynamicDataCommandList& Commands)
 {
 	for (auto& Command : Commands)
@@ -298,12 +300,21 @@ void FNiagaraSystemRenderData::RecacheRenderers(const FNiagaraSystemInstance& Sy
 
 	// If we have renderers then the draw order on the system should match, when compiling the number of renderers can be zero
 	checkf((EmitterRenderers_GT.Num() == 0) || (EmitterRenderers_GT.Num() == RendererDrawOrder.Num()), TEXT("EmitterRenderers Num %d does not match System DrawOrder %d"), EmitterRenderers_GT.Num(), RendererDrawOrder.Num());
-	
+
+	UE::Tasks::FTaskEvent TaskEvent{ UE_SOURCE_LOCATION };
+
+	ENQUEUE_RENDER_COMMAND(SignalFence)(UE::RenderCommandPipe::NiagaraDynamicData, [TaskEvent] () mutable
+	{
+		TaskEvent.Trigger();
+	});
+
 	// NOTE: Since this object and its render thread resources may be concurrently accessed on the render thread, we have to create the renderers and pass them off to
 	// replace the current ones on the render thread's time line
 	ENQUEUE_RENDER_COMMAND(NiagaraRecacheRenderers)(
-		[this, EmitterRenderers_Copy=EmitterRenderers_GT](FRHICommandListImmediate& RHICmdList) mutable
+		[this, EmitterRenderers_Copy=EmitterRenderers_GT, TaskEvent = MoveTemp(TaskEvent)](FRHICommandListImmediate& RHICmdList) mutable
 		{
+			TaskEvent.Wait();
+
 			// Release/destroy the current renderers
 			Destroy_RenderThread();
 
diff --git a/Engine/Plugins/Runtime/CableComponent/Source/CableComponent/Private/CableComponent.cpp b/Engine/Plugins/Runtime/CableComponent/Source/CableComponent/Private/CableComponent.cpp
--- a/Engine/Plugins/Runtime/CableComponent/Source/CableComponent/Private/CableComponent.cpp
+++ b/Engine/Plugins/Runtime/CableComponent/Source/CableComponent/Private/CableComponent.cpp
@@ -14,7 +14,7 @@
 #include "StaticMeshResources.h"
 #include "SceneInterface.h"
 
-DEFINE_RENDER_COMMAND_PIPE(Cable);
+DEFINE_RENDER_COMMAND_PIPE(Cable, ERenderCommandPipeFlags::None);
 
 static TAutoConsoleVariable<int32> CVarRayTracingCableMeshes(
 	TEXT("r.RayTracing.Geometry.Cable"),
diff --git a/Engine/Source/Runtime/Engine/Private/GPUSkinCache.cpp b/Engine/Source/Runtime/Engine/Private/GPUSkinCache.cpp
--- a/Engine/Source/Runtime/Engine/Private/GPUSkinCache.cpp
+++ b/Engine/Source/Runtime/Engine/Private/GPUSkinCache.cpp
@@ -9,6 +9,7 @@ GPUSkinCache.cpp: Performs skinning on a compute shader into a buffer to avoid v
 #include "CanvasTypes.h"
 #include "Engine/Engine.h"
 #include "Engine/SkeletalMesh.h"
+#include "Rendering/RenderCommandPipes.h"
 #include "SkeletalRenderGPUSkin.h"
 #include "MeshDrawShaderBindings.h"
 #include "ShaderParameterUtils.h"
@@ -2282,7 +2283,7 @@ void FGPUSkinCache::CVarSinkFunction()
 	if (NewGPUSkinCacheValue != GEnableGPUSkinCache || NewRecomputeTangentsValue != GSkinCacheRecomputeTangents
 		|| NewSceneMaxSizeInMb != GSkinCacheSceneMemoryLimitInMB || NewNumTangentIntermediateBuffers != GNumTangentIntermediateBuffers)
 	{
-		ENQUEUE_RENDER_COMMAND(DoEnableSkinCaching)(
+		ENQUEUE_RENDER_COMMAND(DoEnableSkinCaching)(UE::RenderCommandPipe::SkeletalMesh,
 			[NewRecomputeTangentsValue, NewGPUSkinCacheValue, NewSceneMaxSizeInMb, NewNumTangentIntermediateBuffers](FRHICommandList& RHICmdList)
 		{
 			GNumTangentIntermediateBuffers = FMath::Max(NewNumTangentIntermediateBuffers, 1);
diff --git a/Engine/Source/Runtime/Engine/Private/Rendering/RenderCommandPipes.cpp b/Engine/Source/Runtime/Engine/Private/Rendering/RenderCommandPipes.cpp
--- a/Engine/Source/Runtime/Engine/Private/Rendering/RenderCommandPipes.cpp
+++ b/Engine/Source/Runtime/Engine/Private/Rendering/RenderCommandPipes.cpp
@@ -2,4 +2,4 @@
 
 #include "Rendering/RenderCommandPipes.h"
 
-DEFINE_RENDER_COMMAND_PIPE(SkeletalMesh);
\ No newline at end of file
+DEFINE_RENDER_COMMAND_PIPE(SkeletalMesh, ERenderCommandPipeFlags::None);
\ No newline at end of file
diff --git a/Engine/Source/Runtime/RenderCore/Private/RenderingThread.cpp b/Engine/Source/Runtime/RenderCore/Private/RenderingThread.cpp
--- a/Engine/Source/Runtime/RenderCore/Private/RenderingThread.cpp
+++ b/Engine/Source/Runtime/RenderCore/Private/RenderingThread.cpp
@@ -1540,7 +1540,7 @@ inline ERenderCommandPipeMode GetValidatedRenderCommandPipeMode(int32 CVarValue)
 ERenderCommandPipeMode GRenderCommandPipeMode = ERenderCommandPipeMode::None;
 FAutoConsoleVariable CVarRenderCommandPipeMode(
 	TEXT("r.RenderCommandPipeMode"),
-	1,
+	2,
 	TEXT("Controls behavior of the main render thread command pipe.")
 	TEXT(" 0: Render commands are launched individually as tasks;\n")
 	TEXT(" 1: Render commands are enqueued into a render command pipe for the render thread only.;\n")
@@ -1600,6 +1600,9 @@ public:
 
 		for (TLinkedList<FRenderCommandPipe*>::TIterator PipeIt(GlobalList); PipeIt; PipeIt.Next())
 		{
+			FRenderCommandPipe* Pipe = *PipeIt;
+			Pipe->SetEnabled(Pipe->ConsoleVariable->GetBool());
+
 			Pipes.Emplace(*PipeIt);
 		}
 #endif // !UE_SERVER
@@ -1773,9 +1776,13 @@ namespace UE::RenderCommandPipe
 	}
 }
 
-FRenderCommandPipe::FRenderCommandPipe(const TCHAR* InName)
+FRenderCommandPipe::FRenderCommandPipe(const TCHAR* InName, ERenderCommandPipeFlags Flags, const TCHAR* CVarName, const TCHAR* CVarDescription)
 	: Name(InName)
 	, GlobalListLink(this)
+	, ConsoleVariable(CVarName, !EnumHasAnyFlags(Flags, ERenderCommandPipeFlags::Disabled), CVarDescription, FConsoleVariableDelegate::CreateLambda([this](IConsoleVariable* Variable)
+	{
+		SetEnabled(Variable->GetBool());
+	}))
 {
 	GlobalListLink.LinkHead(GRenderCommandPipeRegistry.GlobalList);
 }
@@ -1832,6 +1839,8 @@ void FRenderCommandPipe::EnqueueAndLaunch(FFunctionVariant&& FunctionVariant, co
 			{
 				Command.Function.Get<FEmptyFunction>()();
 			}
+
+			Command.Function = {};
 		}
 
 	}, Frame_GameThread->TaskEvent);
diff --git a/Engine/Source/Runtime/RenderCore/Public/RenderingThread.h b/Engine/Source/Runtime/RenderCore/Public/RenderingThread.h
--- a/Engine/Source/Runtime/RenderCore/Public/RenderingThread.h
+++ b/Engine/Source/Runtime/RenderCore/Public/RenderingThread.h
@@ -284,6 +284,16 @@ enum class ERenderCommandPipeMode
 	All
 };
 
+enum class ERenderCommandPipeFlags : uint8
+{
+	None = 0,
+
+	/** Initializes the render command pipe in a disabled state. */
+	Disabled = 1 << 0
+};
+
+ENUM_CLASS_FLAGS(ERenderCommandPipeFlags);
+
 namespace UE::RenderCommandPipe
 {
 	// [Game Thread] Initializes all statically initialized render command pipes.
@@ -404,7 +414,7 @@ public:
 	using FCommandListFunction = TUniqueFunction<void(FRHICommandList&)>;
 	using FEmptyFunction = TUniqueFunction<void()>;
 
-	RENDERCORE_API FRenderCommandPipe(const TCHAR* Name);
+	RENDERCORE_API FRenderCommandPipe(const TCHAR* Name, ERenderCommandPipeFlags Flags, const TCHAR* CVarName, const TCHAR* CVarDescription);
 	RENDERCORE_API ~FRenderCommandPipe();
 
 	FORCEINLINE const TCHAR* GetName() const
@@ -529,6 +539,7 @@ private:
 	FFrame* Frame_GameThread = nullptr;
 	FFrame* Frame_RenderThread = nullptr;
 	TLinkedList<FRenderCommandPipe*> GlobalListLink;
+	FAutoConsoleVariable ConsoleVariable;
 	bool bEnabled = true;
 };
 
@@ -537,23 +548,17 @@ private:
 	namespace UE::RenderCommandPipe { extern PrefixKeywords FRenderCommandPipe Name; }
 
 /** Defines a render command pipe. */
-#define DEFINE_RENDER_COMMAND_PIPE(Name) \
+#define DEFINE_RENDER_COMMAND_PIPE(Name, Flags) \
 	namespace UE::RenderCommandPipe \
 	{ \
-		FRenderCommandPipe Name(TEXT(#Name)); \
-		namespace \
-		{ \
-			FAutoConsoleVariable CVar_##Name( \
-				TEXT("r.RenderCommandPipe." #Name), true, \
-				TEXT("Whether to enable the " #Name " Render Command Pipe") \
-				TEXT(" 0: off;") \
-				TEXT(" 1: on (default)"), \
-				FConsoleVariableDelegate::CreateLambda([](IConsoleVariable* Variable) \
-			{ \
-				UE::RenderCommandPipe::Name.SetEnabled(Variable->GetBool()); \
-			}), \
-			ECVF_Default); \
-		} \
+		FRenderCommandPipe Name( \
+			TEXT(#Name), \
+			Flags, \
+			TEXT("r.RenderCommandPipe." #Name), \
+			TEXT("Whether to enable the " #Name " Render Command Pipe") \
+			TEXT(" 0: off;") \
+			TEXT(" 1: on (default)") \
+		); \
 	}
 
 /** Enqueues a render command to a render pipe. The default implementation takes a lambda and schedules on the render thread.
