Fix for incorrectly sized dynamic index buffers

diff --git a/Engine/Source/Runtime/RenderCore/Private/GlobalRenderResources.cpp b/Engine/Source/Runtime/RenderCore/Private/GlobalRenderResources.cpp
--- a/Engine/Source/Runtime/RenderCore/Private/GlobalRenderResources.cpp
+++ b/Engine/Source/Runtime/RenderCore/Private/GlobalRenderResources.cpp
@@ -892,14 +892,14 @@ void FGlobalDynamicVertexBuffer::Commit()
 	Commit(FRHICommandListImmediate::Get());
 }
 
-FGlobalDynamicIndexBuffer::FAllocation FGlobalDynamicIndexBuffer::Allocate(uint32 SizeInBytes, uint32 IndexStride)
+FGlobalDynamicIndexBuffer::FAllocation FGlobalDynamicIndexBuffer::Allocate(uint32 NumIndices, uint32 IndexStride)
 {
-	return Allocate(FRHICommandListImmediate::Get(), SizeInBytes, IndexStride);
+	return Allocate(FRHICommandListImmediate::Get(), NumIndices, IndexStride);
 }
 
 static TGlobalResource<TDynamicBufferPool<FDynamicIndexBuffer>, FRenderResource::EInitPhase::Pre> GDynamicIndexBufferPool;
 
-FGlobalDynamicIndexBuffer::FAllocation FGlobalDynamicIndexBuffer::Allocate(FRHICommandListBase& RHICmdList, uint32 SizeInBytes, uint32 IndexStride)
+FGlobalDynamicIndexBuffer::FAllocation FGlobalDynamicIndexBuffer::Allocate(FRHICommandListBase& RHICmdList, uint32 NumIndices, uint32 IndexStride)
 {
 	FAllocation Allocation;
 
@@ -908,6 +908,8 @@ FGlobalDynamicIndexBuffer::FAllocation FGlobalDynamicIndexBuffer::Allocate(FRHIC
 		return Allocation;
 	}
 
+	const uint32 SizeInBytes = NumIndices * IndexStride;
+
 	TArray<FDynamicIndexBuffer*>& IndexBuffers = (IndexStride == 2)
 		? IndexBuffers16
 		: IndexBuffers32;
