Removed ensure to check against render commands being enqueued from the rendering thread timeline, as this there are some valid cases in (shader compilation).

diff --git a/Engine/Source/Runtime/RenderCore/Private/RenderingThread.cpp b/Engine/Source/Runtime/RenderCore/Private/RenderingThread.cpp
--- a/Engine/Source/Runtime/RenderCore/Private/RenderingThread.cpp
+++ b/Engine/Source/Runtime/RenderCore/Private/RenderingThread.cpp
@@ -1554,9 +1554,6 @@ FRenderThreadCommandPipe FRenderThreadCommandPipe::Instance;
 
 void FRenderThreadCommandPipe::EnqueueAndLaunch(const TCHAR* Name, uint32& SpecId, TStatId StatId, TUniqueFunction<void(FRHICommandListImmediate&)>&& Function)
 {
-	ensureMsgf(!FTaskTagScope::IsCurrentTag(ETaskTag::EParallelRenderingThread) && !FTaskTagScope::IsCurrentTag(ETaskTag::ERenderingThread),
-		TEXT("Attempted to enqueue a render command from the rendering thread timeline which a pipeline violation. Render commands can only be enqueued from the game thread timeline."));
-
 	Mutex.Lock();
 	bool bWasEmpty = Queues[ProduceIndex].IsEmpty();
 	Queues[ProduceIndex].Emplace(Name, SpecId, StatId, MoveTemp(Function));
