Fix wrong printf parameters in checkf & ensuremsgf found in engine code

#rnx

diff --git a/Engine/Plugins/AI/MLAdapter/Source/MLAdapter/Private/Sensors/MLAdapterSensor_Attribute.cpp b/Engine/Plugins/AI/MLAdapter/Source/MLAdapter/Private/Sensors/MLAdapterSensor_Attribute.cpp
--- a/Engine/Plugins/AI/MLAdapter/Source/MLAdapter/Private/Sensors/MLAdapterSensor_Attribute.cpp
+++ b/Engine/Plugins/AI/MLAdapter/Source/MLAdapter/Private/Sensors/MLAdapterSensor_Attribute.cpp
@@ -117,7 +117,7 @@ void UMLAdapterSensor_Attribute::BindAttributes(AActor& Actor)
 	Attributes.Reset(AttributeNames.Num());
 	if (AttributeSetsFound.Num() > 0)
 	{
-		ensureMsgf(AttributeSetsFound.Num() == 1, TEXT("Found %s attribute sets, using the first one"), AttributeSetsFound.Num());
+		ensureMsgf(AttributeSetsFound.Num() == 1, TEXT("Found %d attribute sets, using the first one"), AttributeSetsFound.Num());
 		AttributeSet = AttributeSetsFound[0];
 		UClass* AttributeSetClass = AttributeSet->GetClass();
 		for (const FName& Name : AttributeNames)
diff --git a/Engine/Plugins/AI/MassAI/Source/MassAIBehavior/Private/Tasks/MassUseSmartObjectTask.cpp b/Engine/Plugins/AI/MassAI/Source/MassAIBehavior/Private/Tasks/MassUseSmartObjectTask.cpp
--- a/Engine/Plugins/AI/MassAI/Source/MassAIBehavior/Private/Tasks/MassUseSmartObjectTask.cpp
+++ b/Engine/Plugins/AI/MassAI/Source/MassAIBehavior/Private/Tasks/MassUseSmartObjectTask.cpp
@@ -141,7 +141,7 @@ EStateTreeRunStatus FMassUseSmartObjectTask::Tick(FStateTreeExecutionContext& Co
 		break;
 
 	default:
-		ensureMsgf(false, TEXT("Unhandled interaction status % => Returning EStateTreeRunStatus::Failed"), *UEnum::GetValueAsString(SOUser.InteractionStatus));
+		ensureMsgf(false, TEXT("Unhandled interaction status %s => Returning EStateTreeRunStatus::Failed"), *UEnum::GetValueAsString(SOUser.InteractionStatus));
 		Status = EStateTreeRunStatus::Failed;
 	}
 
diff --git a/Engine/Plugins/AI/MassAI/Source/MassZoneGraphNavigation/Private/MassZoneGraphNavigationUtils.cpp b/Engine/Plugins/AI/MassAI/Source/MassZoneGraphNavigation/Private/MassZoneGraphNavigationUtils.cpp
--- a/Engine/Plugins/AI/MassAI/Source/MassZoneGraphNavigation/Private/MassZoneGraphNavigationUtils.cpp
+++ b/Engine/Plugins/AI/MassAI/Source/MassZoneGraphNavigation/Private/MassZoneGraphNavigationUtils.cpp
@@ -29,7 +29,7 @@ namespace UE::MassNavigation
 		MoveTarget.DistanceToGoal = 0.0f;
 		MoveTarget.DesiredSpeed.Set(0.0f);
 
-		if (!ensureMsgf(MoveTarget.GetCurrentAction() == EMassMovementAction::Move, TEXT("Expecting action 'Move': Invalid action %s"), MoveTarget.GetCurrentAction()))
+		if (!ensureMsgf(MoveTarget.GetCurrentAction() == EMassMovementAction::Move, TEXT("Expecting action 'Move': Invalid action %u"), MoveTarget.GetCurrentAction()))
 		{
 			return false;
 		}
@@ -91,7 +91,7 @@ namespace UE::MassNavigation
 		MoveTarget.DistanceToGoal = 0.0f;
 		MoveTarget.DesiredSpeed.Set(0.0f);
 
-		if (!ensureMsgf(MoveTarget.GetCurrentAction() == EMassMovementAction::Stand, TEXT("Expecting action 'Stand': Invalid action %s"), MoveTarget.GetCurrentAction()))
+		if (!ensureMsgf(MoveTarget.GetCurrentAction() == EMassMovementAction::Stand, TEXT("Expecting action 'Stand': Invalid action %u"), MoveTarget.GetCurrentAction()))
 		{
 			return false;
 		}
@@ -115,7 +115,7 @@ namespace UE::MassNavigation
 		MoveTarget.DistanceToGoal = 0.0f;
 		MoveTarget.DesiredSpeed.Set(0.0f);
 
-		if (!ensureMsgf(MoveTarget.GetCurrentAction() == EMassMovementAction::Animate, TEXT("Expecting action 'Animate': Invalid action %s"), MoveTarget.GetCurrentAction()))
+		if (!ensureMsgf(MoveTarget.GetCurrentAction() == EMassMovementAction::Animate, TEXT("Expecting action 'Animate': Invalid action %u"), MoveTarget.GetCurrentAction()))
 		{
 			return false;
 		}
diff --git a/Engine/Plugins/Editor/EngineAssetDefinitions/Source/Private/AssetDefinition_World.cpp b/Engine/Plugins/Editor/EngineAssetDefinitions/Source/Private/AssetDefinition_World.cpp
--- a/Engine/Plugins/Editor/EngineAssetDefinitions/Source/Private/AssetDefinition_World.cpp
+++ b/Engine/Plugins/Editor/EngineAssetDefinitions/Source/Private/AssetDefinition_World.cpp
@@ -54,7 +54,7 @@ EAssetCommandResult UAssetDefinition_World::OpenAssets(const FAssetOpenArgs& Ope
 		UWorld* World = Cast<UWorld>(WorldAsset.GetAsset());
 		if (World != nullptr && 
 			ensureMsgf(World->GetPackage(), TEXT("World(%s) is not in a package and cannot be opened"), *World->GetFullName()) && 
-			ensureMsgf(!World->GetPackage()->HasAnyPackageFlags(PKG_NewlyCreated), TEXT("World(%s) is unsaved and cannot be opened")))
+			ensureMsgf(!World->GetPackage()->HasAnyPackageFlags(PKG_NewlyCreated), TEXT("World(%s) is unsaved and cannot be opened"), *World->GetFullName()))
 		{
 			const FString FileToOpen = FPackageName::LongPackageNameToFilename(World->GetOutermost()->GetName(), FPackageName::GetMapPackageExtension());
 			const bool bLoadAsTemplate = false;
diff --git a/Engine/Plugins/Experimental/ChaosVD/Source/ChaosVD/Private/ChaosVDWorldOutlinerMode.cpp b/Engine/Plugins/Experimental/ChaosVD/Source/ChaosVD/Private/ChaosVDWorldOutlinerMode.cpp
--- a/Engine/Plugins/Experimental/ChaosVD/Source/ChaosVD/Private/ChaosVDWorldOutlinerMode.cpp
+++ b/Engine/Plugins/Experimental/ChaosVD/Source/ChaosVD/Private/ChaosVDWorldOutlinerMode.cpp
@@ -53,7 +53,7 @@ void FChaosVDWorldOutlinerMode::OnItemSelectionChanged(FSceneOutlinerTreeItemPtr
 	if (const UTypedElementSelectionSet* SelectionSet = GetSelectionSetObject())
 	{
 		TArray<AActor*> SceneSelectedActors = SelectionSet->GetSelectedObjects<AActor>();
-		ensureMsgf(SceneSelectedActors.Num() < 2, TEXT("Multi Selection is not supported, but [%s] Actors are selected... Choosing the first one"));
+		ensureMsgf(SceneSelectedActors.Num() < 2, TEXT("Multi Selection is not supported, but [%d] Actors are selected... Choosing the first one"), SceneSelectedActors.Num());
 
 		AActor* SelectedActor = OutlinerSelectedActors.IsEmpty() ? nullptr : OutlinerSelectedActors[0];
 
diff --git a/Engine/Plugins/Experimental/NNERuntimeRDG/Source/NNERuntimeRDG/Private/NNERuntimeRDGModelHlsl.cpp b/Engine/Plugins/Experimental/NNERuntimeRDG/Source/NNERuntimeRDG/Private/NNERuntimeRDGModelHlsl.cpp
--- a/Engine/Plugins/Experimental/NNERuntimeRDG/Source/NNERuntimeRDG/Private/NNERuntimeRDGModelHlsl.cpp
+++ b/Engine/Plugins/Experimental/NNERuntimeRDG/Source/NNERuntimeRDG/Private/NNERuntimeRDGModelHlsl.cpp
@@ -276,7 +276,7 @@ int FModelInstance::PrepareTensorShapesAndData()
 	checkCode(
 		for (int i = 0; i < AllInitializedTensors.Num(); ++i)
 		{
-			checkf(AllInitializedTensors[i], TEXT("Tensor at index %d, was not initialized by model preparation."));
+			checkf(AllInitializedTensors[i], TEXT("Tensor at index %d, was not initialized by model preparation."), i);
 		};
 	);
 
diff --git a/Engine/Plugins/Experimental/PixelStreamingPlayer/Source/PixelStreamingPlayer/Public/PixelStreamingWebRTCWrappers.h b/Engine/Plugins/Experimental/PixelStreamingPlayer/Source/PixelStreamingPlayer/Public/PixelStreamingWebRTCWrappers.h
--- a/Engine/Plugins/Experimental/PixelStreamingPlayer/Source/PixelStreamingPlayer/Public/PixelStreamingWebRTCWrappers.h
+++ b/Engine/Plugins/Experimental/PixelStreamingPlayer/Source/PixelStreamingPlayer/Public/PixelStreamingWebRTCWrappers.h
@@ -55,7 +55,7 @@ struct PIXELSTREAMINGPLAYER_API FPixelStreamingIceCandidateWrapper
 	{
 		webrtc::SdpParseError Error;
 		TUniquePtr<webrtc::IceCandidateInterface> Candidate = TUniquePtr<webrtc::IceCandidateInterface>(webrtc::CreateIceCandidate(TCHAR_TO_UTF8(*Mid), MLineIndex, TCHAR_TO_UTF8(*SDP), &Error));
-		verifyf(Candidate, TEXT("Failed to create ICE candicate: %s"), Error.description.c_str());
+		verifyf(Candidate, TEXT("Failed to create ICE candicate: %hs"), Error.description.c_str());
 		return Candidate;
 	}
 
diff --git a/Engine/Plugins/Runtime/CommonUI/Source/CommonUI/Private/CommonActivatableWidget.cpp b/Engine/Plugins/Runtime/CommonUI/Source/CommonUI/Private/CommonActivatableWidget.cpp
--- a/Engine/Plugins/Runtime/CommonUI/Source/CommonUI/Private/CommonActivatableWidget.cpp
+++ b/Engine/Plugins/Runtime/CommonUI/Source/CommonUI/Private/CommonActivatableWidget.cpp
@@ -250,7 +250,7 @@ void UCommonActivatableWidget::ReleaseSlateResources(bool bReleaseChildren)
 
 void UCommonActivatableWidget::NativeOnActivated()
 {
-	if (ensureMsgf(bIsActive, TEXT("[%s] has called NativeOnActivated, but isn't actually activated! Never call this directly - call ActivateWidget()")))
+	if (ensureMsgf(bIsActive, TEXT("[%s] has called NativeOnActivated, but isn't actually activated! Never call this directly - call ActivateWidget()"), *GetName()))
 	{
 		if (bSetVisibilityOnActivated)
 		{
diff --git a/Engine/Plugins/Runtime/MassEntity/Source/MassEntity/Private/MassArchetypeData.cpp b/Engine/Plugins/Runtime/MassEntity/Source/MassEntity/Private/MassArchetypeData.cpp
--- a/Engine/Plugins/Runtime/MassEntity/Source/MassEntity/Private/MassArchetypeData.cpp
+++ b/Engine/Plugins/Runtime/MassEntity/Source/MassEntity/Private/MassArchetypeData.cpp
@@ -81,7 +81,7 @@ void FMassArchetypeData::Initialize(const FMassArchetypeCompositionDescriptor& I
 
 void FMassArchetypeData::InitializeWithSimilar(const FMassArchetypeData& BaseArchetype, FMassArchetypeCompositionDescriptor&& NewComposition, const uint32 ArchetypeDataVersion)
 {
-	checkf(IsInitialized() == false, TEXT("Trying to %s but this archetype has already been initialized"));
+	checkf(IsInitialized() == false, TEXT("Trying to InitializeWithSimilar but this archetype has already been initialized"));
 
 	CreatedArchetypeDataVersion = ArchetypeDataVersion;
 
diff --git a/Engine/Plugins/Runtime/MassEntity/Source/MassEntity/Private/MassEntityManager.cpp b/Engine/Plugins/Runtime/MassEntity/Source/MassEntity/Private/MassEntityManager.cpp
--- a/Engine/Plugins/Runtime/MassEntity/Source/MassEntity/Private/MassEntityManager.cpp
+++ b/Engine/Plugins/Runtime/MassEntity/Source/MassEntity/Private/MassEntityManager.cpp
@@ -1220,7 +1220,7 @@ void FMassEntityManager::CheckIfEntityIsValid(FMassEntityHandle Entity) const
 
 void FMassEntityManager::CheckIfEntityIsActive(FMassEntityHandle Entity) const
 {
-	checkf(IsEntityBuilt(Entity), TEXT("Entity not yet created(ID: %d, SN:%d)"));
+	checkf(IsEntityBuilt(Entity), TEXT("Entity not yet created(ID: %d, SN:%d)"), Entity.Index, Entity.SerialNumber);
 }
 
 void FMassEntityManager::GetMatchingArchetypes(const FMassFragmentRequirements& Requirements, TArray<FMassArchetypeHandle>& OutValidArchetypes, const uint32 FromArchetypeDataVersion) const
diff --git a/Engine/Plugins/Runtime/Metasound/Source/MetasoundGraphCore/Private/MetasoundOperatorSettings.cpp b/Engine/Plugins/Runtime/Metasound/Source/MetasoundGraphCore/Private/MetasoundOperatorSettings.cpp
--- a/Engine/Plugins/Runtime/Metasound/Source/MetasoundGraphCore/Private/MetasoundOperatorSettings.cpp
+++ b/Engine/Plugins/Runtime/Metasound/Source/MetasoundGraphCore/Private/MetasoundOperatorSettings.cpp
@@ -30,7 +30,7 @@ namespace Metasound
 	{
 		using namespace MetasoundOperatorSettingsPrivate;
 
-		if (ensureMsgf(InSampleRate > 0, TEXT("Invalid sample rate %d. Defaulting to %d"), InSampleRate, DefaultSampleRate))
+		if (ensureMsgf(InSampleRate > 0, TEXT("Invalid sample rate %d. Defaulting to %f"), InSampleRate, DefaultSampleRate))
 		{
 			SampleRate = InSampleRate;
 		}
diff --git a/Engine/Plugins/Runtime/ReplicationGraph/Source/Private/ReplicationGraph.cpp b/Engine/Plugins/Runtime/ReplicationGraph/Source/Private/ReplicationGraph.cpp
--- a/Engine/Plugins/Runtime/ReplicationGraph/Source/Private/ReplicationGraph.cpp
+++ b/Engine/Plugins/Runtime/ReplicationGraph/Source/Private/ReplicationGraph.cpp
@@ -3216,7 +3216,7 @@ void FStreamingLevelActorListCollection::AddActor(const FNewReplicatedActorInfo&
 
 	if (CVar_RepGraph_Verify)
 	{
-		ensureMsgf(Item->ReplicationActorList.Contains(ActorInfo.Actor) == false, TEXT("%s being added to %s twice! Streaming level: %s"), *GetActorRepListTypeDebugString(ActorInfo.Actor), *ActorInfo.StreamingLevelName.ToString() );
+		ensureMsgf(Item->ReplicationActorList.Contains(ActorInfo.Actor) == false, TEXT("%s being added to a list twice! Streaming level: %s"), *GetActorRepListTypeDebugString(ActorInfo.Actor), *ActorInfo.StreamingLevelName.ToString() );
 	}
 
 	Item->ReplicationActorList.Add(ActorInfo.Actor);
@@ -3237,7 +3237,7 @@ bool FStreamingLevelActorListCollection::RemoveActor(const FNewReplicatedActorIn
 
 			if (CVar_RepGraph_Verify)
 			{
-				ensureMsgf(StreamingList.ReplicationActorList.Contains(ActorInfo.Actor) == false, TEXT("Actor %s is still in %s after removal. Streaming Level: %s"), *GetActorRepListTypeDebugString(ActorInfo.Actor), *GetPathNameSafe(Outer));
+				ensureMsgf(StreamingList.ReplicationActorList.Contains(ActorInfo.Actor) == false, TEXT("Actor %s is still in list after removal. Streaming Level: %s"), *GetActorRepListTypeDebugString(ActorInfo.Actor), *GetPathNameSafe(Outer));
 			}
 			break;
 		}
@@ -3379,7 +3379,7 @@ void UReplicationGraphNode_ActorList::NotifyAddNetworkActor(const FNewReplicated
 	{
 		if (CVar_RepGraph_Verify)
 		{
-			ensureMsgf(ReplicationActorList.Contains(ActorInfo.Actor) == false, TEXT("%s being added to %s twice!"), *GetActorRepListTypeDebugString(ActorInfo.Actor) );
+			ensureMsgf(ReplicationActorList.Contains(ActorInfo.Actor) == false, TEXT("%s being added to %s twice!"), *GetActorRepListTypeDebugString(ActorInfo.Actor), *GetPathName());
 		}
 
 		ReplicationActorList.Add(ActorInfo.Actor);
@@ -3571,7 +3571,7 @@ void UReplicationGraphNode_ActorListFrequencyBuckets::NotifyAddNetworkActor(cons
 
 			if (CVar_RepGraph_Verify)
 			{
-				ensureMsgf(List.Contains(ActorInfo.Actor) == false, TEXT("%s being added to %s twice!"), *GetActorRepListTypeDebugString(ActorInfo.Actor) );
+				ensureMsgf(List.Contains(ActorInfo.Actor) == false, TEXT("%s being added to %s twice!"), *GetActorRepListTypeDebugString(ActorInfo.Actor), *GetPathName());
 			}
 		}
 
diff --git a/Engine/Plugins/Runtime/RigVM/Source/RigVMDeveloper/Private/RigVMModel/RigVMController.cpp b/Engine/Plugins/Runtime/RigVM/Source/RigVMDeveloper/Private/RigVMModel/RigVMController.cpp
--- a/Engine/Plugins/Runtime/RigVM/Source/RigVMDeveloper/Private/RigVMModel/RigVMController.cpp
+++ b/Engine/Plugins/Runtime/RigVM/Source/RigVMDeveloper/Private/RigVMModel/RigVMController.cpp
@@ -6027,7 +6027,7 @@ URigVMFunctionReferenceNode* URigVMController::PromoteCollapseNodeToFunctionRefe
 	// Create Function
 	URigVMLibraryNode* FunctionDefinition = nullptr;
 	if (!InExistingFunctionDefinitionPath.IsEmpty() && 
-		ensureAlwaysMsgf(!FPackageName::IsShortPackageName(InExistingFunctionDefinitionPath), TEXT("Expected full path name for function definition path: \"%s\"), *InExistingFunctionDefinitionPath")))
+		ensureAlwaysMsgf(!FPackageName::IsShortPackageName(InExistingFunctionDefinitionPath), TEXT("Expected full path name for function definition path: \"%s\")"), *InExistingFunctionDefinitionPath))
 	{
 		FunctionDefinition = FindObject<URigVMLibraryNode>(nullptr, *InExistingFunctionDefinitionPath);
 	}
diff --git a/Engine/Plugins/Runtime/StateTree/Source/StateTreeModule/Private/Debugger/StateTreeTraceAnalyzer.cpp b/Engine/Plugins/Runtime/StateTree/Source/StateTreeModule/Private/Debugger/StateTreeTraceAnalyzer.cpp
--- a/Engine/Plugins/Runtime/StateTree/Source/StateTreeModule/Private/Debugger/StateTreeTraceAnalyzer.cpp
+++ b/Engine/Plugins/Runtime/StateTree/Source/StateTreeModule/Private/Debugger/StateTreeTraceAnalyzer.cpp
@@ -201,7 +201,7 @@ bool FStateTreeTraceAnalyzer::OnEvent(const uint16 RouteId, EStyle Style, const
 			break;
 		}
 	default:
-		ensureMsgf(false, TEXT("Unhandle route id: %s"), RouteId);
+		ensureMsgf(false, TEXT("Unhandle route id: %u"), RouteId);
 	}
 
 	return true;
diff --git a/Engine/Plugins/VirtualProduction/DMX/DMXEngine/Source/DMXEditor/Private/DMXEditorStyle.cpp b/Engine/Plugins/VirtualProduction/DMX/DMXEngine/Source/DMXEditor/Private/DMXEditorStyle.cpp
--- a/Engine/Plugins/VirtualProduction/DMX/DMXEngine/Source/DMXEditor/Private/DMXEditorStyle.cpp
+++ b/Engine/Plugins/VirtualProduction/DMX/DMXEngine/Source/DMXEditor/Private/DMXEditorStyle.cpp
@@ -203,7 +203,7 @@ void FDMXEditorStyle::Initialize()
 void FDMXEditorStyle::Shutdown()
 {
 	// DEPRECATED 5.0
-	ensureMsgf(StyleInstance_DEPRECATED.IsValid(), TEXT("%S called, but StyleInstance wasn't initialized"));
+	ensureMsgf(StyleInstance_DEPRECATED.IsValid(), TEXT("FDMXEditorStyle::Shutdown called, but StyleInstance wasn't initialized"));
 	FSlateStyleRegistry::UnRegisterSlateStyle(*StyleInstance_DEPRECATED);
 	ensure(StyleInstance_DEPRECATED.IsUnique());
 	StyleInstance_DEPRECATED.Reset();
diff --git a/Engine/Plugins/VirtualProduction/DMX/DMXEngine/Source/DMXEditor/Private/Factories/DMXLibraryFromMVRFactory.cpp b/Engine/Plugins/VirtualProduction/DMX/DMXEngine/Source/DMXEditor/Private/Factories/DMXLibraryFromMVRFactory.cpp
--- a/Engine/Plugins/VirtualProduction/DMX/DMXEngine/Source/DMXEditor/Private/Factories/DMXLibraryFromMVRFactory.cpp
+++ b/Engine/Plugins/VirtualProduction/DMX/DMXEngine/Source/DMXEditor/Private/Factories/DMXLibraryFromMVRFactory.cpp
@@ -97,7 +97,7 @@ bool UDMXLibraryFromMVRFactory::CanReimport(UObject* Obj, TArray<FString>& OutFi
 void UDMXLibraryFromMVRFactory::SetReimportPaths(UObject* Obj, const TArray<FString>& NewReimportPaths)
 {
 	UDMXMVRAssetImportData* MVRAssetImportData = GetMVRAssetImportData(Obj);
-	if (!ensureMsgf(MVRAssetImportData, TEXT("Invalid MVR Asset Import Data for General Scene Description %s.")))
+	if (!ensureMsgf(MVRAssetImportData, TEXT("Invalid MVR Asset Import Data for General Scene Description.")))
 	{
 		return;
 	}
diff --git a/Engine/Plugins/VirtualProduction/DMX/DMXEngine/Source/DMXEditor/Private/Factories/DMXLibraryFromMVRImportOptions.cpp b/Engine/Plugins/VirtualProduction/DMX/DMXEngine/Source/DMXEditor/Private/Factories/DMXLibraryFromMVRImportOptions.cpp
--- a/Engine/Plugins/VirtualProduction/DMX/DMXEngine/Source/DMXEditor/Private/Factories/DMXLibraryFromMVRImportOptions.cpp
+++ b/Engine/Plugins/VirtualProduction/DMX/DMXEngine/Source/DMXEditor/Private/Factories/DMXLibraryFromMVRImportOptions.cpp
@@ -55,7 +55,7 @@ void UDMXLibraryFromMVRImportOptions::ApplyOptions(UDMXLibrary* DMXLibrary)
 				{
 					return InputPortConfig.GetPortGuid() == InputPort->GetPortGuid();
 				});
-			if (ensureAlwaysMsgf(InputPortConfigPtr, TEXT("Cannot find config of Input Port '%' in Project Settings."), *InputPort->GetPortName()))
+			if (ensureAlwaysMsgf(InputPortConfigPtr, TEXT("Cannot find config of Input Port '%s' in Project Settings."), *InputPort->GetPortName()))
 			{
 				FDMXInputPortConfig& PortConfig = *InputPortConfigPtr;
 
@@ -92,7 +92,7 @@ void UDMXLibraryFromMVRImportOptions::ApplyOptions(UDMXLibrary* DMXLibrary)
 				{
 					return OutputPortConfig.GetPortGuid() == OutputPort->GetPortGuid();
 				});
-			if (ensureAlwaysMsgf(OutputPortConfigPtr, TEXT("Cannot find config of Output Port '%' in Project Settings."), *OutputPort->GetPortName()))
+			if (ensureAlwaysMsgf(OutputPortConfigPtr, TEXT("Cannot find config of Output Port '%s' in Project Settings."), *OutputPort->GetPortName()))
 			{
 				FDMXOutputPortConfig& PortConfig = *OutputPortConfigPtr;
 
diff --git a/Engine/Plugins/VirtualProduction/DMX/DMXEngine/Source/DMXEditor/Private/Factories/DMXLibraryFromMVRImporter.cpp b/Engine/Plugins/VirtualProduction/DMX/DMXEngine/Source/DMXEditor/Private/Factories/DMXLibraryFromMVRImporter.cpp
--- a/Engine/Plugins/VirtualProduction/DMX/DMXEngine/Source/DMXEditor/Private/Factories/DMXLibraryFromMVRImporter.cpp
+++ b/Engine/Plugins/VirtualProduction/DMX/DMXEngine/Source/DMXEditor/Private/Factories/DMXLibraryFromMVRImporter.cpp
@@ -112,7 +112,7 @@ namespace UE::DMX::DMXLibraryFromMVRImporter::Private
 bool FDMXLibraryFromMVRImporter::LoadMVRFile(const FString& InFilename)
 {
 	Filename = InFilename;
-	if (!ensureAlwaysMsgf(FPaths::FileExists(InFilename), TEXT("Cannot import MVR File '%s'. File does not exist.")))
+	if (!ensureAlwaysMsgf(FPaths::FileExists(InFilename), TEXT("Cannot import MVR File '%s'. File does not exist."), *InFilename))
 	{
 		return false;
 	}
diff --git a/Engine/Source/Developer/DerivedDataCache/Private/DerivedDataBuildFunctionRegistry.cpp b/Engine/Source/Developer/DerivedDataCache/Private/DerivedDataBuildFunctionRegistry.cpp
--- a/Engine/Source/Developer/DerivedDataCache/Private/DerivedDataBuildFunctionRegistry.cpp
+++ b/Engine/Source/Developer/DerivedDataCache/Private/DerivedDataBuildFunctionRegistry.cpp
@@ -148,7 +148,7 @@ bool IsValidBuildFunctionName(FUtf8StringView Function)
 void AssertValidBuildFunctionName(FUtf8StringView Function, FStringView Name)
 {
 	checkf(IsValidBuildFunctionName(Function),
-		TEXT("A build function name must be alphanumeric and non-empty for build of '%.*s' by '%.*s'."),
+		TEXT("A build function name must be alphanumeric and non-empty for build of '%.*s' by '%s'."),
 		Name.Len(), Name.GetData(), *WriteToString<32>(Function));
 }
 
diff --git a/Engine/Source/Developer/TextureFormat/Public/ChildTextureFormat.h b/Engine/Source/Developer/TextureFormat/Public/ChildTextureFormat.h
--- a/Engine/Source/Developer/TextureFormat/Public/ChildTextureFormat.h
+++ b/Engine/Source/Developer/TextureFormat/Public/ChildTextureFormat.h
@@ -66,7 +66,7 @@ public:
 		ITextureFormatManagerModule& TFM = GetTextureFormatManagerRef();
 		const ITextureFormat* FormatObject = TFM.FindTextureFormat(BaseFormatName);
 
-		checkf(FormatObject != nullptr, TEXT("Bad FormatName %s passed to FChildTextureFormat::GetBaseFormatObject()"));
+		checkf(FormatObject != nullptr, TEXT("Bad FormatName %s passed to FChildTextureFormat::GetBaseFormatObject()"), *BaseFormatName.ToString());
 
 		return FormatObject;
 	}
diff --git a/Engine/Source/Developer/Virtualization/Private/VirtualizationManager.cpp b/Engine/Source/Developer/Virtualization/Private/VirtualizationManager.cpp
--- a/Engine/Source/Developer/Virtualization/Private/VirtualizationManager.cpp
+++ b/Engine/Source/Developer/Virtualization/Private/VirtualizationManager.cpp
@@ -803,7 +803,7 @@ bool FVirtualizationManager::PullData(TArrayView<FPullRequest> Requests)
 
 		if (Request.IsSuccess())
 		{
-			checkf(Request.GetIdentifier() == Request.GetPayload().GetRawHash(), TEXT("[%s] Invalid payload for '%s'"), *LexToString(Request.GetIdentifier()));
+			checkf(Request.GetIdentifier() == Request.GetPayload().GetRawHash(), TEXT("Invalid payload for '%s'"), *LexToString(Request.GetIdentifier()));
 		}
 	}	
 
diff --git a/Engine/Source/Developer/Virtualization/Private/VirtualizationSourceControlBackend.cpp b/Engine/Source/Developer/Virtualization/Private/VirtualizationSourceControlBackend.cpp
--- a/Engine/Source/Developer/Virtualization/Private/VirtualizationSourceControlBackend.cpp
+++ b/Engine/Source/Developer/Virtualization/Private/VirtualizationSourceControlBackend.cpp
@@ -78,7 +78,7 @@ public:
 
 	~FSemaphore()
 	{
-		checkf(DebugCount == 0, TEXT("'%d' threads are still waiting on the UE::Virtualization::FSemaphore being destroyed"));
+		checkf(DebugCount == 0, TEXT("'%d' threads are still waiting on the UE::Virtualization::FSemaphore being destroyed"), DebugCount.load());
 	}
 
 	/** Will block until the calling thread can pass through the semaphore. Note that it might return an error if the WaitEvent fails */
diff --git a/Engine/Source/Editor/LevelEditor/Private/SLevelEditor.cpp b/Engine/Source/Editor/LevelEditor/Private/SLevelEditor.cpp
--- a/Engine/Source/Editor/LevelEditor/Private/SLevelEditor.cpp
+++ b/Engine/Source/Editor/LevelEditor/Private/SLevelEditor.cpp
@@ -1680,7 +1680,7 @@ TSharedRef<SWidget> SLevelEditor::RestoreContentArea( const TSharedRef<SDockTab>
 		// Warn user/developer
 		const FString WarningMessage = FString::Format(TEXT("Level editor layout could not be loaded from the config file {0}, trying to reset this config file to the default one."), { *GEditorLayoutIni });
 		UE_LOG(LogTemp, Warning, TEXT("%s"), *WarningMessage);
-		ensureMsgf(false, TEXT("%s Some additional testing of that layout file should be done."));
+		ensureMsgf(false, TEXT("Some additional testing of that layout file should be done."));
 	}
 	check(ContentAreaWidget.IsValid());
 	return ContentAreaWidget.ToSharedRef();
diff --git a/Engine/Source/Editor/UnrealEd/Private/PlayLevel.cpp b/Engine/Source/Editor/UnrealEd/Private/PlayLevel.cpp
--- a/Engine/Source/Editor/UnrealEd/Private/PlayLevel.cpp
+++ b/Engine/Source/Editor/UnrealEd/Private/PlayLevel.cpp
@@ -211,7 +211,7 @@ private:
 			FMessageLog(NAME_CategoryPIE).SuppressLoggingToOutputLog(true).Error(Line.Message);
 			break;
 		case ELogVerbosity::Fatal:
-			checkf(false, *Line.Message.ToString());
+			checkf(false, TEXT("%s"), *Line.Message.ToString());
 			break;
 		}
 	}
@@ -513,7 +513,7 @@ void UEditorEngine::EndPlayMap()
 		// We cannot safely recover from this.
 		if (UObjectBaseUtility::IsPendingKillEnabled())
 		{
-			checkf(false, *FText::Format(
+			checkf(false, TEXT("%s"), *FText::Format(
 				LOCTEXT("PIEObjectStillReferenced", "Object '{Object}' from PIE level still referenced. Shortest path from root: {Path}"), Arguments).ToString());
 		}
 		else
diff --git a/Engine/Source/Runtime/AIModule/Public/IndexedHandle.h b/Engine/Source/Runtime/AIModule/Public/IndexedHandle.h
--- a/Engine/Source/Runtime/AIModule/Public/IndexedHandle.h
+++ b/Engine/Source/Runtime/AIModule/Public/IndexedHandle.h
@@ -253,7 +253,7 @@ public:
 	bool IsValidHandle(TIndexedHandle IndexedHandle) const
 	{
 #if UE_DO_INDEXED_HANDLE_MANAGER_ID
-		if (!ensureMsgf(ManagerID == IndexedHandle.ManagerID, TEXT("ManagerID %d does not match IndexedHandle.ManagerID %d, handles must only be used by the same manager they were created by!"), IndexedHandle.ManagerID))
+		if (!ensureMsgf(ManagerID == IndexedHandle.ManagerID, TEXT("ManagerID %d does not match IndexedHandle.ManagerID %d, handles must only be used by the same manager they were created by!"), ManagerID, IndexedHandle.ManagerID))
 		{
 			return false;
 		}
diff --git a/Engine/Source/Runtime/Core/Private/UObject/UnrealNames.cpp b/Engine/Source/Runtime/Core/Private/UObject/UnrealNames.cpp
--- a/Engine/Source/Runtime/Core/Private/UObject/UnrealNames.cpp
+++ b/Engine/Source/Runtime/Core/Private/UObject/UnrealNames.cpp
@@ -3017,7 +3017,7 @@ private:
 			}
 
 			checkf(false, TEXT("FName's %d max length exceeded. Got %d characters excluding null-terminator:\n%.*s"), 
-				NAME_SIZE - 1, View.Len, NAME_SIZE, View.IsAnsi() ? ANSI_TO_TCHAR(View.Data) : View.Data);
+				NAME_SIZE - 1, View.Len, NAME_SIZE, View.IsAnsi() ? ANSI_TO_TCHAR(View.Ansi) : View.Wide)
 
 			const ANSICHAR* ErrorString = "ERROR_NAME_SIZE_EXCEEDED";
 			return FindOrStoreString(FNameStringView(ErrorString, FCStringAnsi::Strlen(ErrorString), false), FNAME_Add);
diff --git a/Engine/Source/Runtime/CoreUObject/Private/Serialization/AsyncLoading2.cpp b/Engine/Source/Runtime/CoreUObject/Private/Serialization/AsyncLoading2.cpp
--- a/Engine/Source/Runtime/CoreUObject/Private/Serialization/AsyncLoading2.cpp
+++ b/Engine/Source/Runtime/CoreUObject/Private/Serialization/AsyncLoading2.cpp
@@ -1287,10 +1287,11 @@ public:
 			{
 				UObject* Object = static_cast<UObject*>(GUObjectArray.IndexToObject(ObjectIndex)->Object);
 				ensureMsgf(!Object->HasAnyInternalFlags(EInternalObjectFlags::LoaderImport) || GUObjectArray.IsDisregardForGC(Object),
-						TEXT("FGlobalImportStore::VerifyPackageForRemoval: The loaded public export object '%s' with flags (ObjectFlags=%x, InternalObjectFlags=%x) and id 0x%llX:0x%llX is probably still referenced by the loader."),
+						TEXT("FGlobalImportStore::VerifyPackageForRemoval: The loaded public export object '%s' with flags (ObjectFlags=%x, InternalObjectFlags=%x) and id 0x%llX is probably still referenced by the loader."),
 						*Object->GetFullName(),
 						Object->GetFlags(),
-						Object->GetInternalFlags());
+						Object->GetInternalFlags(),
+						PackageId.ValueForDebugging());
 
 				FPublicExportKey* PublicExportKey = ObjectIndexToPublicExport.Find(ObjectIndex);
 				UE_CLOG(!PublicExportKey, LogStreaming, Fatal,
diff --git a/Engine/Source/Runtime/CoreUObject/Private/UObject/SavePackage2.cpp b/Engine/Source/Runtime/CoreUObject/Private/UObject/SavePackage2.cpp
--- a/Engine/Source/Runtime/CoreUObject/Private/UObject/SavePackage2.cpp
+++ b/Engine/Source/Runtime/CoreUObject/Private/UObject/SavePackage2.cpp
@@ -2499,7 +2499,7 @@ ESavePackageResult FinalizeFile(FStructuredArchive::FRecord& StructuredArchiveRo
 		}
 
 		// Move file to its real destination
-		checkf(SaveContext.GetTempFilename().IsSet(), TEXT("The package should've been saved to a tmp file first! (%s)"), *SaveContext.GetFilename());
+		checkf(SaveContext.GetTempFilename().IsSet(), TEXT("The package should've been saved to a tmp file first! (%s)"), SaveContext.GetFilename());
 
 		// When saving in text format we will have two temp files, so we need to manually delete the non-textbased one
 		if (SaveContext.IsTextFormat())
diff --git a/Engine/Source/Runtime/Engine/Private/Animation/AnimSequenceBase.cpp b/Engine/Source/Runtime/Engine/Private/Animation/AnimSequenceBase.cpp
--- a/Engine/Source/Runtime/Engine/Private/Animation/AnimSequenceBase.cpp
+++ b/Engine/Source/Runtime/Engine/Private/Animation/AnimSequenceBase.cpp
@@ -299,7 +299,7 @@ void UAnimSequenceBase::PostDuplicate(EDuplicateMode::Type DuplicateMode)
 	Super::PostDuplicate(DuplicateMode);
 
 #if WITH_EDITOR
-	checkf(DataModelInterface.GetObject()->GetOuter() == this, TEXT("Animation Data Model interface has incorrect outer, expected %s - found %f"), *this->GetName(), *DataModelInterface.GetObject()->GetOuter()->GetName());
+	checkf(DataModelInterface.GetObject()->GetOuter() == this, TEXT("Animation Data Model interface has incorrect outer, expected %s - found %s"), *this->GetName(), *DataModelInterface.GetObject()->GetOuter()->GetName());
 	BindToModelModificationEvent();
 #endif // WITH_EDITOR
 }
diff --git a/Engine/Source/Runtime/Engine/Private/PackageMapClient.cpp b/Engine/Source/Runtime/Engine/Private/PackageMapClient.cpp
--- a/Engine/Source/Runtime/Engine/Private/PackageMapClient.cpp
+++ b/Engine/Source/Runtime/Engine/Private/PackageMapClient.cpp
@@ -1238,7 +1238,7 @@ bool UPackageMapClient::ExportNetGUIDForReplay(FNetworkGUID& NetGUID, UObject* O
 		GuidCache->SetNetworkChecksumMode(RestoreMode);
 
 		check(!Writer.IsError());
-		ensureMsgf(GUIDMemory.Num() <= MaxReservedSize, TEXT("ExportNetGUIDForReplay exceeded CVarReservedNetGuidSize. Max=%l Count=%l"), MaxReservedSize, GUIDMemory.Num());
+		ensureMsgf(GUIDMemory.Num() <= MaxReservedSize, TEXT("ExportNetGUIDForReplay exceeded CVarReservedNetGuidSize. Max=%d Count=%d"), MaxReservedSize, GUIDMemory.Num());
 
 		GUIDMemory.Shrink();
 
diff --git a/Engine/Source/Runtime/Engine/Private/Particles/ParticleModules_Location.cpp b/Engine/Source/Runtime/Engine/Private/Particles/ParticleModules_Location.cpp
--- a/Engine/Source/Runtime/Engine/Private/Particles/ParticleModules_Location.cpp
+++ b/Engine/Source/Runtime/Engine/Private/Particles/ParticleModules_Location.cpp
@@ -1998,7 +1998,7 @@ void UParticleModuleLocationBoneSocket::GetSkeletalMeshComponentSource(FParticle
 
 bool UParticleModuleLocationBoneSocket::GetSocketInfoForSourceIndex(FModuleLocationBoneSocketInstancePayload* InstancePayload, USkeletalMeshComponent* SourceComponent, int32 SourceIndex, USkeletalMeshSocket*& OutSocket, FVector& OutOffset) const
 {
-	if (!ensureMsgf(SourceType == BONESOCKETSOURCE_Sockets, TEXT("Invalid source type %d for %s"), SourceType, *GetPathName()) ||
+	if (!ensureMsgf(SourceType == BONESOCKETSOURCE_Sockets, TEXT("Invalid source type %d for %s"), SourceType.GetIntValue(), *GetPathName()) ||
 		!ensureMsgf(SourceComponent, TEXT("Null SkeletalMeshComponent for %s"), *GetPathName()) ||
 		!ensureMsgf(SourceComponent->GetSkeletalMeshAsset(), TEXT("Null SkeletalMesh on Component %s for %s"), *SourceComponent->GetPathName(), *GetPathName()))
 	{
diff --git a/Engine/Source/Runtime/Engine/Private/Slate/SlateGameResources.cpp b/Engine/Source/Runtime/Engine/Private/Slate/SlateGameResources.cpp
--- a/Engine/Source/Runtime/Engine/Private/Slate/SlateGameResources.cpp
+++ b/Engine/Source/Runtime/Engine/Private/Slate/SlateGameResources.cpp
@@ -51,7 +51,7 @@ void FSlateGameResources::SetContentRoot( const FString& InContentRootDir )
 
 const FSlateBrush* FSlateGameResources::GetBrush( const FName PropertyName, const ANSICHAR* Specifier, const ISlateStyle* RequestingStyle ) const
 {
-	ensureMsgf(Specifier == NULL, TEXT("Attempting to look up resource (%s, %s). \n Specifiers not supported by Slate Resource Sets loaded from content browser."), *PropertyName.ToString(), Specifier);
+	ensureMsgf(Specifier == NULL, TEXT("Attempting to look up resource (%s, %s). \n Specifiers not supported by Slate Resource Sets loaded from content browser."), *PropertyName.ToString(), ANSI_TO_TCHAR(Specifier));
 	FName CleanName = GetCleanName(PropertyName);
 	auto* Resource = UIResources.Find(CleanName);
 	if(Resource)
@@ -65,7 +65,7 @@ const FSlateBrush* FSlateGameResources::GetBrush( const FName PropertyName, cons
 
 const FSlateBrush* FSlateGameResources::GetOptionalBrush(const FName PropertyName, const ANSICHAR* Specifier, const FSlateBrush* const InDefaultBrush) const
 {
-	ensureMsgf(Specifier == NULL, TEXT("Attempting to look up resource (%s, %s). \n Specifiers not supported by Slate Resource Sets loaded from content browser."), *PropertyName.ToString(), Specifier);
+	ensureMsgf(Specifier == NULL, TEXT("Attempting to look up resource (%s, %s). \n Specifiers not supported by Slate Resource Sets loaded from content browser."), *PropertyName.ToString(), ANSI_TO_TCHAR(Specifier));
 	FName CleanName = GetCleanName(PropertyName);
 	auto* Resource = UIResources.Find(CleanName);
 	if(Resource)
diff --git a/Engine/Source/Runtime/Engine/Public/MaterialShared.h b/Engine/Source/Runtime/Engine/Public/MaterialShared.h
--- a/Engine/Source/Runtime/Engine/Public/MaterialShared.h
+++ b/Engine/Source/Runtime/Engine/Public/MaterialShared.h
@@ -1426,7 +1426,7 @@ public:
 #endif
 		const bool bValid = GetFrozenContentSize() > 0u;
 
-		checkf(bValid || !bFailOnInvalid, TEXT("FMaterialShaderMap %s invalid for rendering: bCompilationFinalized: %i, bCompiledSuccessfully: %i, bDeletedThroughDeferredCleanup: %i, FrozenContentSize: %d"), *GetFriendlyName(),
+		checkf(bValid || !bFailOnInvalid, TEXT("FMaterialShaderMap %s invalid for rendering: bCompilationFinalized: %i, bCompiledSuccessfully: %i, bDeletedThroughDeferredCleanup: %i, FrozenContentSize: %d"), GetFriendlyName(),
 			bCompilationFinalized, bCompiledSuccessfully, bDeletedThroughDeferredCleanup ? 1 : 0, GetFrozenContentSize());
 		return bValid;
 	}
diff --git a/Engine/Source/Runtime/Experimental/GeometryCollectionEngine/Private/GeometryCollection/GeometryCollectionObject.cpp b/Engine/Source/Runtime/Experimental/GeometryCollectionEngine/Private/GeometryCollection/GeometryCollectionObject.cpp
--- a/Engine/Source/Runtime/Experimental/GeometryCollectionEngine/Private/GeometryCollection/GeometryCollectionObject.cpp
+++ b/Engine/Source/Runtime/Experimental/GeometryCollectionEngine/Private/GeometryCollection/GeometryCollectionObject.cpp
@@ -710,7 +710,7 @@ bool UGeometryCollection::RemoveLastMaterialSlot()
 /** Returns true if there is anything to render */
 bool UGeometryCollection::HasVisibleGeometry() const
 {
-	if(ensureMsgf(GeometryCollection.IsValid(), TEXT("Geometry Collection %s has an invalid internal collection")))
+	if(ensureMsgf(GeometryCollection.IsValid(), TEXT("Geometry Collection has an invalid internal collection")))
 	{
 		return ( (EnableNanite && RenderData && RenderData->bHasNaniteData) || GeometryCollection->HasVisibleGeometry());
 	}
diff --git a/Engine/Source/Runtime/Experimental/Iris/Core/Private/Iris/ReplicationSystem/Filtering/NetObjectGroups.cpp b/Engine/Source/Runtime/Experimental/Iris/Core/Private/Iris/ReplicationSystem/Filtering/NetObjectGroups.cpp
--- a/Engine/Source/Runtime/Experimental/Iris/Core/Private/Iris/ReplicationSystem/Filtering/NetObjectGroups.cpp
+++ b/Engine/Source/Runtime/Experimental/Iris/Core/Private/Iris/ReplicationSystem/Filtering/NetObjectGroups.cpp
@@ -224,7 +224,7 @@ void FNetObjectGroups::AddToGroup(FNetObjectGroupHandle GroupHandle, FInternalNe
 		}
 		else
 		{
-			checkf(false, TEXT("FNetObjectGroups::AddToGroup, Failed to add ( InternalIndex: %u ) to (GroupIndex: %u) A NetObject can only be a member of %u groups."), InternalIndex, GroupHandle, FNetObjectGroupMembership::MaxAssignedGroupCount);
+			checkf(false, TEXT("FNetObjectGroups::AddToGroup, Failed to add ( InternalIndex: %u ) to (GroupIndex: %u) A NetObject can only be a member of %u groups."), InternalIndex, GroupHandle.GetRawValue(), FNetObjectGroupMembership::MaxAssignedGroupCount);
 		}
 	}
 }
diff --git a/Engine/Source/Runtime/Experimental/Iris/Core/Private/Iris/ReplicationSystem/ReplicationReader.cpp b/Engine/Source/Runtime/Experimental/Iris/Core/Private/Iris/ReplicationSystem/ReplicationReader.cpp
--- a/Engine/Source/Runtime/Experimental/Iris/Core/Private/Iris/ReplicationSystem/ReplicationReader.cpp
+++ b/Engine/Source/Runtime/Experimental/Iris/Core/Private/Iris/ReplicationSystem/ReplicationReader.cpp
@@ -1650,7 +1650,7 @@ void FReplicationReader::ReadObjects(FNetSerializationContext& Context, uint32 O
 		--ObjectBatchCountToRead;
 	}
 
-	ensureAlwaysMsgf(!Context.HasErrorOrOverflow(), TEXT("Overflow: %c Error: %s Bit stream bits left: %u position: %u"), "YN"[Context.HasError()], ToCStr(Context.GetError().ToString()), Reader.GetBitsLeft(), Reader.GetPosBits());
+	ensureAlwaysMsgf(!Context.HasErrorOrOverflow(), TEXT("Overflow: %c Error: %s Bit stream bits left: %u position: %u"), TEXT("YN")[Context.HasError()], ToCStr(Context.GetError().ToString()), Reader.GetBitsLeft(), Reader.GetPosBits());
 }
 
 void FReplicationReader::ProcessHugeObjectAttachment(FNetSerializationContext& Context, const TRefCountPtr<FNetBlob>& Attachment)
diff --git a/Engine/Source/Runtime/RenderCore/Private/RenderGraphBuilder.cpp b/Engine/Source/Runtime/RenderCore/Private/RenderGraphBuilder.cpp
--- a/Engine/Source/Runtime/RenderCore/Private/RenderGraphBuilder.cpp
+++ b/Engine/Source/Runtime/RenderCore/Private/RenderGraphBuilder.cpp
@@ -3585,7 +3585,7 @@ void FRDGBuilder::BeginResourceRHI(FRDGPassHandle PassHandle, FRDGTextureRef Tex
 
 		// Cannot begin a resource within a merged render pass region.
 		checkf(GetPrologueBarrierPassHandle(PassHandle) == PassHandle,
-			TEXT("Cannot begin a resource within a merged render pass. Pass (Handle: %d, Name: %s), Resource %s"), PassHandle, Pass->GetName(), Texture->Name);
+			TEXT("Cannot begin a resource within a merged render pass. Pass (Handle: %d, Name: %s), Resource %s"), PassHandle.GetIndex(), Pass->GetName(), Texture->Name);
 	}
 #endif
 
@@ -3675,7 +3675,7 @@ void FRDGBuilder::BeginResourceRHI(FRDGPassHandle PassHandle, FRDGBufferRef Buff
 
 		// Cannot begin a resource within a merged render pass region.
 		checkf(GetPrologueBarrierPassHandle(PassHandle) == PassHandle,
-			TEXT("Cannot begin a resource within a merged render pass. Pass (Handle: %d, Name: %s), Resource %s"), PassHandle, Pass->GetName(), Buffer->Name);
+			TEXT("Cannot begin a resource within a merged render pass. Pass (Handle: %d, Name: %s), Resource %s"), PassHandle.GetIndex(), Pass->GetName(), Buffer->Name);
 	}
 #endif
 	Buffer->FinalizeDesc();
diff --git a/Engine/Source/Runtime/RenderCore/Private/RenderGraphValidation.cpp b/Engine/Source/Runtime/RenderCore/Private/RenderGraphValidation.cpp
--- a/Engine/Source/Runtime/RenderCore/Private/RenderGraphValidation.cpp
+++ b/Engine/Source/Runtime/RenderCore/Private/RenderGraphValidation.cpp
@@ -748,10 +748,10 @@ void FRDGUserValidation::ValidateAddPass(const void* ParameterStruct, const FSha
 		TEXT("Pass %s specified both Raster and AsyncCompute. They are mutually exclusive."), Name.GetTCHAR());
 
 	checkf(!EnumHasAllFlags(Flags, ERDGPassFlags::SkipRenderPass) || EnumHasAllFlags(Flags, ERDGPassFlags::Raster),
-		TEXT("Pass %s specified SkipRenderPass without Raster. Only raster passes support this flag."));
+		TEXT("Pass %s specified SkipRenderPass without Raster. Only raster passes support this flag."), Name.GetTCHAR());
 
 	checkf(!EnumHasAllFlags(Flags, ERDGPassFlags::NeverMerge) || EnumHasAllFlags(Flags, ERDGPassFlags::Raster),
-		TEXT("Pass %s specified NeverMerge without Raster. Only raster passes support this flag."));
+		TEXT("Pass %s specified NeverMerge without Raster. Only raster passes support this flag."), Name.GetTCHAR());
 }
 
 void FRDGUserValidation::ValidateAddPass(const FRDGPass* Pass, bool bSkipPassAccessMarking)
diff --git a/Engine/Source/Runtime/RenderCore/Private/ShaderCore.cpp b/Engine/Source/Runtime/RenderCore/Private/ShaderCore.cpp
--- a/Engine/Source/Runtime/RenderCore/Private/ShaderCore.cpp
+++ b/Engine/Source/Runtime/RenderCore/Private/ShaderCore.cpp
@@ -2678,7 +2678,7 @@ void AppendKeyStringShaderDependencies(
 	for (const FShaderTypeDependency& ShaderTypeDependency : ShaderTypeDependencies)
 	{
 		const FShaderType* ShaderType = FindShaderTypeByName(ShaderTypeDependency.ShaderTypeName);
-		checkf(ShaderType != nullptr, TEXT("Failed to find FShaderType for dependency %s (total in the NameToTypeMap: %d)"), ShaderTypeDependency.ShaderTypeName.GetDebugString().String.Get(), FShaderType::GetNameToTypeMap().Num());
+		checkf(ShaderType != nullptr, TEXT("Failed to find FShaderType for dependency %hs (total in the NameToTypeMap: %d)"), ShaderTypeDependency.ShaderTypeName.GetDebugString().String.Get(), FShaderType::GetNameToTypeMap().Num());
 
 		OutKeyString.AppendChar('_');
 		OutKeyString.Append(ShaderType->GetName());
@@ -2713,7 +2713,7 @@ void AppendKeyStringShaderDependencies(
 	for (const FShaderPipelineTypeDependency& Dependency : ShaderPipelineTypeDependencies)
 	{
 		const FShaderPipelineType* ShaderPipelineType = FShaderPipelineType::GetShaderPipelineTypeByName(Dependency.ShaderPipelineTypeName);
-		checkf(ShaderPipelineType != nullptr, TEXT("Failed to find FShaderPipelineType for dependency %s (total in the NameToTypeMap: %d)"), Dependency.ShaderPipelineTypeName.GetDebugString().String.Get(), FShaderType::GetNameToTypeMap().Num());
+		checkf(ShaderPipelineType != nullptr, TEXT("Failed to find FShaderPipelineType for dependency %hs (total in the NameToTypeMap: %d)"), Dependency.ShaderPipelineTypeName.GetDebugString().String.Get(), FShaderType::GetNameToTypeMap().Num());
 
 		OutKeyString.AppendChar('_');
 		OutKeyString.Append(ShaderPipelineType->GetName());
diff --git a/Engine/Source/Runtime/RenderCore/Private/UnifiedBuffer.cpp b/Engine/Source/Runtime/RenderCore/Private/UnifiedBuffer.cpp
--- a/Engine/Source/Runtime/RenderCore/Private/UnifiedBuffer.cpp
+++ b/Engine/Source/Runtime/RenderCore/Private/UnifiedBuffer.cpp
@@ -1397,14 +1397,14 @@ RENDERCORE_API bool ResizeResourceSOAIfNeeded<FRWBufferStructured>(FRHICommandLi
 {
 	const uint32 BytesPerElement = 16;
 
-	checkf(Params.NumBytes % BytesPerElement == 0, TEXT("NumBytes (%s) must be a multiple of BytesPerElement (%s)"), Params.NumBytes, BytesPerElement);
-	checkf(Buffer.NumBytes % BytesPerElement == 0, TEXT("NumBytes (%s) must be a multiple of BytesPerElement (%s)"), Buffer.NumBytes, BytesPerElement);
+	checkf(Params.NumBytes % BytesPerElement == 0, TEXT("NumBytes (%u) must be a multiple of BytesPerElement (%u)"), Params.NumBytes, BytesPerElement);
+	checkf(Buffer.NumBytes % BytesPerElement == 0, TEXT("NumBytes (%u) must be a multiple of BytesPerElement (%u)"), Buffer.NumBytes, BytesPerElement);
 
 	uint32 NumElements = Params.NumBytes / BytesPerElement;
 	uint32 NumElementsOld = Buffer.NumBytes / BytesPerElement;
 
-	checkf(NumElements % Params.NumArrays == 0, TEXT("NumElements (%s) must be a multiple of NumArrays (%s)"), NumElements, Params.NumArrays);
-	checkf(NumElementsOld % Params.NumArrays == 0, TEXT("NumElements (%s) must be a multiple of NumArrays (%s)"), NumElementsOld, Params.NumArrays);
+	checkf(NumElements % Params.NumArrays == 0, TEXT("NumElements (%u) must be a multiple of NumArrays (%u)"), NumElements, Params.NumArrays);
+	checkf(NumElementsOld % Params.NumArrays == 0, TEXT("NumElements (%u) must be a multiple of NumArrays (%u)"), NumElementsOld, Params.NumArrays);
 
 	if (Buffer.NumBytes == 0)
 	{
@@ -1450,14 +1450,14 @@ RENDERCORE_API bool ResizeResourceSOAIfNeeded(FRDGBuilder& GraphBuilder, FRWBuff
 {
 	const uint32 BytesPerElement = 16;
 
-	checkf(Params.NumBytes % BytesPerElement == 0, TEXT("NumBytes (%s) must be a multiple of BytesPerElement (%s)"), Params.NumBytes, BytesPerElement);
-	checkf(Buffer.NumBytes % BytesPerElement == 0, TEXT("NumBytes (%s) must be a multiple of BytesPerElement (%s)"), Buffer.NumBytes, BytesPerElement);
+	checkf(Params.NumBytes % BytesPerElement == 0, TEXT("NumBytes (%u) must be a multiple of BytesPerElement (%u)"), Params.NumBytes, BytesPerElement);
+	checkf(Buffer.NumBytes % BytesPerElement == 0, TEXT("NumBytes (%u) must be a multiple of BytesPerElement (%u)"), Buffer.NumBytes, BytesPerElement);
 
 	uint32 NumElements = Params.NumBytes / BytesPerElement;
 	uint32 NumElementsOld = Buffer.NumBytes / BytesPerElement;
 
-	checkf(NumElements % Params.NumArrays == 0, TEXT("NumElements (%s) must be a multiple of NumArrays (%s)"), NumElements, Params.NumArrays);
-	checkf(NumElementsOld % Params.NumArrays == 0, TEXT("NumElements (%s) must be a multiple of NumArrays (%s)"), NumElementsOld, Params.NumArrays);
+	checkf(NumElements % Params.NumArrays == 0, TEXT("NumElements (%u) must be a multiple of NumArrays (%u)"), NumElements, Params.NumArrays);
+	checkf(NumElementsOld % Params.NumArrays == 0, TEXT("NumElements (%u) must be a multiple of NumArrays (%u)"), NumElementsOld, Params.NumArrays);
 
 	if (Buffer.NumBytes == 0)
 	{
@@ -1528,7 +1528,7 @@ RENDERCORE_API bool ResizeResourceIfNeeded(FRDGBuilder& GraphBuilder, FRWBufferS
 {
 	const uint32 BytesPerElement = 16;
 
-	checkf((NumBytes % BytesPerElement) == 0, TEXT("NumBytes (%s) must be a multiple of BytesPerElement (%s)"), NumBytes, BytesPerElement);
+	checkf((NumBytes % BytesPerElement) == 0, TEXT("NumBytes (%u) must be a multiple of BytesPerElement (%u)"), NumBytes, BytesPerElement);
 
 	uint32 NumElements = NumBytes / BytesPerElement;
 
diff --git a/Engine/Source/Runtime/RenderCore/Public/CrossCompilerCommon.h b/Engine/Source/Runtime/RenderCore/Public/CrossCompilerCommon.h
--- a/Engine/Source/Runtime/RenderCore/Public/CrossCompilerCommon.h
+++ b/Engine/Source/Runtime/RenderCore/Public/CrossCompilerCommon.h
@@ -84,7 +84,7 @@ namespace CrossCompiler
 		case SHADER_STAGE_COMPUTE:	return 'c';
 		default: break;
 		}
-		checkf(0, TEXT("invalid value: ShaderStageIndexToTypeName(ShaderStage = '%c')"), (ANSICHAR)ShaderStage);
+		checkf(0, TEXT("invalid value: ShaderStageIndexToTypeName(ShaderStage = '%u')"), ShaderStage);
 		return 0;
 	}
 
diff --git a/Engine/Source/Runtime/RenderCore/Public/ShaderPlatformCachedIniValue.h b/Engine/Source/Runtime/RenderCore/Public/ShaderPlatformCachedIniValue.h
--- a/Engine/Source/Runtime/RenderCore/Public/ShaderPlatformCachedIniValue.h
+++ b/Engine/Source/Runtime/RenderCore/Public/ShaderPlatformCachedIniValue.h
@@ -41,7 +41,7 @@ struct FShaderPlatformCachedIniValue
 		// so instead of hitting an asser below, we detect that the request SP is the current SP, and use the CVar value that is set currently
 		if (IniPlatformName == FPlatformProperties::IniPlatformName() || ActualShaderPlatform == GMaxRHIShaderPlatform)
 		{
-			checkf(CVar != nullptr, TEXT("Failed to find CVar %s when getting current value for FShaderPlatformCachedIniValue"));
+			checkf(CVar != nullptr, TEXT("Failed to find CVar %s when getting current value for FShaderPlatformCachedIniValue"), *CVarName);
 
 			CVar->GetValue(Value);
 			return Value;
diff --git a/Engine/Source/Runtime/Renderer/Private/OverridePassSequence.h b/Engine/Source/Runtime/Renderer/Private/OverridePassSequence.h
--- a/Engine/Source/Runtime/Renderer/Private/OverridePassSequence.h
+++ b/Engine/Source/Runtime/Renderer/Private/OverridePassSequence.h
@@ -35,7 +35,7 @@ public:
 			{
 				if (Passes[PassIndex].bEnabled)
 				{
-					checkf(Passes[PassIndex].bAccepted, TEXT("Pass was enabled but not accepted: %s."), Passes[PassIndex].bAccepted);
+					checkf(Passes[PassIndex].bAccepted, TEXT("Pass was enabled but not accepted: %s."), Passes[PassIndex].Name);
 				}
 			}
 		}
diff --git a/Engine/Source/Runtime/Renderer/Private/ShaderBaseClasses.cpp b/Engine/Source/Runtime/Renderer/Private/ShaderBaseClasses.cpp
--- a/Engine/Source/Runtime/Renderer/Private/ShaderBaseClasses.cpp
+++ b/Engine/Source/Runtime/Renderer/Private/ShaderBaseClasses.cpp
@@ -263,7 +263,7 @@ void FMaterialShader::SetParameters(
 	const ERHIFeatureLevel::Type FeatureLevel = View.GetFeatureLevel();
 	FMaterialShaderMap* ShaderMap = Material.GetRenderingThreadShaderMap();
 	checkf(ShaderMap, TEXT("RenderingThreadShaderMap: %i"), ShaderMap ? 1 : 0);
-	checkf(ShaderMap->IsValidForRendering(true) && Material.GetFeatureLevel() == FeatureLevel, TEXT("IsValid:%i, MaterialFeatureLevel:%i, FeatureLevel:%i"), ShaderMap->IsValidForRendering() ? 1 : 0, Material.GetFeatureLevel(), FeatureLevel);
+	checkf(ShaderMap->IsValidForRendering(true) && Material.GetFeatureLevel() == FeatureLevel, TEXT("IsValid:%i, MaterialFeatureLevel:%i, FeatureLevel:%i"), ShaderMap->IsValidForRendering() ? 1 : 0, (ERHIFeatureLevel::Type)Material.GetFeatureLevel(), FeatureLevel);
 
 	FUniformExpressionCache* UniformExpressionCache = &MaterialRenderProxy->UniformExpressionCache[FeatureLevel];
 	bool bUniformExpressionCacheNeedsDelete = false;
diff --git a/Engine/Source/Runtime/Slate/Private/Framework/Application/SlateApplication.cpp b/Engine/Source/Runtime/Slate/Private/Framework/Application/SlateApplication.cpp
--- a/Engine/Source/Runtime/Slate/Private/Framework/Application/SlateApplication.cpp
+++ b/Engine/Source/Runtime/Slate/Private/Framework/Application/SlateApplication.cpp
@@ -6232,7 +6232,7 @@ bool FSlateApplication::OnControllerButtonPressed(FGamepadKeyNames::Type KeyName
 {
 	FKey Key(KeyName);
 	TOptional<int32> UserIndex = GetUserIndexForInputDevice(InputDeviceId);
-	if (UserIndex.IsSet() && ensureMsgf(Key.IsValid(), TEXT("OnControllerButtonPressed(KeyName=%s,InputDeviceId=%d,IsRepeat=%b) key is invalid"), *KeyName.ToString(), InputDeviceId.GetId(), IsRepeat))
+	if (UserIndex.IsSet() && ensureMsgf(Key.IsValid(), TEXT("OnControllerButtonPressed(KeyName=%s,InputDeviceId=%d,IsRepeat=%u) key is invalid"), *KeyName.ToString(), InputDeviceId.GetId(), IsRepeat))
 	{
 		FKeyEvent KeyEvent(Key, PlatformApplication->GetModifierKeys(), InputDeviceId, IsRepeat, 0, 0, UserIndex);
 		
@@ -6245,7 +6245,7 @@ bool FSlateApplication::OnControllerButtonReleased(FGamepadKeyNames::Type KeyNam
 {
 	FKey Key(KeyName);
 	TOptional<int32> UserIndex = GetUserIndexForInputDevice(InputDeviceId);
-	if (UserIndex.IsSet() && ensureMsgf(Key.IsValid(), TEXT("OnControllerButtonReleased(KeyName=%s,InputDeviceId=%d,IsRepeat=%b) key is invalid"), *KeyName.ToString(), InputDeviceId.GetId(), IsRepeat))
+	if (UserIndex.IsSet() && ensureMsgf(Key.IsValid(), TEXT("OnControllerButtonReleased(KeyName=%s,InputDeviceId=%d,IsRepeat=%u) key is invalid"), *KeyName.ToString(), InputDeviceId.GetId(), IsRepeat))
 	{
 		FKeyEvent KeyEvent(Key, PlatformApplication->GetModifierKeys(), InputDeviceId, IsRepeat, 0, 0, UserIndex);
 		
diff --git a/Engine/Source/Runtime/SlateCore/Private/Types/SlateAttributeMetaData.cpp b/Engine/Source/Runtime/SlateCore/Private/Types/SlateAttributeMetaData.cpp
--- a/Engine/Source/Runtime/SlateCore/Private/Types/SlateAttributeMetaData.cpp
+++ b/Engine/Source/Runtime/SlateCore/Private/Types/SlateAttributeMetaData.cpp
@@ -206,7 +206,7 @@ void FSlateAttributeMetaData::RegisterContainAttributeImpl(ISlateAttributeContai
 	const FName ContainerName = OwningContainer.GetContainerName();
 	const FSlateAttributeDescriptor::OffsetType AttributeOffset = UE::Private::FindContainOffset(OwningContainer, Attribute);
 	const FSlateAttributeDescriptor::FAttribute* FoundAttribute = Descriptor.FindContainedAttribute(ContainerName, AttributeOffset);
-	ensureAlwaysMsgf(FoundAttribute, TEXT("Attribute with offset '%d' is not registered. All Slot Attributes needs to be registered."));
+	ensureAlwaysMsgf(FoundAttribute, TEXT("Attribute with offset '%d' is not registered. All Slot Attributes needs to be registered."), AttributeOffset);
 	if (FoundAttribute)
 	{
 		const FSlateAttributeDescriptor::FContainer* FoundCountainer = Descriptor.FindContainer(ContainerName);
diff --git a/Engine/Source/Runtime/TimeManagement/Private/TimeManagementTests.cpp b/Engine/Source/Runtime/TimeManagement/Private/TimeManagementTests.cpp
--- a/Engine/Source/Runtime/TimeManagement/Private/TimeManagementTests.cpp
+++ b/Engine/Source/Runtime/TimeManagement/Private/TimeManagementTests.cpp
@@ -139,7 +139,7 @@ bool FFrameTimeAdditionTest::RunTest(const FString& Parameters)
 			FFrameTime Actual   = Time + TimeToAdd;
 			FFrameTime Expected = ExpectedTimes[Index];
 
-			ensureAlwaysMsgf(IsNearlyEqual(Actual, Expected), TEXT("%d+.3f + 10.1: %d+%.3f (actual) == %d+%.3f (expected)"),
+			ensureAlwaysMsgf(IsNearlyEqual(Actual, Expected), TEXT("%d+%.3f + 10.1: %d+%.3f (actual) == %d+%.3f (expected)"),
 				Time.GetFrame().Value,     Time.GetSubFrame(),
 				Actual.GetFrame().Value,   Actual.GetSubFrame(),
 				Expected.GetFrame().Value, Expected.GetSubFrame()
@@ -594,7 +594,7 @@ bool FFrameTimeToSecondsConversionTest::RunTest(const FString& Parameters)
 
 			// Verify seconds value against interval * FrameIndex
 			ensureAlwaysMsgf(FMath::IsNearlyEqual(IntervalSeconds, FrameAsSeconds),
-						TEXT("Did not expect value %d ::AsSeconds for input %i (%d)."),
+						TEXT("Did not expect value %f ::AsSeconds for input %i (%f)."),
 						FrameAsSeconds, FrameIndex, IntervalSeconds
 					);
 			
@@ -604,19 +604,19 @@ bool FFrameTimeToSecondsConversionTest::RunTest(const FString& Parameters)
 
 			// Verify frame numbers match original FrameIndex value
 			ensureAlwaysMsgf(FrameTime.GetFrame() == FrameIndex,
-						TEXT("Did not expect frame number %i ::AsFrameTime for input %i (%d)."),
+						TEXT("Did not expect frame number %i ::AsFrameTime for input %i (%f)."),
 						FrameNumber.Value, FrameIndex, FrameAsSeconds
 					);
 			
 			ensureAlwaysMsgf(FrameNumber.Value == FrameIndex,
-						TEXT("Did not expect frame number %i ::AsFrameNumber for input %i (%d)."),
+						TEXT("Did not expect frame number %i ::AsFrameNumber for input %i (%f)."),
 						FrameNumber.Value, FrameIndex, FrameAsSeconds
 					);
 			
 			// Convert FrameTime back to seconds and verify against value used to generate it
 			const double ConvertedSeconds = FrameRate.AsSeconds(FrameTime);
 			ensureAlwaysMsgf(FMath::IsNearlyEqual(ConvertedSeconds, FrameAsSeconds),
-				TEXT("Did not expect value %d ::AsSeconds for input %i (%d) - %s."),
+				TEXT("Did not expect value %f ::AsSeconds for input %i (%f) - %s."),
 				ConvertedSeconds, FrameIndex, FrameAsSeconds, *LexToString(FrameTime)
 			);
 		}
