Added validation to issue an ensure when a sync scope is used during end of frame updates.

diff --git a/Engine/Source/Runtime/Engine/Private/LevelTick.cpp b/Engine/Source/Runtime/Engine/Private/LevelTick.cpp
--- a/Engine/Source/Runtime/Engine/Private/LevelTick.cpp
+++ b/Engine/Source/Runtime/Engine/Private/LevelTick.cpp
@@ -975,6 +975,8 @@ void BeginSendEndOfFrameUpdatesDrawEvent(FSendAllEndOfFrameUpdates& SendAllEndOf
 {
 	BEGIN_DRAW_EVENTF_GAMETHREAD(SendAllEndOfFrameUpdates, SendAllEndOfFrameUpdates.DrawEvent, TEXT("SendAllEndOfFrameUpdates"));
 
+	UE::RenderCommandPipe::StartParallelRecordingValidation();
+
 	ENQUEUE_RENDER_COMMAND(BeginDrawEventCommand)(UE::RenderCommandPipe::SkeletalMesh,
 		[GPUSkinCache = SendAllEndOfFrameUpdates.GPUSkinCache]
 	{
@@ -998,6 +1000,8 @@ void EndSendEndOfFrameUpdatesDrawEvent(FSendAllEndOfFrameUpdates& SendAllEndOfFr
 		}
 	});
 
+	UE::RenderCommandPipe::StopParallelRecordingValidation();
+
 	STOP_DRAW_EVENT_GAMETHREAD(SendAllEndOfFrameUpdates.DrawEvent);
 }
 
diff --git a/Engine/Source/Runtime/RenderCore/Private/RenderingThread.cpp b/Engine/Source/Runtime/RenderCore/Private/RenderingThread.cpp
--- a/Engine/Source/Runtime/RenderCore/Private/RenderingThread.cpp
+++ b/Engine/Source/Runtime/RenderCore/Private/RenderingThread.cpp
@@ -1611,6 +1611,7 @@ public:
 	{
 #if !UE_SERVER
 		check(IsInGameThread());
+		ensureMsgf(!bParallelRecordingValidation, TEXT("StartRecording has been called within a region designated for issuing parallel render thread commands. This is a race condition."));
 		SCOPED_NAMED_EVENT(FRenderCommandPipe_StartRecording, FColor::Magenta);
 
 		if (bRecording || UE::RenderCommandPipe::GetMode() != ERenderCommandPipeMode::All || Pipes.IsEmpty() || !ShouldExecuteOnRenderThread())
@@ -1663,6 +1664,7 @@ public:
 	{
 #if !UE_SERVER
 		check(IsInGameThread());
+		ensureMsgf(!bParallelRecordingValidation, TEXT("StopRecording has been called within a region designated for issuing parallel render thread commands. This is a race condition."));
 		SCOPED_NAMED_EVENT(FRenderCommandPipe_StopRecording, FColor::Magenta);
 
 		if (!bRecording)
@@ -1708,6 +1710,22 @@ public:
 #endif // !UE_SERVER
 	}
 
+	void StartParallelRecordingValidation()
+	{
+#if !UE_SERVER
+		check(!bParallelRecordingValidation);
+		bParallelRecordingValidation = true;
+#endif
+	}
+
+	void StopParallelRecordingValidation()
+	{
+#if !UE_SERVER
+		check(bParallelRecordingValidation);
+		bParallelRecordingValidation = false;
+#endif
+	}
+
 	bool IsRecording() const
 	{
 #if !UE_SERVER
@@ -1735,6 +1753,7 @@ private:
 	TArray<FRenderCommandPipe*> Pipes;
 	bool bRecording = false;
 	bool bReplaying = false;
+	bool bParallelRecordingValidation = false;
 #endif // !UE_SERVER
 };
 
@@ -1771,6 +1790,16 @@ namespace UE::RenderCommandPipe
 	{
 		GRenderCommandPipeRegistry.StopRecording();
 	}
+
+	void StartParallelRecordingValidation()
+	{
+		GRenderCommandPipeRegistry.StartParallelRecordingValidation();
+	}
+
+	void StopParallelRecordingValidation()
+	{
+		GRenderCommandPipeRegistry.StopParallelRecordingValidation();
+	}
 }
 
 FRenderCommandPipe::FRenderCommandPipe(const TCHAR* InName, ERenderCommandPipeFlags Flags, const TCHAR* CVarName, const TCHAR* CVarDescription)
diff --git a/Engine/Source/Runtime/RenderCore/Public/RenderingThread.h b/Engine/Source/Runtime/RenderCore/Public/RenderingThread.h
--- a/Engine/Source/Runtime/RenderCore/Public/RenderingThread.h
+++ b/Engine/Source/Runtime/RenderCore/Public/RenderingThread.h
@@ -314,6 +314,12 @@ namespace UE::RenderCommandPipe
 	// [Game Thread] Stops recording commands into pipes and syncs all remaining pipe work to the render thread.
 	extern RENDERCORE_API void StopRecording();
 
+	// [Game Thread] Enables validation that will issue an ensure if recording is stopped.
+	extern RENDERCORE_API void StartParallelRecordingValidation();
+
+	// [Game Thread] Disables validation that will issue an ensure if recording is stopped.
+	extern RENDERCORE_API void StopParallelRecordingValidation();
+
 	// [Game Thread] Stops render command pipe recording during the duration of the scope and restarts recording once the scope is complete.
 	class FSyncScope
 	{
