[shaders] mutate ddc key when enabling preprocessed job cache (so if unstripped shader code for symbols is desired, the cvar can be toggled off and new shaders compiled)

diff --git a/Engine/Source/Runtime/RenderCore/Private/Shader.cpp b/Engine/Source/Runtime/RenderCore/Private/Shader.cpp
--- a/Engine/Source/Runtime/RenderCore/Private/Shader.cpp
+++ b/Engine/Source/Runtime/RenderCore/Private/Shader.cpp
@@ -2134,6 +2134,14 @@ void ShaderMapAppendKeyString(EShaderPlatform Platform, FString& KeyString)
 		}
 	}
 
+	{
+		static const auto CVar = IConsoleManager::Get().FindTConsoleVariableDataBool(TEXT("r.ShaderCompiler.PreprocessedJobCache"));
+		if (CVar && CVar->GetValueOnAnyThread())
+		{
+			KeyString += TEXT("_PJC");
+		}
+	}
+
 	if (RHISupportsShaderRootConstants(Platform))
 	{
 		KeyString += TEXT("_SHRC");
