Don't add unbind commands if the RHI doesn't do anything with them.
- Discovered that a number of RHISetShaderResourceViewParameter and RHISetUAVParameter implementations do nothing with null resources, meaning that RHISetShaderUnbinds calls end up doing zero work with all the overhead.
- Adding GRHIGlobals::NeedsShaderUnbinds for RHIs to signal they need/handle unbind commands.
- RHIs that actually do work with RHISetShaderUnbinds now set NeedsShaderUnbinds on startup.
- Removed RHISetShaderUnbinds implementations that did zero work.

diff --git a/Engine/Source/Runtime/Apple/MetalRHI/Private/MetalCommands.cpp b/Engine/Source/Runtime/Apple/MetalRHI/Private/MetalCommands.cpp
--- a/Engine/Source/Runtime/Apple/MetalRHI/Private/MetalCommands.cpp
+++ b/Engine/Source/Runtime/Apple/MetalRHI/Private/MetalCommands.cpp
@@ -390,16 +390,6 @@ void FMetalRHICommandContext::RHISetShaderParameters(FRHIComputeShader* Shader,
 	);
 }
 
-void FMetalRHICommandContext::RHISetShaderUnbinds(FRHIComputeShader* Shader, TConstArrayView<FRHIShaderParameterUnbind> InUnbinds)
-{
-	UE::RHICore::RHISetShaderUnbindsShared(*this, Shader, InUnbinds);
-}
-
-void FMetalRHICommandContext::RHISetShaderUnbinds(FRHIGraphicsShader* Shader, TConstArrayView<FRHIShaderParameterUnbind> InUnbinds)
-{
-	UE::RHICore::RHISetShaderUnbindsShared(*this, Shader, InUnbinds);
-}
-
 void FMetalRHICommandContext::RHISetStencilRef(uint32 StencilRef)
 {
 	Context->GetCurrentState().SetStencilRef(StencilRef);
diff --git a/Engine/Source/Runtime/Apple/MetalRHI/Public/MetalRHIContext.h b/Engine/Source/Runtime/Apple/MetalRHI/Public/MetalRHIContext.h
--- a/Engine/Source/Runtime/Apple/MetalRHI/Public/MetalRHIContext.h
+++ b/Engine/Source/Runtime/Apple/MetalRHI/Public/MetalRHIContext.h
@@ -144,8 +144,6 @@ public:
 	
 	virtual void RHISetShaderParameters(FRHIGraphicsShader* Shader, TConstArrayView<uint8> InParametersData, TConstArrayView<FRHIShaderParameter> InParameters, TConstArrayView<FRHIShaderParameterResource> InResourceParameters, TConstArrayView<FRHIShaderParameterResource> InBindlessParameters) final override;
 	virtual void RHISetShaderParameters(FRHIComputeShader* Shader, TConstArrayView<uint8> InParametersData, TConstArrayView<FRHIShaderParameter> InParameters, TConstArrayView<FRHIShaderParameterResource> InResourceParameters, TConstArrayView<FRHIShaderParameterResource> InBindlessParameters) final override;
-	virtual void RHISetShaderUnbinds(FRHIComputeShader* Shader, TConstArrayView<FRHIShaderParameterUnbind> InUnbinds) final override;
-	virtual void RHISetShaderUnbinds(FRHIGraphicsShader* Shader, TConstArrayView<FRHIShaderParameterUnbind> InUnbinds) final override;
 
 	virtual void RHISetStencilRef(uint32 StencilRef) final override;
 
diff --git a/Engine/Source/Runtime/D3D12RHI/Private/D3D12RHI.cpp b/Engine/Source/Runtime/D3D12RHI/Private/D3D12RHI.cpp
--- a/Engine/Source/Runtime/D3D12RHI/Private/D3D12RHI.cpp
+++ b/Engine/Source/Runtime/D3D12RHI/Private/D3D12RHI.cpp
@@ -255,6 +255,8 @@ FD3D12DynamicRHI::FD3D12DynamicRHI(const TArray<TSharedPtr<FD3D12Adapter>>& Chos
 	GRHISupportsFrameCyclesBubblesRemoval = true;
 	GRHISupportsGPUTimestampBubblesRemoval = true;
 	GRHISupportsRHIOnTaskThread = true;
+
+	GRHIGlobals.NeedsShaderUnbinds = true;
 }
 
 void FD3D12DynamicRHI::PostInit()
diff --git a/Engine/Source/Runtime/NullDrv/Public/NullRHI.h b/Engine/Source/Runtime/NullDrv/Public/NullRHI.h
--- a/Engine/Source/Runtime/NullDrv/Public/NullRHI.h
+++ b/Engine/Source/Runtime/NullDrv/Public/NullRHI.h
@@ -507,14 +507,6 @@ public:
 	{
 	}
 
-	virtual void RHISetShaderUnbinds(FRHIGraphicsShader* ComputeShader, TConstArrayView<FRHIShaderParameterUnbind> InUnbinds)
-	{
-	}
-
-	virtual void RHISetShaderUnbinds(FRHIComputeShader* ComputeShader, TConstArrayView<FRHIShaderParameterUnbind> InUnbinds)
-	{
-	}
-
 	virtual void RHISetDepthStencilState(FRHIDepthStencilState* NewState, uint32 StencilRef) final override
 	{
 
diff --git a/Engine/Source/Runtime/OpenGLDrv/Private/OpenGLDevice.cpp b/Engine/Source/Runtime/OpenGLDrv/Private/OpenGLDevice.cpp
--- a/Engine/Source/Runtime/OpenGLDrv/Private/OpenGLDevice.cpp
+++ b/Engine/Source/Runtime/OpenGLDrv/Private/OpenGLDevice.cpp
@@ -1252,6 +1252,8 @@ static void InitRHICapabilitiesForGL()
 	}
 	
 	GRHISupportsPipelineFileCache = !GRHISupportsPSOPrecaching || CVarEnablePSOFileCacheWhenPrecachingActive.GetValueOnAnyThread();
+
+	GRHIGlobals.NeedsShaderUnbinds = true;
 }
 
 FDynamicRHI* FOpenGLDynamicRHIModule::CreateRHI(ERHIFeatureLevel::Type InRequestedFeatureLevel)
diff --git a/Engine/Source/Runtime/RHI/Public/RHICommandList.h b/Engine/Source/Runtime/RHI/Public/RHICommandList.h
--- a/Engine/Source/Runtime/RHI/Public/RHICommandList.h
+++ b/Engine/Source/Runtime/RHI/Public/RHICommandList.h
@@ -1006,6 +1006,12 @@ public:
 		return ScratchShaderUnbinds;
 	}
 
+	// Returns true if the RHI needs unbind commands
+	bool NeedsShaderUnbinds() const
+	{
+		return GRHIGlobals.NeedsShaderUnbinds;
+	}
+
 protected:
 	FMemStackBase& GetAllocator() { return MemManager; }
 
@@ -2599,15 +2605,18 @@ PRAGMA_ENABLE_DEPRECATION_WARNINGS
 
 	FORCEINLINE_DEBUGGABLE void SetShaderUnbinds(FRHIComputeShader* InShader, TConstArrayView<FRHIShaderParameterUnbind> InUnbinds)
 	{
-		ValidateBoundShader(InShader);
-
-		if (Bypass())
+		if (NeedsShaderUnbinds())
 		{
-			GetComputeContext().RHISetShaderUnbinds(InShader, InUnbinds);
-			return;
-		}
+			ValidateBoundShader(InShader);
+
+			if (Bypass())
+			{
+				GetComputeContext().RHISetShaderUnbinds(InShader, InUnbinds);
+				return;
+			}
 
-		ALLOC_COMMAND(FRHICommandSetShaderUnbinds<FRHIComputeShader>)(InShader, AllocArray(InUnbinds));
+			ALLOC_COMMAND(FRHICommandSetShaderUnbinds<FRHIComputeShader>)(InShader, AllocArray(InUnbinds));
+		}
 	}
 
 	FORCEINLINE_DEBUGGABLE void SetBatchedShaderUnbinds(FRHIComputeShader* InShader, FRHIBatchedShaderUnbinds& InBatchedUnbinds)
@@ -3335,15 +3344,18 @@ PRAGMA_ENABLE_DEPRECATION_WARNINGS
 
 	FORCEINLINE_DEBUGGABLE void SetShaderUnbinds(FRHIGraphicsShader* InShader, TConstArrayView<FRHIShaderParameterUnbind> InUnbinds)
 	{
-		ValidateBoundShader(InShader);
-
-		if (Bypass())
+		if (NeedsShaderUnbinds())
 		{
-			GetContext().RHISetShaderUnbinds(InShader, InUnbinds);
-			return;
-		}
+			ValidateBoundShader(InShader);
 
-		ALLOC_COMMAND(FRHICommandSetShaderUnbinds<FRHIGraphicsShader>)(InShader, AllocArray(InUnbinds));
+			if (Bypass())
+			{
+				GetContext().RHISetShaderUnbinds(InShader, InUnbinds);
+				return;
+			}
+
+			ALLOC_COMMAND(FRHICommandSetShaderUnbinds<FRHIGraphicsShader>)(InShader, AllocArray(InUnbinds));
+		}
 	}
 
 	using FRHIComputeCommandList::SetBatchedShaderUnbinds;
diff --git a/Engine/Source/Runtime/RHI/Public/RHIContext.h b/Engine/Source/Runtime/RHI/Public/RHIContext.h
--- a/Engine/Source/Runtime/RHI/Public/RHIContext.h
+++ b/Engine/Source/Runtime/RHI/Public/RHIContext.h
@@ -349,7 +349,10 @@ public:
 		InBatchedParameters.Reset();
 	}
 
-	virtual void RHISetShaderUnbinds(FRHIComputeShader* ComputeShader, TConstArrayView<FRHIShaderParameterUnbind> InUnbinds) = 0;
+	virtual void RHISetShaderUnbinds(FRHIComputeShader* ComputeShader, TConstArrayView<FRHIShaderParameterUnbind> InUnbinds)
+	{
+		checkf(false, TEXT("RHISetShaderUnbinds called when the active RHI hasn't overridden it and GRHIGlobals.NeedsShaderUnbinds is set."));
+	}
 
 	virtual void RHISetShaderParameter(FRHIComputeShader* ComputeShader, uint32 BufferIndex, uint32 BaseIndex, uint32 NumBytes, const void* NewValue) = 0;
 
@@ -755,7 +758,10 @@ public:
 		InBatchedParameters.Reset();
 	}
 
-	virtual void RHISetShaderUnbinds(FRHIGraphicsShader* Shader, TConstArrayView<FRHIShaderParameterUnbind> InUnbinds) = 0;
+	virtual void RHISetShaderUnbinds(FRHIGraphicsShader* Shader, TConstArrayView<FRHIShaderParameterUnbind> InUnbinds)
+	{
+		checkf(false, TEXT("RHISetShaderUnbinds called when the active RHI hasn't overridden it and GRHIGlobals.NeedsShaderUnbinds is set."));
+	}
 
 	virtual void RHISetShaderParameter(FRHIGraphicsShader* Shader, uint32 BufferIndex, uint32 BaseIndex, uint32 NumBytes, const void* NewValue) = 0;
 
diff --git a/Engine/Source/Runtime/RHI/Public/RHIGlobals.h b/Engine/Source/Runtime/RHI/Public/RHIGlobals.h
--- a/Engine/Source/Runtime/RHI/Public/RHIGlobals.h
+++ b/Engine/Source/Runtime/RHI/Public/RHIGlobals.h
@@ -603,6 +603,9 @@ struct FRHIGlobals
 	
 	/** True if the RHI has initialized a device with the debug layer enabled. */
 	bool IsDebugLayerEnabled = false;
+
+	/** True if the RHI needs shader unbinds (SetShaderUnbinds). RHIs that don't need them can avoid creating extra commands. */
+	bool NeedsShaderUnbinds = false;
 };
 
 extern RHI_API FRHIGlobals GRHIGlobals;
diff --git a/Engine/Source/Runtime/RenderCore/Private/BuiltInRayTracingShaders.cpp b/Engine/Source/Runtime/RenderCore/Private/BuiltInRayTracingShaders.cpp
--- a/Engine/Source/Runtime/RenderCore/Private/BuiltInRayTracingShaders.cpp
+++ b/Engine/Source/Runtime/RenderCore/Private/BuiltInRayTracingShaders.cpp
@@ -81,11 +81,14 @@ void FRayTracingDispatchDescCS::Dispatch(FRHICommandList& RHICmdList,
 
 	RHICmdList.DispatchComputeShader(1, 1, 1);
 
-	FRHIBatchedShaderUnbinds& BatchedUnbinds = RHICmdList.GetScratchShaderUnbinds();
+	if (RHICmdList.NeedsShaderUnbinds())
+	{
+		FRHIBatchedShaderUnbinds& BatchedUnbinds = RHICmdList.GetScratchShaderUnbinds();
 
-	UnsetSRVParameter(BatchedUnbinds, ComputeShader->DispatchDimensionsParam);
-	UnsetUAVParameter(BatchedUnbinds, ComputeShader->DispatchDescOutputParam);
-	RHICmdList.SetBatchedShaderUnbinds(ComputeShader.GetComputeShader(), BatchedUnbinds);
+		UnsetSRVParameter(BatchedUnbinds, ComputeShader->DispatchDimensionsParam);
+		UnsetUAVParameter(BatchedUnbinds, ComputeShader->DispatchDescOutputParam);
+		RHICmdList.SetBatchedShaderUnbinds(ComputeShader.GetComputeShader(), BatchedUnbinds);
+	}
 }
 
 #endif // RHI_RAYTRACING
diff --git a/Engine/Source/Runtime/RenderCore/Private/RayTracingValidationShaders.cpp b/Engine/Source/Runtime/RenderCore/Private/RayTracingValidationShaders.cpp
--- a/Engine/Source/Runtime/RenderCore/Private/RayTracingValidationShaders.cpp
+++ b/Engine/Source/Runtime/RenderCore/Private/RayTracingValidationShaders.cpp
@@ -74,12 +74,15 @@ void FRayTracingValidateGeometryBuildParamsCS::Dispatch(FRHICommandList& RHICmdL
 		RHICmdList.DispatchComputeShader(NumGroupsX, 1, 1);
 	}
 
-	FRHIBatchedShaderUnbinds& BatchedUnbinds = RHICmdList.GetScratchShaderUnbinds();
+	if (RHICmdList.NeedsShaderUnbinds())
+	{
+		FRHIBatchedShaderUnbinds& BatchedUnbinds = RHICmdList.GetScratchShaderUnbinds();
 
-	UnsetSRVParameter(BatchedUnbinds, ComputeShader->VertexBufferParam);
-	UnsetSRVParameter(BatchedUnbinds, ComputeShader->IndexBufferParam);
+		UnsetSRVParameter(BatchedUnbinds, ComputeShader->VertexBufferParam);
+		UnsetSRVParameter(BatchedUnbinds, ComputeShader->IndexBufferParam);
 
-	RHICmdList.SetBatchedShaderUnbinds(ShaderRHI, BatchedUnbinds);
+		RHICmdList.SetBatchedShaderUnbinds(ShaderRHI, BatchedUnbinds);
+	}
 
 	RHICmdList.PopEvent();
 }
diff --git a/Engine/Source/Runtime/RenderCore/Public/ClearReplacementShaders.h b/Engine/Source/Runtime/RenderCore/Public/ClearReplacementShaders.h
--- a/Engine/Source/Runtime/RenderCore/Public/ClearReplacementShaders.h
+++ b/Engine/Source/Runtime/RenderCore/Public/ClearReplacementShaders.h
@@ -398,9 +398,12 @@ inline void ClearUAVShader_T(FRHIComputeCommandList& RHICmdList, FRHIUnorderedAc
 		FMath::DivideAndRoundUp(SizeZ, ComputeShader->ThreadGroupSizeZ)
 	);
 
-	FRHIBatchedShaderUnbinds& BatchedUnbinds = RHICmdList.GetScratchShaderUnbinds();
-	UnsetUAVParameter(BatchedUnbinds, ComputeShader->GetClearResourceParam());
-	RHICmdList.SetBatchedShaderUnbinds(ShaderRHI, BatchedUnbinds);
+	if (RHICmdList.NeedsShaderUnbinds())
+	{
+		FRHIBatchedShaderUnbinds& BatchedUnbinds = RHICmdList.GetScratchShaderUnbinds();
+		UnsetUAVParameter(BatchedUnbinds, ComputeShader->GetClearResourceParam());
+		RHICmdList.SetBatchedShaderUnbinds(ShaderRHI, BatchedUnbinds);
+	}
 
 	if (bBarriers)
 	{
diff --git a/Engine/Source/Runtime/RenderCore/Public/ShaderParameterStruct.h b/Engine/Source/Runtime/RenderCore/Public/ShaderParameterStruct.h
--- a/Engine/Source/Runtime/RenderCore/Public/ShaderParameterStruct.h
+++ b/Engine/Source/Runtime/RenderCore/Public/ShaderParameterStruct.h
@@ -193,27 +193,36 @@ inline void UnsetShaderSRVs(FRHIBatchedShaderUnbinds& BatchedUnbinds, const TSha
 template<typename TRHICmdList, typename TShaderClass>
 inline void UnsetShaderUAVs(TRHICmdList& RHICmdList, const TShaderRef<TShaderClass>& Shader, FRHIComputeShader* ShaderRHI)
 {
-	FRHIBatchedShaderUnbinds& BatchedUnbinds = RHICmdList.GetScratchShaderUnbinds();
-	UnsetShaderUAVs(BatchedUnbinds, Shader);
-	RHICmdList.SetBatchedShaderUnbinds(ShaderRHI, BatchedUnbinds);
+	if (RHICmdList.NeedsShaderUnbinds())
+	{
+		FRHIBatchedShaderUnbinds& BatchedUnbinds = RHICmdList.GetScratchShaderUnbinds();
+		UnsetShaderUAVs(BatchedUnbinds, Shader);
+		RHICmdList.SetBatchedShaderUnbinds(ShaderRHI, BatchedUnbinds);
+	}
 }
 
 /** Unset compute shader SRVs. */
 template<typename TRHICmdList, typename TShaderClass>
 inline void UnsetShaderSRVs(TRHICmdList& RHICmdList, const TShaderRef<TShaderClass>& Shader, FRHIComputeShader* ShaderRHI)
 {
-	FRHIBatchedShaderUnbinds& BatchedUnbinds = RHICmdList.GetScratchShaderUnbinds();
-	UnsetShaderSRVs(BatchedUnbinds, Shader);
-	RHICmdList.SetBatchedShaderUnbinds(ShaderRHI, BatchedUnbinds);
+	if (RHICmdList.NeedsShaderUnbinds())
+	{
+		FRHIBatchedShaderUnbinds& BatchedUnbinds = RHICmdList.GetScratchShaderUnbinds();
+		UnsetShaderSRVs(BatchedUnbinds, Shader);
+		RHICmdList.SetBatchedShaderUnbinds(ShaderRHI, BatchedUnbinds);
+	}
 }
 
 /** Unset compute shader SRVs. */
 template<typename TRHICmdList, typename TShaderClass>
 inline void UnsetShaderSRVs(TRHICmdList& RHICmdList, const TShaderRef<TShaderClass>& Shader, FRHIGraphicsShader* ShaderRHI)
 {
-	FRHIBatchedShaderUnbinds& BatchedUnbinds = RHICmdList.GetScratchShaderUnbinds();
-	UnsetShaderSRVs(BatchedUnbinds, Shader);
-	RHICmdList.SetBatchedShaderUnbinds(ShaderRHI, BatchedUnbinds);
+	if (RHICmdList.NeedsShaderUnbinds())
+	{
+		FRHIBatchedShaderUnbinds& BatchedUnbinds = RHICmdList.GetScratchShaderUnbinds();
+		UnsetShaderSRVs(BatchedUnbinds, Shader);
+		RHICmdList.SetBatchedShaderUnbinds(ShaderRHI, BatchedUnbinds);
+	}
 }
 
 RENDERCORE_API void SetShaderParameters(
diff --git a/Engine/Source/Runtime/RenderCore/Public/ShaderParameterUtils.h b/Engine/Source/Runtime/RenderCore/Public/ShaderParameterUtils.h
--- a/Engine/Source/Runtime/RenderCore/Public/ShaderParameterUtils.h
+++ b/Engine/Source/Runtime/RenderCore/Public/ShaderParameterUtils.h
@@ -389,18 +389,24 @@ inline void SetShaderParametersLegacyCS(TRHICmdList& RHICmdList, const TShaderRe
 template<typename TRHICmdList, typename TShaderType>
 inline void UnsetShaderParametersLegacyPS(TRHICmdList& RHICmdList, const TShaderRef<TShaderType>& InShader)
 {
-	FRHIBatchedShaderUnbinds& BatchedUnbinds = RHICmdList.GetScratchShaderUnbinds();
-	InShader->UnsetParameters(BatchedUnbinds);
-	RHICmdList.SetBatchedShaderUnbinds(InShader.GetPixelShader(), BatchedUnbinds);
+	if (RHICmdList.NeedsShaderUnbinds())
+	{
+		FRHIBatchedShaderUnbinds& BatchedUnbinds = RHICmdList.GetScratchShaderUnbinds();
+		InShader->UnsetParameters(BatchedUnbinds);
+		RHICmdList.SetBatchedShaderUnbinds(InShader.GetPixelShader(), BatchedUnbinds);
+	}
 }
 
 /// Utility to unset all legacy parameters for a Compute shader. Requires the shader type to implement UnsetParameters(FRHIBatchedShaderUnbinds& BatchedUnbinds)
 template<typename TRHICmdList, typename TShaderType>
 inline void UnsetShaderParametersLegacyCS(TRHICmdList& RHICmdList, const TShaderRef<TShaderType>& InShader)
 {
-	FRHIBatchedShaderUnbinds& BatchedUnbinds = RHICmdList.GetScratchShaderUnbinds();
-	InShader->UnsetParameters(BatchedUnbinds);
-	RHICmdList.SetBatchedShaderUnbinds(InShader.GetComputeShader(), BatchedUnbinds);
+	if (RHICmdList.NeedsShaderUnbinds())
+	{
+		FRHIBatchedShaderUnbinds& BatchedUnbinds = RHICmdList.GetScratchShaderUnbinds();
+		InShader->UnsetParameters(BatchedUnbinds);
+		RHICmdList.SetBatchedShaderUnbinds(InShader.GetComputeShader(), BatchedUnbinds);
+	}
 }
 
 /**
diff --git a/Engine/Source/Runtime/VulkanRHI/Private/VulkanCommands.cpp b/Engine/Source/Runtime/VulkanRHI/Private/VulkanCommands.cpp
--- a/Engine/Source/Runtime/VulkanRHI/Private/VulkanCommands.cpp
+++ b/Engine/Source/Runtime/VulkanRHI/Private/VulkanCommands.cpp
@@ -434,16 +434,6 @@ void FVulkanCommandListContext::RHISetShaderParameters(FRHIComputeShader* Shader
 	);
 }
 
-void FVulkanCommandListContext::RHISetShaderUnbinds(FRHIComputeShader* Shader, TConstArrayView<FRHIShaderParameterUnbind> InUnbinds)
-{
-	UE::RHICore::RHISetShaderUnbindsShared(*this, Shader, InUnbinds);
-}
-
-void FVulkanCommandListContext::RHISetShaderUnbinds(FRHIGraphicsShader* Shader, TConstArrayView<FRHIShaderParameterUnbind> InUnbinds)
-{
-	UE::RHICore::RHISetShaderUnbindsShared(*this, Shader, InUnbinds);
-}
-
 void FVulkanCommandListContext::RHISetStaticUniformBuffers(const FUniformBufferStaticBindings& InUniformBuffers)
 {
 	FMemory::Memzero(GlobalUniformBuffers.GetData(), GlobalUniformBuffers.Num() * sizeof(FRHIUniformBuffer*));
diff --git a/Engine/Source/Runtime/VulkanRHI/Private/VulkanContext.h b/Engine/Source/Runtime/VulkanRHI/Private/VulkanContext.h
--- a/Engine/Source/Runtime/VulkanRHI/Private/VulkanContext.h
+++ b/Engine/Source/Runtime/VulkanRHI/Private/VulkanContext.h
@@ -62,8 +62,6 @@ public:
 	virtual void RHISetShaderParameter(FRHIComputeShader* ComputeShader, uint32 BufferIndex, uint32 BaseIndex, uint32 NumBytes, const void* NewValue) final override;
 	virtual void RHISetShaderParameters(FRHIGraphicsShader* Shader, TConstArrayView<uint8> InParametersData, TConstArrayView<FRHIShaderParameter> InParameters, TConstArrayView<FRHIShaderParameterResource> InResourceParameters, TConstArrayView<FRHIShaderParameterResource> InBindlessParameters) final override;
 	virtual void RHISetShaderParameters(FRHIComputeShader* Shader, TConstArrayView<uint8> InParametersData, TConstArrayView<FRHIShaderParameter> InParameters, TConstArrayView<FRHIShaderParameterResource> InResourceParameters, TConstArrayView<FRHIShaderParameterResource> InBindlessParameters) final override;
-	virtual void RHISetShaderUnbinds(FRHIComputeShader* Shader, TConstArrayView<FRHIShaderParameterUnbind> InUnbinds) final override;
-	virtual void RHISetShaderUnbinds(FRHIGraphicsShader* Shader, TConstArrayView<FRHIShaderParameterUnbind> InUnbinds) final override;
 	virtual void RHISetStencilRef(uint32 StencilRef) final override;
 	virtual void RHIDrawPrimitive(uint32 BaseVertexIndex, uint32 NumPrimitives, uint32 NumInstances) final override;
 	virtual void RHIDrawPrimitiveIndirect(FRHIBuffer* ArgumentBuffer, uint32 ArgumentOffset) final override;
diff --git a/Engine/Source/Runtime/Windows/D3D11RHI/Private/D3D11Device.cpp b/Engine/Source/Runtime/Windows/D3D11RHI/Private/D3D11Device.cpp
--- a/Engine/Source/Runtime/Windows/D3D11RHI/Private/D3D11Device.cpp
+++ b/Engine/Source/Runtime/Windows/D3D11RHI/Private/D3D11Device.cpp
@@ -213,6 +213,8 @@ FD3D11DynamicRHI::FD3D11DynamicRHI(IDXGIFactory1* InDXGIFactory1, D3D_FEATURE_LE
 	GRHIMaxDispatchThreadGroupsPerDimension.Y = D3D11_CS_DISPATCH_MAX_THREAD_GROUPS_PER_DIMENSION;
 	GRHIMaxDispatchThreadGroupsPerDimension.Z = D3D11_CS_DISPATCH_MAX_THREAD_GROUPS_PER_DIMENSION;
 
+	GRHIGlobals.NeedsShaderUnbinds = true;
+
 	GRHITransitionPrivateData_SizeInBytes = sizeof(FD3D11TransitionData);
 	GRHITransitionPrivateData_AlignInBytes = alignof(FD3D11TransitionData);
 
