Check IsRayTracingAllowed() when creating or destroying FRayTracingGeometry.

diff --git a/Engine/Source/Runtime/RenderCore/Private/RayTracingGeometry.cpp b/Engine/Source/Runtime/RenderCore/Private/RayTracingGeometry.cpp
--- a/Engine/Source/Runtime/RenderCore/Private/RayTracingGeometry.cpp
+++ b/Engine/Source/Runtime/RenderCore/Private/RayTracingGeometry.cpp
@@ -254,8 +254,17 @@ void FRayTracingGeometry::RemoveBuildRequest()
 	}
 }
 
+void FRayTracingGeometry::InitResource(FRHICommandListBase& RHICmdList)
+{
+	checkf(IsRayTracingAllowed(), TEXT("FRayTracingGeometry should only be initialized when Ray Tracing is allowed"));
+
+	FRenderResource::InitResource(RHICmdList);
+}
+
 void FRayTracingGeometry::ReleaseResource()
 {
+	checkf(IsRayTracingAllowed(), TEXT("FRayTracingGeometry should only be initialized when Ray Tracing is allowed"));
+
 	// Release any resource references held by the initializer.
 	// This includes index and vertex buffers used for building the BLAS.
 	Initializer = FRayTracingGeometryInitializer{};
diff --git a/Engine/Source/Runtime/RenderCore/Public/RayTracingGeometry.h b/Engine/Source/Runtime/RenderCore/Public/RayTracingGeometry.h
--- a/Engine/Source/Runtime/RenderCore/Public/RayTracingGeometry.h
+++ b/Engine/Source/Runtime/RenderCore/Public/RayTracingGeometry.h
@@ -134,6 +134,7 @@ public:
 	RENDERCORE_API virtual void InitRHI(FRHICommandListBase& RHICmdList) override;
 	RENDERCORE_API virtual void ReleaseRHI() override;
 
+	RENDERCORE_API virtual void InitResource(FRHICommandListBase& RHICmdList) override;
 	RENDERCORE_API virtual void ReleaseResource() override;
 protected:
 	RENDERCORE_API void RemoveBuildRequest();
