Fix validation assert when attempting to validate RHI resources. Now only validate if RHI access is allowed

diff --git a/Engine/Source/Runtime/RenderCore/Private/ShaderParameterStruct.cpp b/Engine/Source/Runtime/RenderCore/Private/ShaderParameterStruct.cpp
--- a/Engine/Source/Runtime/RenderCore/Private/ShaderParameterStruct.cpp
+++ b/Engine/Source/Runtime/RenderCore/Private/ShaderParameterStruct.cpp
@@ -492,7 +492,7 @@ struct FShaderParameterReader
 
 #if DO_CHECK
 
-void ValidateShaderParameters(const TShaderRef<FShader>& Shader, const FShaderParametersMetadata* ParametersMetadata, const void* ParametersData, bool bValidateRHI)
+void ValidateShaderParameters(const TShaderRef<FShader>& Shader, const FShaderParametersMetadata* ParametersMetadata, const void* ParametersData)
 {
 	const FShaderParameterBindings& Bindings = Shader->Bindings;
 
@@ -525,7 +525,7 @@ void ValidateShaderParameters(const TShaderRef<FShader>& Shader, const FShaderPa
 				}
 
 #if ENABLE_RHI_VALIDATION
-				if (GRHIValidationEnabled && bValidateRHI)
+				if (GRHIValidationEnabled && GRDGAllowRHIAccess)
 				{
 					if (BaseType == UBMT_SRV)
 					{
@@ -568,9 +568,8 @@ void ValidateShaderParameters(const TShaderRef<FShader>& Shader, const FShaderPa
 				}
 
 #if ENABLE_RHI_VALIDATION
-				if (GRHIValidationEnabled && bValidateRHI)
+				if (GRHIValidationEnabled && GRDGAllowRHIAccess)
 				{
-					RDG_ALLOW_RHI_ACCESS_SCOPE();
 					if (BaseType == UBMT_RDG_TEXTURE_SRV || BaseType == UBMT_RDG_BUFFER_SRV)
 					{
 						const FRHIShaderResourceView* SRV = static_cast<const FRHIShaderResourceView*>(GraphResource->GetRHI());
diff --git a/Engine/Source/Runtime/RenderCore/Public/RenderGraphUtils.h b/Engine/Source/Runtime/RenderCore/Public/RenderGraphUtils.h
--- a/Engine/Source/Runtime/RenderCore/Public/RenderGraphUtils.h
+++ b/Engine/Source/Runtime/RenderCore/Public/RenderGraphUtils.h
@@ -201,7 +201,7 @@ void ClearUnusedGraphResources(
 	// Verify the shader have all the parameters it needs. This is done before the
 	// ClearUnusedGraphResourcesImpl() to not mislead user on why some resource are missing
 	// when debugging a validation failure.
-	ValidateShaderParameters(Shader, ParametersMetadata, InoutParameters, false);
+	ValidateShaderParameters(Shader, ParametersMetadata, InoutParameters);
 
 	// Clear the resources the shader won't need.
 	return ClearUnusedGraphResourcesImpl(Shader->Bindings, ParametersMetadata, InoutParameters, ExcludeList);
@@ -231,8 +231,8 @@ void ClearUnusedGraphResources(
 	// Verify the shader have all the parameters it needs. This is done before the
 	// ClearUnusedGraphResourcesImpl() to not mislead user on why some resource are missing
 	// when debugging a validation failure.
-	ValidateShaderParameters(ShaderA, ParametersMetadata, InoutParameters, false);
-	ValidateShaderParameters(ShaderB, ParametersMetadata, InoutParameters, false);
+	ValidateShaderParameters(ShaderA, ParametersMetadata, InoutParameters);
+	ValidateShaderParameters(ShaderB, ParametersMetadata, InoutParameters);
 
 	// Clear the resources the shader won't need.
 	const FShaderParameterBindings* ShaderBindings[] = { &ShaderA->Bindings, &ShaderB->Bindings };
diff --git a/Engine/Source/Runtime/RenderCore/Public/ShaderParameterStruct.h b/Engine/Source/Runtime/RenderCore/Public/ShaderParameterStruct.h
--- a/Engine/Source/Runtime/RenderCore/Public/ShaderParameterStruct.h
+++ b/Engine/Source/Runtime/RenderCore/Public/ShaderParameterStruct.h
@@ -117,15 +117,15 @@ extern RENDERCORE_API void EmitNullShaderParameterFatalError(const TShaderRef<FS
 
 /** Validates that all resource parameters of a shader are set. */
 #if DO_CHECK
-extern RENDERCORE_API void ValidateShaderParameters(const TShaderRef<FShader>& Shader, const FShaderParametersMetadata* ParametersMetadata, const void* ParametersData, bool bValidateRHI);
+extern RENDERCORE_API void ValidateShaderParameters(const TShaderRef<FShader>& Shader, const FShaderParametersMetadata* ParametersMetadata, const void* ParametersData);
 #else
-FORCEINLINE void ValidateShaderParameters(const TShaderRef<FShader>& Shader, const FShaderParametersMetadata* ParametersMetadata, const void* ParametersData, bool bValidateRHI) {}
+FORCEINLINE void ValidateShaderParameters(const TShaderRef<FShader>& Shader, const FShaderParametersMetadata* ParametersMetadata, const void* ParametersData) {}
 #endif
 
 template<typename TShaderClass>
 FORCEINLINE void ValidateShaderParameters(const TShaderRef<TShaderClass>& Shader, const typename TShaderClass::FParameters& Parameters)
 {
-	return ValidateShaderParameters(Shader, TShaderClass::FParameters::FTypeInfo::GetStructMetadata(), &Parameters, false);
+	return ValidateShaderParameters(Shader, TShaderClass::FParameters::FTypeInfo::GetStructMetadata(), &Parameters);
 }
 
 template<typename TParameterType>
@@ -238,7 +238,7 @@ inline void SetShaderParameters(
 	const FShaderParametersMetadata* ParametersMetadata,
 	const typename TShaderClass::FParameters& Parameters)
 {
-	ValidateShaderParameters(Shader, ParametersMetadata, &Parameters, false);
+	ValidateShaderParameters(Shader, ParametersMetadata, &Parameters);
 	SetShaderParameters(BatchedParameters, Shader->Bindings, ParametersMetadata, &Parameters);
 }
 
@@ -278,7 +278,7 @@ inline void SetShaderParameters(
 	const FShaderParametersMetadata* ParametersMetadata,
 	const typename TShaderClass::FParameters& Parameters)
 {
-	ValidateShaderParameters(Shader, ParametersMetadata, &Parameters, false);
+	ValidateShaderParameters(Shader, ParametersMetadata, &Parameters);
 
 	// TODO(RDG): Once all shader sets their parameter through this, can refactor RHI so all shader parameters get sets through a single RHI function call.
 
@@ -306,7 +306,7 @@ void SetShaderParameters(FRayTracingShaderBindingsWriter& RTBindingsWriter, cons
 {
 	const FShaderParametersMetadata* ParametersMetadata = TShaderClass::FParameters::FTypeInfo::GetStructMetadata();
 
-	ValidateShaderParameters(Shader, ParametersMetadata, &Parameters, true);
+	ValidateShaderParameters(Shader, ParametersMetadata, &Parameters);
 
 	checkf(Shader->Bindings.Parameters.Num() == 0, TEXT("Ray tracing shader should use SHADER_USE_ROOT_PARAMETER_STRUCT() to passdown the cbuffer layout to the shader compiler."));
 
