Change FGlobalDynamicReadBuffer to BUF_Dynamic instead of BUF_Volatile.
Only VulkanRHI has a special path for BUF_Volatile and it does not align with FGlobalDynamicReadBuffer expectations about buffer lifetime

diff --git a/Engine/Source/Runtime/RenderCore/Private/DynamicBufferAllocator.cpp b/Engine/Source/Runtime/RenderCore/Private/DynamicBufferAllocator.cpp
--- a/Engine/Source/Runtime/RenderCore/Private/DynamicBufferAllocator.cpp
+++ b/Engine/Source/Runtime/RenderCore/Private/DynamicBufferAllocator.cpp
@@ -143,7 +143,7 @@ FGlobalDynamicReadBuffer::FAllocation FGlobalDynamicReadBuffer_AllocateInternal(
 			const uint32 NewBufferSize = FMath::Max(AlignedNum, (uint32)GMinReadBufferRenderingBufferSize);
 			Buffer = new FDynamicAllocReadBuffer();
 			BufferPool->Buffers.Add(Buffer);
-			Buffer->Initialize(RHICmdList, TEXT("FGlobalDynamicReadBuffer_AllocateInternal"), sizeof(Type), NewBufferSize, Format, BUF_Volatile);
+			Buffer->Initialize(RHICmdList, TEXT("FGlobalDynamicReadBuffer_AllocateInternal"), sizeof(Type), NewBufferSize, Format, BUF_Dynamic);
 		}
 
 		// Lock the buffer if needed.
