Handle namespaced types (Namespace::Type) when parsing shader types.

diff --git a/Engine/Source/Runtime/RenderCore/Private/ShaderParameterParser.cpp b/Engine/Source/Runtime/RenderCore/Private/ShaderParameterParser.cpp
--- a/Engine/Source/Runtime/RenderCore/Private/ShaderParameterParser.cpp
+++ b/Engine/Source/Runtime/RenderCore/Private/ShaderParameterParser.cpp
@@ -631,6 +631,15 @@ bool FShaderParameterParser::ParseParameters(
 				{
 					// Keep browsing what might be type of the parameter.
 				}
+				else if (Char == ':')
+				{
+					// Handle :: in type names
+					if (UpComing[1] == ':')
+					{
+						// next loop iteration takes us to the next ':', so go past that
+						Cursor++;
+					}
+				}
 				else if (Char == '<')
 				{
 					// Found what looks like the beginning of template argument that is legal on resource types for Instance Texture2D< float >
@@ -657,9 +666,18 @@ bool FShaderParameterParser::ParseParameters(
 				{
 					// Keep browsing what might be type of the parameter.
 				}
+				else if (Char == ':')
+				{
+					// Handle :: in type names
+					if (UpComing[1] == ':')
+					{
+						// next loop iteration takes us to the next ':', so go past that
+						Cursor++;
+					}
+				}
 				else if (bIsWhiteSpace || Char == ',')
 				{
-					// Spaces and comas are legal within agrument of the template arguments.
+					// Spaces and comas are legal within argument of the template arguments.
 				}
 				else if (Char == '>')
 				{
