Hardened thread safety of render command pipe system and added ability to sync specific pipes.
 - Added additional sync scopes to handle VT standalone path.

diff --git a/Engine/Plugins/MovieScene/MovieRenderPipeline/Source/MovieRenderPipelineCore/Private/MoviePipelineRendering.cpp b/Engine/Plugins/MovieScene/MovieRenderPipeline/Source/MovieRenderPipelineCore/Private/MoviePipelineRendering.cpp
--- a/Engine/Plugins/MovieScene/MovieRenderPipeline/Source/MovieRenderPipelineCore/Private/MoviePipelineRendering.cpp
+++ b/Engine/Plugins/MovieScene/MovieRenderPipeline/Source/MovieRenderPipelineCore/Private/MoviePipelineRendering.cpp
@@ -563,6 +563,8 @@ void UMoviePipeline::FlushAsyncEngineSystems()
 		}
 	}
 
+	UE::RenderCommandPipe::FSyncScope SyncScope;
+
 	// Flush virtual texture tile calculations
 	ERHIFeatureLevel::Type FeatureLevel = GetWorld()->GetFeatureLevel();
 	ENQUEUE_RENDER_COMMAND(VirtualTextureSystemFlushCommand)(
diff --git a/Engine/Plugins/Runtime/ComputeFramework/Source/ComputeFramework/Private/ComputeFramework.cpp b/Engine/Plugins/Runtime/ComputeFramework/Source/ComputeFramework/Private/ComputeFramework.cpp
--- a/Engine/Plugins/Runtime/ComputeFramework/Source/ComputeFramework/Private/ComputeFramework.cpp
+++ b/Engine/Plugins/Runtime/ComputeFramework/Source/ComputeFramework/Private/ComputeFramework.cpp
@@ -12,6 +12,7 @@
 #include "RenderGraphBuilder.h"
 #include "ShaderCore.h"
 #include "UObject/UObjectIterator.h"
+#include "Rendering/RenderCommandPipes.h"
 
 DEFINE_LOG_CATEGORY(LogComputeFramework);
 
@@ -87,7 +88,7 @@ namespace ComputeFramework
 		FComputeGraphTaskWorker* ComputeGraphWorker = ComputeSystem != nullptr ? ComputeSystem->GetComputeWorker(InScene) : nullptr;
 		if (ComputeGraphWorker)
 		{
-			UE::RenderCommandPipe::FSyncScope SyncScope;
+			UE::RenderCommandPipe::FSyncScope SyncScope({ &UE::RenderCommandPipe::SkeletalMesh });
 
 			ENQUEUE_RENDER_COMMAND(ComputeFrameworkFlushCommand)(
 				[ComputeGraphWorker, InExecutionGroupName](FRHICommandListImmediate& RHICmdList)
diff --git a/Engine/Plugins/Runtime/MeshModelingToolset/Source/ModelingComponents/Private/AssetUtils/Texture2DUtil.cpp b/Engine/Plugins/Runtime/MeshModelingToolset/Source/ModelingComponents/Private/AssetUtils/Texture2DUtil.cpp
--- a/Engine/Plugins/Runtime/MeshModelingToolset/Source/ModelingComponents/Private/AssetUtils/Texture2DUtil.cpp
+++ b/Engine/Plugins/Runtime/MeshModelingToolset/Source/ModelingComponents/Private/AssetUtils/Texture2DUtil.cpp
@@ -322,6 +322,8 @@ bool UE::AssetUtils::ForceVirtualTexturePrefetch(FImageDimensions ScreenSpaceDim
 	{
 		const FVector2D ScreenSpaceSize(ScreenSpaceDimensions.GetWidth(), ScreenSpaceDimensions.GetHeight());
 
+		UE::RenderCommandPipe::FSyncScope SyncScope;
+
 		ENQUEUE_RENDER_COMMAND(AssetUtils_ForceVirtualTexturePrefetch)(
 			[ScreenSpaceSize](FRHICommandListImmediate& RHICmdList)
 		{
diff --git a/Engine/Source/Editor/TextureEditor/Private/Models/TextureEditorViewportClient.cpp b/Engine/Source/Editor/TextureEditor/Private/Models/TextureEditorViewportClient.cpp
--- a/Engine/Source/Editor/TextureEditor/Private/Models/TextureEditorViewportClient.cpp
+++ b/Engine/Source/Editor/TextureEditor/Private/Models/TextureEditorViewportClient.cpp
@@ -253,6 +253,8 @@ void FTextureEditorViewportClient::Draw(FViewport* Viewport, FCanvas* Canvas)
 			const FVector2D UV0 = TileItem.UV0;
 			const FVector2D UV1 = TileItem.UV1;
 
+			UE::RenderCommandPipe::FSyncScope SyncScope;
+
 			const ERHIFeatureLevel::Type InFeatureLevel = GMaxRHIFeatureLevel;
 			ENQUEUE_RENDER_COMMAND(MakeTilesResident)(
 				[InFeatureLevel, VTResource, ScreenSpaceSize, ViewportPositon, ViewportSize, UV0, UV1, MipLevel](FRHICommandListImmediate& RHICmdList)
diff --git a/Engine/Source/Editor/UnrealEd/Private/Lightmass/LightmassRender.cpp b/Engine/Source/Editor/UnrealEd/Private/Lightmass/LightmassRender.cpp
--- a/Engine/Source/Editor/UnrealEd/Private/Lightmass/LightmassRender.cpp
+++ b/Engine/Source/Editor/UnrealEd/Private/Lightmass/LightmassRender.cpp
@@ -1316,6 +1316,8 @@ bool FLightmassMaterialRenderer::GenerateMaterialPropertyData(
 					const ERHIFeatureLevel::Type FeatureLevel = GMaxRHIFeatureLevel;
 					const FVector2D ScreenSpaceSize(InOutSizeX, InOutSizeY);
 
+					UE::RenderCommandPipe::FSyncScope SyncScope;
+
 					ENQUEUE_RENDER_COMMAND(LoadTiles)(
 						[FeatureLevel, ScreenSpaceSize](FRHICommandListImmediate& RHICmdList)
 					{
diff --git a/Engine/Source/Editor/WorldPartitionEditor/Private/WorldPartition/SWorldPartitionEditorGrid2D.cpp b/Engine/Source/Editor/WorldPartitionEditor/Private/WorldPartition/SWorldPartitionEditorGrid2D.cpp
--- a/Engine/Source/Editor/WorldPartitionEditor/Private/WorldPartition/SWorldPartitionEditorGrid2D.cpp
+++ b/Engine/Source/Editor/WorldPartitionEditor/Private/WorldPartition/SWorldPartitionEditorGrid2D.cpp
@@ -1967,6 +1967,8 @@ int32 SWorldPartitionEditorGrid2D::PaintMinimap(const FGeometry& AllottedGeometr
 			const ERHIFeatureLevel::Type InFeatureLevel = GMaxRHIFeatureLevel;
 			const int32 MipLevel = -1;
 
+			UE::RenderCommandPipe::FSyncScope SyncScope;
+
 			ENQUEUE_RENDER_COMMAND(MakeTilesResident)(
 				[InFeatureLevel, VTResource, ScreenSpaceSize, ViewportPositon, ViewportSize, UV0, UV1, MipLevel](FRHICommandListImmediate& RHICmdList)
 			{
diff --git a/Engine/Source/Runtime/Engine/Private/GPUSkinVertexFactory.cpp b/Engine/Source/Runtime/Engine/Private/GPUSkinVertexFactory.cpp
--- a/Engine/Source/Runtime/Engine/Private/GPUSkinVertexFactory.cpp
+++ b/Engine/Source/Runtime/Engine/Private/GPUSkinVertexFactory.cpp
@@ -17,6 +17,7 @@
 #include "RenderUtils.h"
 #include "ShaderPlatformCachedIniValue.h"
 #include "Engine/RendererSettings.h"
+#include "Rendering/RenderCommandPipes.h"
 
 #if INTEL_ISPC
 #include "GPUSkinVertexFactory.ispc.generated.h"
@@ -242,7 +243,7 @@ static TAutoConsoleVariable<int32> CVarRHICmdDeferSkeletalLockAndFillToRHIThread
 
 static bool DeferSkeletalLockAndFillToRHIThread()
 {
-	return IsRunningRHIInSeparateThread() && CVarRHICmdDeferSkeletalLockAndFillToRHIThread.GetValueOnRenderThread() > 0 && !UE::RenderCommandPipe::IsReplaying();
+	return IsRunningRHIInSeparateThread() && CVarRHICmdDeferSkeletalLockAndFillToRHIThread.GetValueOnRenderThread() > 0 && !UE::RenderCommandPipe::SkeletalMesh.IsReplaying();
 }
 
 bool FGPUBaseSkinVertexFactory::FShaderDataType::UpdateBoneData(FRHICommandList& RHICmdList, const TArray<FMatrix44f>& ReferenceToLocalMatrices,
diff --git a/Engine/Source/Runtime/Engine/Private/LevelTick.cpp b/Engine/Source/Runtime/Engine/Private/LevelTick.cpp
--- a/Engine/Source/Runtime/Engine/Private/LevelTick.cpp
+++ b/Engine/Source/Runtime/Engine/Private/LevelTick.cpp
@@ -975,8 +975,6 @@ void BeginSendEndOfFrameUpdatesDrawEvent(FSendAllEndOfFrameUpdates& SendAllEndOf
 {
 	BEGIN_DRAW_EVENTF_GAMETHREAD(SendAllEndOfFrameUpdates, SendAllEndOfFrameUpdates.DrawEvent, TEXT("SendAllEndOfFrameUpdates"));
 
-	UE::RenderCommandPipe::StartParallelRecordingValidation();
-
 	ENQUEUE_RENDER_COMMAND(BeginDrawEventCommand)(UE::RenderCommandPipe::SkeletalMesh,
 		[GPUSkinCache = SendAllEndOfFrameUpdates.GPUSkinCache]
 	{
@@ -1000,8 +998,6 @@ void EndSendEndOfFrameUpdatesDrawEvent(FSendAllEndOfFrameUpdates& SendAllEndOfFr
 		}
 	});
 
-	UE::RenderCommandPipe::StopParallelRecordingValidation();
-
 	STOP_DRAW_EVENT_GAMETHREAD(SendAllEndOfFrameUpdates.DrawEvent);
 }
 
@@ -1023,6 +1019,8 @@ void UWorld::SendAllEndOfFrameUpdates()
 		return;
 	}
 
+	UE::RenderCommandPipe::StartRecording();
+
 	// Wait for tasks that are generating data for the render proxies, but are not awaited in any TickFunctions 
 	// E.g., see cloth USkeletalMeshComponent::UpdateClothStateAndSimulate
 	for (UActorComponent* Component : ComponentsThatNeedPreEndOfFrameSync)
@@ -1165,6 +1163,8 @@ void UWorld::SendAllEndOfFrameUpdates()
 	LocalComponentsThatNeedEndOfFrameUpdate.Reset();
 
 	EndSendEndOfFrameUpdatesDrawEvent(SendAllEndOfFrameUpdates);
+
+	UE::RenderCommandPipe::StopRecording();
 }
 
 /**
diff --git a/Engine/Source/Runtime/Engine/Private/UserInterface/Canvas.cpp b/Engine/Source/Runtime/Engine/Private/UserInterface/Canvas.cpp
--- a/Engine/Source/Runtime/Engine/Private/UserInterface/Canvas.cpp
+++ b/Engine/Source/Runtime/Engine/Private/UserInterface/Canvas.cpp
@@ -559,8 +559,6 @@ FCanvasRenderThreadScope::~FCanvasRenderThreadScope()
 {
 	RenderCommandFunctionArray* RenderCommandArray = RenderCommands;
 
-	UE::RenderCommandPipe::FSyncScope SyncScope;
-
 	ENQUEUE_RENDER_COMMAND(DispatchCanvasRenderCommands)(
 		[RenderCommandArray, RenderTarget = Canvas.RenderTarget, ViewRect = Canvas.ViewRect, ScissorRect = Canvas.ScissorRect, bScaledToRenderTarget = Canvas.bScaledToRenderTarget](FRHICommandListImmediate& RHICmdList) mutable
 	{
diff --git a/Engine/Source/Runtime/RenderCore/Private/RenderingThread.cpp b/Engine/Source/Runtime/RenderCore/Private/RenderingThread.cpp
--- a/Engine/Source/Runtime/RenderCore/Private/RenderingThread.cpp
+++ b/Engine/Source/Runtime/RenderCore/Private/RenderingThread.cpp
@@ -1295,7 +1295,7 @@ void FlushRenderingCommands()
 		FTaskGraphInterface::Get().ProcessThreadUntilIdle(ENamedThreads::GameThread_Local);
 	}
 
-	UE::RenderCommandPipe::FSyncScope SyncScope;
+	UE::RenderCommandPipe::StopRecording();
 
 	ENQUEUE_RENDER_COMMAND(FlushPendingDeleteRHIResourcesCmd)([](FRHICommandListImmediate& RHICmdList)
 	{
@@ -1372,6 +1372,7 @@ FPendingCleanupObjects::~FPendingCleanupObjects()
 	if (CleanupArray.Num())
 	{
 		QUICK_SCOPE_CYCLE_COUNTER(STAT_FPendingCleanupObjects_Destruct);
+		UE::RenderCommandPipe::StopRecording();
 
 		const bool bBatchingEnabled = bEnablePendingCleanupObjectsCommandBatching;
 		if (bBatchingEnabled)
@@ -1410,20 +1411,24 @@ FPendingCleanupObjects::FPendingCleanupObjects()
 
 FPendingCleanupObjects::~FPendingCleanupObjects()
 {
-	QUICK_SCOPE_CYCLE_COUNTER(STAT_FPendingCleanupObjects_Destruct);
-
-	const bool bBatchingEnabled = bEnablePendingCleanupObjectsCommandBatching;
-	if (bBatchingEnabled)
-	{
-		StartRenderCommandFenceBundler();
-	}
-	for (int32 ObjectIndex = 0; ObjectIndex < CleanupArray.Num(); ObjectIndex++)
-	{
-		delete CleanupArray[ObjectIndex];
-	}
-	if (bBatchingEnabled)
+	if (CleanupArray.Num())
 	{
-		StopRenderCommandFenceBundler();
+		QUICK_SCOPE_CYCLE_COUNTER(STAT_FPendingCleanupObjects_Destruct);
+		UE::RenderCommandPipe::StopRecording();
+
+		const bool bBatchingEnabled = bEnablePendingCleanupObjectsCommandBatching;
+		if (bBatchingEnabled)
+		{
+			StartRenderCommandFenceBundler();
+		}
+		for (int32 ObjectIndex = 0; ObjectIndex < CleanupArray.Num(); ObjectIndex++)
+		{
+			delete CleanupArray[ObjectIndex];
+		}
+		if (bBatchingEnabled)
+		{
+			StopRenderCommandFenceBundler();
+		}
 	}
 }
 
@@ -1585,220 +1590,304 @@ void FRenderThreadCommandPipe::EnqueueAndLaunch(const TCHAR* Name, uint32& SpecI
 	}
 }
 
+#if !UE_SERVER
+
 class FRenderCommandPipeRegistry
 {
 public:
-	TLinkedList<FRenderCommandPipe*>* GlobalList = nullptr;
+	TLinkedList<FRenderCommandPipe*>*& GetGlobalList()
+	{
+		static TLinkedList<FRenderCommandPipe*>* GlobalList = nullptr;
+		return GlobalList;
+	}
 
 	void Initialize()
 	{
-#if !UE_SERVER
-		Pipes.Reset();
+		AllPipes.Reset();
 
-		for (TLinkedList<FRenderCommandPipe*>::TIterator PipeIt(GlobalList); PipeIt; PipeIt.Next())
+		for (TLinkedList<FRenderCommandPipe*>::TIterator PipeIt(GetGlobalList()); PipeIt; PipeIt.Next())
 		{
 			FRenderCommandPipe* Pipe = *PipeIt;
 			Pipe->SetEnabled(Pipe->ConsoleVariable->GetBool());
+			Pipe->Index = AllPipes.Num();
 
-			Pipes.Emplace(*PipeIt);
+			AllPipes.Emplace(*PipeIt);
 		}
-#endif // !UE_SERVER
 
 		GRenderCommandPipeMode = GetValidatedRenderCommandPipeMode(CVarRenderCommandPipeMode->GetInt());
 	}
 
-	void StartRecording()
+	void StartRecording(TConstArrayView<FRenderCommandPipe*> Pipes)
 	{
-#if !UE_SERVER
-		check(IsInGameThread());
-		ensureMsgf(!bParallelRecordingValidation, TEXT("StartRecording has been called within a region designated for issuing parallel render thread commands. This is a race condition."));
 		SCOPED_NAMED_EVENT(FRenderCommandPipe_StartRecording, FColor::Magenta);
 
-		if (bRecording || UE::RenderCommandPipe::GetMode() != ERenderCommandPipeMode::All || Pipes.IsEmpty() || !ShouldExecuteOnRenderThread())
+		if (GRenderCommandPipeMode != ERenderCommandPipeMode::All)
 		{
 			return;
 		}
 
-		UE::Tasks::FTaskEvent TaskEvent{ UE_SOURCE_LOCATION };
-		TArray<FRenderCommandPipe::FFrame*> NextFrames;
-		NextFrames.SetNumZeroed(Pipes.Num());
-		bool bHasNextFrames = false;
+		bool bAllPipes = false;
+		if (Pipes.IsEmpty())
+		{
+			bAllPipes = true;
+			Pipes = AllPipes;
+		}
 
-		for (int32 PipeIndex = 0; PipeIndex < Pipes.Num(); ++PipeIndex)
+		if (Pipes.IsEmpty() || !ShouldExecuteOnRenderThread())
 		{
-			FRenderCommandPipe& Pipe = *Pipes[PipeIndex];
+			return;
+		}
+
+		UE::TScopeLock Lock(Mutex);
+
+		if (!bAllPipes)
+		{
+			CheckUnique(Pipes);
+		}
+
+		bool bAnyPipesToStartRecording = false;
 
-			if (Pipe.bEnabled)
+		for (FRenderCommandPipe* Pipe : Pipes)
+		{
+			if (Pipe->bEnabled && !Pipe->bRecording)
 			{
-				Pipe.Frame_GameThread = new FRenderCommandPipe::FFrame(Pipe.Name, TaskEvent);
-				NextFrames[PipeIndex] = Pipe.Frame_GameThread;
-				bHasNextFrames = true;
+				bAnyPipesToStartRecording = true;
+				break;
 			}
 		}
 
-		if (bHasNextFrames)
+		if (!bAnyPipesToStartRecording)
 		{
-			ENQUEUE_RENDER_COMMAND(RenderCommandPipe_Start)([this, TaskEvent, NextFrames = MoveTemp(NextFrames)](FRHICommandListImmediate&) mutable
-			{
-				RHIResourceLifetimeAddRef();
+			return;
+		}
 
-				for (int32 PipeIndex = 0; PipeIndex < Pipes.Num(); ++PipeIndex)
-				{
-					Pipes[PipeIndex]->Frame_RenderThread = NextFrames[PipeIndex];
-				}
+		UE::Tasks::FTaskEvent TaskEvent{ UE_SOURCE_LOCATION };
 
-				bReplaying = true;
-				TaskEvent.Trigger();
-			});
+		struct FPipeToStartRecording
+		{
+			FPipeToStartRecording(FRenderCommandPipe* InPipe, FRenderCommandPipe::FFrame* InFrame)
+				: Pipe(InPipe)
+				, Frame(InFrame)
+			{}
+
+			FRenderCommandPipe* Pipe;
+			FRenderCommandPipe::FFrame* Frame;
+		};
+
+		TArray<FPipeToStartRecording, FConcurrentLinearArrayAllocator> PipesToStartRecording;
+		PipesToStartRecording.Reserve(Pipes.Num());
+
+		for (FRenderCommandPipe* Pipe : Pipes)
+		{
+			if (Pipe->bEnabled && !Pipe->bRecording)
+			{
+				Pipe->bRecording = true;
+
+				FRenderCommandPipe::FFrame* NextFrame = new FRenderCommandPipe::FFrame(Pipe->Name, TaskEvent);
+				PipesToStartRecording.Emplace(Pipe, NextFrame);
 
-			bRecording = true;
+				UE::TScopeLock PipeLock(Pipe->Mutex);
+				Pipe->Frame_GameThread = NextFrame;
+			}
 		}
-		else
+
+		NumPipesRecording += PipesToStartRecording.Num();
+
+		ENQUEUE_RENDER_COMMAND(RenderCommandPipe_Start)([this, TaskEvent, PipesToStartRecording = MoveTemp(PipesToStartRecording)](FRHICommandListImmediate&) mutable
 		{
+			RHIResourceLifetimeAddRef(PipesToStartRecording.Num());
+
+			for (FPipeToStartRecording Pipe : PipesToStartRecording)
+			{
+				Pipe.Pipe->Frame_RenderThread = Pipe.Frame;
+			}
+
+			NumPipesReplaying += PipesToStartRecording.Num();
 			TaskEvent.Trigger();
-		}
-#endif // !UE_SERVER
+		});
 	}
 
-	void StopRecording()
+	void StopRecording(TConstArrayView<FRenderCommandPipe*> Pipes)
 	{
-#if !UE_SERVER
-		check(IsInGameThread());
-		ensureMsgf(!bParallelRecordingValidation, TEXT("StopRecording has been called within a region designated for issuing parallel render thread commands. This is a race condition."));
 		SCOPED_NAMED_EVENT(FRenderCommandPipe_StopRecording, FColor::Magenta);
 
-		if (!bRecording)
+		bool bAllPipes = false;
+
+		if (Pipes.IsEmpty())
+		{
+			bAllPipes = true;
+			Pipes = AllPipes;
+		}
+
+		if (Pipes.IsEmpty())
 		{
 			return;
 		}
 
+		UE::TScopeLock Lock(Mutex);
+
+		if (!bAllPipes)
+		{
+			CheckUnique(Pipes);
+		}
+
+		bool bAnyPipesToStopRecording = false;
+
 		for (FRenderCommandPipe* Pipe : Pipes)
 		{
-			Pipe->Frame_GameThread = nullptr;
+			if (Pipe->bRecording)
+			{
+				bAnyPipesToStopRecording = true;
+				break;
+			}
 		}
 
-		ENQUEUE_RENDER_COMMAND(RenderCommandPipe_Stop)([this](FRHICommandListImmediate& RHICmdList)
+		if (!bAnyPipesToStopRecording)
 		{
-			TArray<FRHICommandListImmediate::FQueuedCommandList> QueuedCommandLists;
-			QueuedCommandLists.Reserve(Pipes.Num());
+			return;
+		}
+
+		TArray<FRenderCommandPipe*, FConcurrentLinearArrayAllocator> PipesToStopRecording;
+		PipesToStopRecording.Reserve(Pipes.Num());
 
-			for (FRenderCommandPipe* Pipe : Pipes)
+		for (FRenderCommandPipe* Pipe : Pipes)
+		{
+			if (Pipe->bRecording)
 			{
-				FRenderCommandPipe::FFrame*& Frame_RenderThread = Pipe->Frame_RenderThread;
+				Pipe->bRecording = false;
+				PipesToStopRecording.Emplace(Pipe);
 
-				if (Frame_RenderThread)
-				{
-					Frame_RenderThread->Pipe.WaitUntilEmpty();
+				UE::TScopeLock PipeLock(Pipe->Mutex);
+				Pipe->Frame_GameThread = nullptr;
+			}
+		}
 
-					if (Frame_RenderThread->RHICmdList)
-					{
-						Frame_RenderThread->RHICmdList->FinishRecording();
-						QueuedCommandLists.Emplace(Frame_RenderThread->RHICmdList);
-					}
+		NumPipesRecording -= PipesToStopRecording.Num();
+
+		ENQUEUE_RENDER_COMMAND(RenderCommandPipe_Stop)([this, PipesToStopRecording = MoveTemp(PipesToStopRecording)](FRHICommandListImmediate& RHICmdList)
+		{
+			TArray<FRHICommandListImmediate::FQueuedCommandList, FConcurrentLinearArrayAllocator> QueuedCommandLists;
+			QueuedCommandLists.Reserve(PipesToStopRecording.Num());
+
+			for (FRenderCommandPipe* Pipe : PipesToStopRecording)
+			{
+				FRenderCommandPipe::FFrame*& Frame_RenderThread = Pipe->Frame_RenderThread;
+				check(Frame_RenderThread);
+				Frame_RenderThread->Pipe.WaitUntilEmpty();
 
-					delete Frame_RenderThread;
-					Frame_RenderThread = nullptr;
+				if (Frame_RenderThread->RHICmdList)
+				{
+					Frame_RenderThread->RHICmdList->FinishRecording();
+					QueuedCommandLists.Emplace(Frame_RenderThread->RHICmdList);
 				}
+
+				delete Frame_RenderThread;
+				Frame_RenderThread = nullptr;
 			}
 
+			NumPipesReplaying -= PipesToStopRecording.Num();
+
 			RHICmdList.QueueAsyncCommandListSubmit(QueuedCommandLists);
-			RHIResourceLifetimeReleaseRef(RHICmdList);
-			bReplaying = false;
+			RHIResourceLifetimeReleaseRef(RHICmdList, PipesToStopRecording.Num());
 		});
-
-		bRecording = false;
-#endif // !UE_SERVER
-	}
-
-	void StartParallelRecordingValidation()
-	{
-#if !UE_SERVER
-		check(!bParallelRecordingValidation);
-		bParallelRecordingValidation = true;
-#endif
-	}
-
-	void StopParallelRecordingValidation()
-	{
-#if !UE_SERVER
-		check(bParallelRecordingValidation);
-		bParallelRecordingValidation = false;
-#endif
 	}
 
 	bool IsRecording() const
 	{
-#if !UE_SERVER
 		ensureMsgf(!FTaskTagScope::IsCurrentTag(ETaskTag::EParallelRenderingThread) && !FTaskTagScope::IsCurrentTag(ETaskTag::ERenderingThread),
 			TEXT("IsRecording() is not valid from the render thread timeline."));
 
-		return bRecording;
-#else
-		return false;
-#endif // !UE_SERVER
+		return NumPipesRecording > 0;
 	}
 
 	bool IsReplaying() const
 	{
-#if !UE_SERVER
-		ensureMsgf(IsInParallelRenderingThread(), TEXT("IsReplaying() is only valid from the render thread timeline."));
-		return bReplaying;
-#else
-		return false;
-#endif // !UE_SERVER
+		ensure(IsInParallelRenderingThread());
+		return NumPipesReplaying > 0;
 	}
 
 private:
-#if !UE_SERVER
-	TArray<FRenderCommandPipe*> Pipes;
-	bool bRecording = false;
-	bool bReplaying = false;
-	bool bParallelRecordingValidation = false;
-#endif // !UE_SERVER
+	void CheckUnique(TConstArrayView<FRenderCommandPipe*> Pipes)
+	{
+#if DO_CHECK
+		CheckUniqueBits.Init(false, Pipes.Num());
+		for (FRenderCommandPipe* Pipe : Pipes)
+		{
+			checkf(Pipe->Index < AllPipes.Num(), TEXT("RenderCommandPipe %s has an invalid index and was likely constructed after system initialization."), Pipe->Name);
+			checkf(!CheckUniqueBits[Pipe->Index], TEXT("RenderCommandPipe %s was supplied more than once. This is invalid. Only a unique set of pipes is allowed."), Pipe->Name);
+			CheckUniqueBits[Pipe->Index] = true;
+		}
+		CheckUniqueBits.Reset();
+#endif
+	}
+
+	UE::FMutex Mutex;
+	TArray<FRenderCommandPipe*> AllPipes;
+	std::atomic_uint32_t NumPipesRecording = 0;
+	uint32 NumPipesReplaying = 0;
+
+#if DO_CHECK
+	TBitArray<> CheckUniqueBits;
+#endif
 };
 
 static FRenderCommandPipeRegistry GRenderCommandPipeRegistry;
 
+#endif // !UE_SERVER
+
 namespace UE::RenderCommandPipe
 {
 	void Initialize()
 	{
+#if !UE_SERVER
 		GRenderCommandPipeRegistry.Initialize();
-	}
-
-	ERenderCommandPipeMode GetMode()
-	{
-		return GRenderCommandPipeMode;
+#endif
 	}
 
 	bool IsRecording()
 	{
+#if !UE_SERVER
 		return GRenderCommandPipeRegistry.IsRecording();
+#else
+		return false;
+#endif
 	}
 
 	bool IsReplaying()
 	{
+#if !UE_SERVER
 		return GRenderCommandPipeRegistry.IsReplaying();
+#else
+		return false;
+#endif
 	}
 
 	void StartRecording()
 	{
-		GRenderCommandPipeRegistry.StartRecording();
+#if !UE_SERVER
+		GRenderCommandPipeRegistry.StartRecording({});
+#endif
 	}
 
-	void StopRecording()
+	void StartRecording(TConstArrayView<FRenderCommandPipe*> Pipes)
 	{
-		GRenderCommandPipeRegistry.StopRecording();
+#if !UE_SERVER
+		GRenderCommandPipeRegistry.StartRecording(Pipes);
+#endif
 	}
 
-	void StartParallelRecordingValidation()
+	void StopRecording()
 	{
-		GRenderCommandPipeRegistry.StartParallelRecordingValidation();
+#if !UE_SERVER
+		GRenderCommandPipeRegistry.StopRecording({});
+#endif
 	}
 
-	void StopParallelRecordingValidation()
+	void StopRecording(TConstArrayView<FRenderCommandPipe*> Pipes)
 	{
-		GRenderCommandPipeRegistry.StopParallelRecordingValidation();
+#if !UE_SERVER
+		GRenderCommandPipeRegistry.StopRecording(Pipes);
+#endif
 	}
 }
 
@@ -1810,7 +1899,7 @@ FRenderCommandPipe::FRenderCommandPipe(const TCHAR* InName, ERenderCommandPipeFl
 		SetEnabled(Variable->GetBool());
 	}))
 {
-	GlobalListLink.LinkHead(GRenderCommandPipeRegistry.GlobalList);
+	GlobalListLink.LinkHead(GRenderCommandPipeRegistry.GetGlobalList());
 }
 
 FRenderCommandPipe::~FRenderCommandPipe()
@@ -1823,51 +1912,49 @@ FRenderCommandPipe::~FRenderCommandPipe()
 
 void FRenderCommandPipe::EnqueueAndLaunch(FFunctionVariant&& FunctionVariant, const TCHAR* CommandName, uint32& CommandSpecId)
 {
-	Frame_GameThread->QueueMutex.Lock();
 	bool bWasEmpty = Frame_GameThread->Queue.IsEmpty();
 	Frame_GameThread->Queue.Emplace(MoveTemp(FunctionVariant), CommandName, CommandSpecId);
-	Frame_GameThread->QueueMutex.Unlock();
 
-	if (!bWasEmpty)
-	{
-		return;
-	}
-
-	Frame_GameThread->Pipe.Launch(Name, [this]
+	if (bWasEmpty)
 	{
-		check(Frame_RenderThread);
-		SCOPED_NAMED_EVENT_TCHAR(Name, FColor::Magenta);
-		FOptionalTaskTagScope Scope(ETaskTag::EParallelRenderingThread);
+		SCOPED_NAMED_EVENT(FRenderCommandPipe_EnqueueAndLaunch_LaunchTask, FColor::Magenta);
 
-		TArray<FCommand> PoppedQueue;
+		Frame_GameThread->Pipe.Launch(Name, [this]
+		{
+			check(Frame_RenderThread);
+			SCOPED_NAMED_EVENT_TCHAR(Name, FColor::Magenta);
+			FOptionalTaskTagScope Scope(ETaskTag::EParallelRenderingThread);
 
-		Frame_RenderThread->QueueMutex.Lock();
-		PoppedQueue = MoveTemp(Frame_RenderThread->Queue);
-		Frame_RenderThread->Queue.Reserve(128);
-		Frame_RenderThread->QueueMutex.Unlock();
+			TArray<FCommand> PoppedQueue;
 
-		for (FCommand& Command : PoppedQueue)
-		{
-			TRACE_CPUPROFILER_EVENT_SCOPE_USE_ON_CHANNEL(*Command.SpecId, Command.Name, EventScope, RenderCommandsChannel, true);
+			Mutex.Lock();
+			PoppedQueue = MoveTemp(Frame_RenderThread->Queue);
+			Frame_RenderThread->Queue.Reserve(128);
+			Mutex.Unlock();
 
-			if (FCommandListFunction* Function = Command.Function.TryGet<FCommandListFunction>())
+			for (FCommand& Command : PoppedQueue)
 			{
-				if (!Frame_RenderThread->RHICmdList)
+				TRACE_CPUPROFILER_EVENT_SCOPE_USE_ON_CHANNEL(*Command.SpecId, Command.Name, EventScope, RenderCommandsChannel, true);
+
+				if (FCommandListFunction* Function = Command.Function.TryGet<FCommandListFunction>())
+				{
+					if (!Frame_RenderThread->RHICmdList)
+					{
+						FRHICommandList* RHICmdList = new FRHICommandList(FRHIGPUMask::All());
+						RHICmdList->SwitchPipeline(ERHIPipeline::Graphics);
+						Frame_RenderThread->RHICmdList = RHICmdList;
+					}
+
+					(*Function)(*Frame_RenderThread->RHICmdList);
+				}
+				else
 				{
-					FRHICommandList* RHICmdList = new FRHICommandList(FRHIGPUMask::All());
-					RHICmdList->SwitchPipeline(ERHIPipeline::Graphics);
-					Frame_RenderThread->RHICmdList = RHICmdList;
+					Command.Function.Get<FEmptyFunction>()();
 				}
 
-				(*Function)(*Frame_RenderThread->RHICmdList);
-			}
-			else
-			{
-				Command.Function.Get<FEmptyFunction>()();
+				Command.Function = {};
 			}
 
-			Command.Function = {};
-		}
-
-	}, Frame_GameThread->TaskEvent);
+		}, Frame_GameThread->TaskEvent);
+	}
 }
\ No newline at end of file
diff --git a/Engine/Source/Runtime/RenderCore/Public/RenderingThread.h b/Engine/Source/Runtime/RenderCore/Public/RenderingThread.h
--- a/Engine/Source/Runtime/RenderCore/Public/RenderingThread.h
+++ b/Engine/Source/Runtime/RenderCore/Public/RenderingThread.h
@@ -294,31 +294,26 @@ enum class ERenderCommandPipeFlags : uint8
 
 ENUM_CLASS_FLAGS(ERenderCommandPipeFlags);
 
+class FRenderCommandPipe;
+
 namespace UE::RenderCommandPipe
 {
 	// [Game Thread] Initializes all statically initialized render command pipes.
 	extern RENDERCORE_API void Initialize();
 
-	// [Game Thread] Returns the active render command pipe mode.
-	extern RENDERCORE_API ERenderCommandPipeMode GetMode();
-
-	// [Game Thread (Parallel)] Returns whether render command pipes are currently recording on the game thread timeline.
+	// [Game Thread (Parallel)] Returns whether any render command pipes are currently recording on the game thread timeline.
 	extern RENDERCORE_API bool IsRecording();
 
-	// [Render Thread (Parallel)] Returns whether render command pipes are currently replaying commands on the render thread timeline.
+	// [Render Thread (Parallel)] Returns whether any render command pipes are currently replaying commands on the render thread timeline.
 	extern RENDERCORE_API bool IsReplaying();
 
-	// [Game Thread] Starts recording render commands into pipes.
+	// [Game Thread] Starts recording render commands into pipes. Returns whether the operation succeeded.
 	extern RENDERCORE_API void StartRecording();
+	extern RENDERCORE_API void StartRecording(TConstArrayView<FRenderCommandPipe*> Pipes);
 
-	// [Game Thread] Stops recording commands into pipes and syncs all remaining pipe work to the render thread.
+	// [Game Thread] Stops recording commands into pipes and syncs all remaining pipe work to the render thread. Returns whether the operation succeeded.
 	extern RENDERCORE_API void StopRecording();
-
-	// [Game Thread] Enables validation that will issue an ensure if recording is stopped.
-	extern RENDERCORE_API void StartParallelRecordingValidation();
-
-	// [Game Thread] Disables validation that will issue an ensure if recording is stopped.
-	extern RENDERCORE_API void StopParallelRecordingValidation();
+	extern RENDERCORE_API void StopRecording(TConstArrayView<FRenderCommandPipe*> Pipes);
 
 	// [Game Thread] Stops render command pipe recording during the duration of the scope and restarts recording once the scope is complete.
 	class FSyncScope
@@ -326,33 +321,41 @@ namespace UE::RenderCommandPipe
 	public:
 		FSyncScope()
 		{
-			if (IsRecording())
-			{
-				StopRecording();
-				bWasRecording = true;
-			}
+			StopRecording();
+		}
+
+		FSyncScope(TConstArrayView<FRenderCommandPipe*> InPipes)
+			: Pipes(InPipes)
+		{
+			StopRecording(Pipes);
 		}
 
 		~FSyncScope()
 		{
-			if (bWasRecording)
+			if (!Pipes.IsEmpty())
+			{
+				StartRecording(Pipes);
+			}
+			else
 			{
 				StartRecording();
 			}
 		}
 
 	private:
-		bool bWasRecording = false;
+		TConstArrayView<FRenderCommandPipe*> Pipes;
 	};
 }
 
+extern RENDERCORE_API ERenderCommandPipeMode GRenderCommandPipeMode;
+
 class FRenderThreadCommandPipe
 {
 public:
 	template <typename RenderCommandTag, typename LambdaType>
 	FORCEINLINE_DEBUGGABLE static void Enqueue(LambdaType&& Lambda)
 	{
-		if (UE::RenderCommandPipe::GetMode() != ERenderCommandPipeMode::None)
+		if (GRenderCommandPipeMode != ERenderCommandPipeMode::None)
 		{
 			Instance.EnqueueAndLaunch(RenderCommandTag::GetName(), RenderCommandTag::GetSpecId(), RenderCommandTag::GetStatId(), MoveTemp(Lambda));
 		}
@@ -428,6 +431,17 @@ public:
 		return Name;
 	}
 
+	FORCEINLINE bool IsReplaying() const
+	{
+		ensure(IsInParallelRenderingThread());
+		return Frame_RenderThread != nullptr;
+	}
+
+	FORCEINLINE bool IsRecording() const
+	{
+		return bRecording;
+	}
+
 	void SetEnabled(bool bInIsEnabled)
 	{
 		check(IsInGameThread());
@@ -437,53 +451,65 @@ public:
 	template <typename RenderCommandTag>
 	static void Enqueue(FRenderCommandPipe* Pipe, FCommandListFunction&& Function)
 	{
-		if (Pipe && Pipe->IsRecording())
-		{
-			Pipe->EnqueueAndLaunch(MoveTemp(Function), RenderCommandTag::GetName(), RenderCommandTag::GetSpecId());
-		}
-		else
+		if (GRenderCommandPipeMode == ERenderCommandPipeMode::All && Pipe)
 		{
-			EnqueueUniqueRenderCommand<RenderCommandTag>([Function = MoveTemp(Function)](FRHICommandListImmediate& RHICmdList) { Function(RHICmdList); });
+			UE::TScopeLock Lock(Pipe->Mutex);
+			if (Pipe->Frame_GameThread)
+			{
+				Pipe->EnqueueAndLaunch(MoveTemp(Function), RenderCommandTag::GetName(), RenderCommandTag::GetSpecId());
+				return;
+			}
 		}
+
+		EnqueueUniqueRenderCommand<RenderCommandTag>([Function = MoveTemp(Function)](FRHICommandListImmediate& RHICmdList) { Function(RHICmdList); });
 	}
 
 	template <typename RenderCommandTag>
 	static void Enqueue(FRenderCommandPipe& Pipe, FCommandListFunction&& Function)
 	{
-		if (Pipe.IsRecording())
-		{
-			Pipe.EnqueueAndLaunch(MoveTemp(Function), RenderCommandTag::GetName(), RenderCommandTag::GetSpecId());
-		}
-		else
+		if (GRenderCommandPipeMode == ERenderCommandPipeMode::All)
 		{
-			EnqueueUniqueRenderCommand<RenderCommandTag>([Function = MoveTemp(Function)](FRHICommandListImmediate& RHICmdList) { Function(RHICmdList); });
+			UE::TScopeLock Lock(Pipe.Mutex);
+			if (Pipe.Frame_GameThread)
+			{
+				Pipe.EnqueueAndLaunch(MoveTemp(Function), RenderCommandTag::GetName(), RenderCommandTag::GetSpecId());
+				return;
+			}
 		}
+		
+		EnqueueUniqueRenderCommand<RenderCommandTag>([Function = MoveTemp(Function)](FRHICommandListImmediate& RHICmdList) { Function(RHICmdList); });
 	}
 
 	template <typename RenderCommandTag>
 	static void Enqueue(FRenderCommandPipe* Pipe, FEmptyFunction&& Function)
 	{
-		if (Pipe && Pipe->IsRecording())
+		if (GRenderCommandPipeMode == ERenderCommandPipeMode::All && Pipe)
 		{
-			Pipe->EnqueueAndLaunch(MoveTemp(Function), RenderCommandTag::GetName(), RenderCommandTag::GetSpecId());
-		}
-		else
-		{
-			EnqueueUniqueRenderCommand<RenderCommandTag>([Function = MoveTemp(Function)](FRHICommandListImmediate&) { Function(); });
+			UE::TScopeLock Lock(Pipe->Mutex);
+			if (Pipe->Frame_GameThread)
+			{
+				Pipe->EnqueueAndLaunch(MoveTemp(Function), RenderCommandTag::GetName(), RenderCommandTag::GetSpecId());
+				return;
+			}
 		}
+
+		EnqueueUniqueRenderCommand<RenderCommandTag>([Function = MoveTemp(Function)](FRHICommandListImmediate&) { Function(); });
 	}
 
 	template <typename RenderCommandTag>
 	static void Enqueue(FRenderCommandPipe& Pipe, FEmptyFunction&& Function)
 	{
-		if (Pipe.IsRecording())
+		if (Pipe.bEnabled)
 		{
-			Pipe.EnqueueAndLaunch(MoveTemp(Function), RenderCommandTag::GetName(), RenderCommandTag::GetSpecId());
-		}
-		else
-		{
-			EnqueueUniqueRenderCommand<RenderCommandTag>([Function = MoveTemp(Function)](FRHICommandListImmediate&) { Function(); });
+			UE::TScopeLock Lock(Pipe.Mutex);
+			if (Pipe.Frame_GameThread)
+			{
+				Pipe.EnqueueAndLaunch(MoveTemp(Function), RenderCommandTag::GetName(), RenderCommandTag::GetSpecId());
+				return;
+			}
 		}
+
+		EnqueueUniqueRenderCommand<RenderCommandTag>([Function = MoveTemp(Function)](FRHICommandListImmediate&) { Function(); });
 	}
 
 	template <typename RenderCommandTag, typename LambdaType>
@@ -509,11 +535,6 @@ private:
 		EnqueueAndLaunch(FFunctionVariant(TInPlaceType<FEmptyFunction>(), MoveTemp(Function)), CommandName, CommandSpecId);
 	}
 
-	FORCEINLINE bool IsRecording() const
-	{
-		return Frame_GameThread != nullptr;
-	}
-
 	struct FCommand
 	{
 		FCommand(FFunctionVariant&& InFunction, const TCHAR* InName, uint32& InOutSpecId)
@@ -527,7 +548,7 @@ private:
 		uint32* SpecId;
 	};
 
-	struct FFrame
+	struct FFrame : public TConcurrentLinearObject<FFrame>
 	{
 		FFrame(const TCHAR* Name, const UE::Tasks::FTaskEvent& InTaskEvent)
 			: Pipe(Name)
@@ -537,15 +558,17 @@ private:
 		UE::Tasks::FPipe Pipe;
 		UE::Tasks::FTaskEvent TaskEvent;
 		TArray<FCommand> Queue;
-		UE::FMutex QueueMutex;
 		FRHICommandList* RHICmdList = nullptr;
 	};
 
 	const TCHAR* Name;
+	UE::FMutex Mutex;
 	FFrame* Frame_GameThread = nullptr;
 	FFrame* Frame_RenderThread = nullptr;
 	TLinkedList<FRenderCommandPipe*> GlobalListLink;
 	FAutoConsoleVariable ConsoleVariable;
+	uint16 Index = uint16(-1);
+	bool bRecording = false;
 	bool bEnabled = true;
 };
 
diff --git a/Engine/Source/Runtime/Renderer/Private/ReflectionEnvironmentCapture.cpp b/Engine/Source/Runtime/Renderer/Private/ReflectionEnvironmentCapture.cpp
--- a/Engine/Source/Runtime/Renderer/Private/ReflectionEnvironmentCapture.cpp
+++ b/Engine/Source/Runtime/Renderer/Private/ReflectionEnvironmentCapture.cpp
@@ -1343,8 +1343,6 @@ void CaptureSceneIntoScratchCubemap(
 
 		FSceneRenderer* SceneRenderer = FSceneRenderer::CreateSceneRenderer(&ViewFamily, nullptr);
 
-		UE::RenderCommandPipe::FSyncScope SyncScope;
-
 		ENQUEUE_RENDER_COMMAND(CaptureCommand)(
 			[SceneRenderer, &ReflectionCubemapTexture, CubeFace, CubemapSize, bCapturingForSkyLight, bLowerHemisphereIsBlack, LowerHemisphereColor, bCapturingForMobile, bInsideTick](FRHICommandListImmediate& RHICmdList)
 		{
@@ -1439,6 +1437,8 @@ void FScene::CaptureOrUploadReflectionCapture(UReflectionCaptureComponent* Captu
 				return;
 			}
 
+			UE::RenderCommandPipe::FSyncScope SyncScope;
+
 			// Prefetch all virtual textures so that we have content available
 			if (UseVirtualTexturing(GetFeatureLevel()))
 			{
@@ -1713,6 +1713,8 @@ void FScene::UpdateSkyCaptureContents(
 
 		TRenderThreadStruct<FReflectionCubemapTexture> ReflectionCubemapTexture(CubemapResolution);
 
+		UE::RenderCommandPipe::FSyncScope SyncScope;
+
 		if (CaptureComponent->SourceType == SLS_CapturedScene)
 		{
 			const bool bStaticSceneOnly = CaptureComponent->Mobility == EComponentMobility::Static;
