[shaders] roll back input hash sanity check; this was making every job output necessarily one to one with input hashes (and so making the last layer of deduplication within the job cache ineffective; multiple input hashes could conceivably generate the same output due to code diffs that don't affect final output; these are the things we need to fix to improve cache hit rate).

diff --git a/Engine/Source/Runtime/RenderCore/Private/ShaderCore.cpp b/Engine/Source/Runtime/RenderCore/Private/ShaderCore.cpp
--- a/Engine/Source/Runtime/RenderCore/Private/ShaderCore.cpp
+++ b/Engine/Source/Runtime/RenderCore/Private/ShaderCore.cpp
@@ -1301,9 +1301,6 @@ public:
 			{
 				Compiler->CompileShader(Job.Input.ShaderFormat, Job.Input, Job.Output, WorkingDirectory);
 			}
-
-			// set the input hash on the output to allow sanity checking the job results on completion/reading from cache
-			Job.Output.InputHash = Job.Input.Hash;
 		}
 #if PLATFORM_WINDOWS
 		__except(HandleShaderCompileException(GetExceptionInformation(), OutExceptionMsg, OutExceptionCallstack))
@@ -3333,8 +3330,6 @@ void FShaderCompileJob::SerializeOutput(FArchive& Ar)
 		Output.CompileTime = ActualCompileTime;
 		Output.PreprocessTime = ActualPreprocessTime;
 	}
-
-	checkf(!Output.bSucceeded || Output.InputHash == Input.Hash, TEXT("Failed sanity check: InputHash reported from compile output does not match InputHash from compile input"));
 }
 
 void FShaderCompileJob::OnComplete()
diff --git a/Engine/Source/Runtime/RenderCore/Public/ShaderCompilerCore.h b/Engine/Source/Runtime/RenderCore/Public/ShaderCompilerCore.h
--- a/Engine/Source/Runtime/RenderCore/Public/ShaderCompilerCore.h
+++ b/Engine/Source/Runtime/RenderCore/Public/ShaderCompilerCore.h
@@ -29,14 +29,14 @@ class FShaderPipelineCompileJob;
 // this is for the protocol, not the data, bump if FShaderCompilerInput or ProcessInputFromArchive changes.
 inline const int32 ShaderCompileWorkerInputVersion = 21;
 // this is for the protocol, not the data, bump if FShaderCompilerOutput or WriteToOutputArchive changes.
-inline const int32 ShaderCompileWorkerOutputVersion = 17;
+inline const int32 ShaderCompileWorkerOutputVersion = 18;
 // this is for the protocol, not the data.
 inline const int32 ShaderCompileWorkerSingleJobHeader = 'S';
 // this is for the protocol, not the data.
 inline const int32 ShaderCompileWorkerPipelineJobHeader = 'P';
 
 // modify this for changes to the FShaderCompilerOutput data structure (in addition to ShaderCompileWorkerOutputVersion)
-inline const int32 FShaderCompilerOutputStructVersion = 1;
+inline const int32 FShaderCompilerOutputStructVersion = 2;
 
 namespace UE::ShaderCompiler
 {
@@ -574,7 +574,6 @@ struct FShaderCompilerOutput
 	TArray<FString> PragmaDirectives;
 	FShaderTarget Target;
 	FShaderCode ShaderCode;
-	FBlake3Hash InputHash;
 	FSHAHash OutputHash;
 	uint32 NumInstructions;
 	uint32 NumTextureSamplers;
@@ -642,7 +641,7 @@ struct FShaderCompilerOutput
 	friend FArchive& operator<<(FArchive& Ar, FShaderCompilerOutput& Output)
 	{
 		// Note: this serialize is used to pass between UE and the shader compile worker, recompile both when modifying
-		Ar << Output.ParameterMap << Output.Errors << Output.Target << Output.ShaderCode << Output.InputHash << Output.OutputHash << Output.NumInstructions << Output.NumTextureSamplers << Output.bSucceeded;
+		Ar << Output.ParameterMap << Output.Errors << Output.Target << Output.ShaderCode << Output.OutputHash << Output.NumInstructions << Output.NumTextureSamplers << Output.bSucceeded;
 		Ar << Output.bSupportsQueryingUsedAttributes << Output.UsedAttributes;
 		Ar << Output.CompileTime;
 		Ar << Output.PreprocessTime;
