Reworked render command pipes sync scopes to only start recording when a pipe was previously recording.
 - Fixes issue where pipes are stopped, then a sync occurs and pipes are erroneously started up again.
 - Removed start / stop from end of frame updates.

diff --git a/Engine/Source/Runtime/Engine/Private/LevelTick.cpp b/Engine/Source/Runtime/Engine/Private/LevelTick.cpp
--- a/Engine/Source/Runtime/Engine/Private/LevelTick.cpp
+++ b/Engine/Source/Runtime/Engine/Private/LevelTick.cpp
@@ -1019,8 +1019,6 @@ void UWorld::SendAllEndOfFrameUpdates()
 		return;
 	}
 
-	UE::RenderCommandPipe::StartRecording();
-
 	// Wait for tasks that are generating data for the render proxies, but are not awaited in any TickFunctions 
 	// E.g., see cloth USkeletalMeshComponent::UpdateClothStateAndSimulate
 	for (UActorComponent* Component : ComponentsThatNeedPreEndOfFrameSync)
@@ -1163,8 +1161,6 @@ void UWorld::SendAllEndOfFrameUpdates()
 	LocalComponentsThatNeedEndOfFrameUpdate.Reset();
 
 	EndSendEndOfFrameUpdatesDrawEvent(SendAllEndOfFrameUpdates);
-
-	UE::RenderCommandPipe::StopRecording();
 }
 
 /**
diff --git a/Engine/Source/Runtime/RenderCore/Private/RenderingThread.cpp b/Engine/Source/Runtime/RenderCore/Private/RenderingThread.cpp
--- a/Engine/Source/Runtime/RenderCore/Private/RenderingThread.cpp
+++ b/Engine/Source/Runtime/RenderCore/Private/RenderingThread.cpp
@@ -1595,7 +1595,7 @@ void FRenderThreadCommandPipe::EnqueueAndLaunch(const TCHAR* Name, uint32& SpecI
 class FRenderCommandPipeRegistry
 {
 public:
-	TLinkedList<FRenderCommandPipe*>*& GetGlobalList()
+	static TLinkedList<FRenderCommandPipe*>*& GetGlobalList()
 	{
 		static TLinkedList<FRenderCommandPipe*>* GlobalList = nullptr;
 		return GlobalList;
@@ -1617,38 +1617,37 @@ public:
 		GRenderCommandPipeMode = GetValidatedRenderCommandPipeMode(CVarRenderCommandPipeMode->GetInt());
 	}
 
-	void StartRecording(TConstArrayView<FRenderCommandPipe*> Pipes)
+	void StartRecording()
 	{
-		SCOPED_NAMED_EVENT(FRenderCommandPipe_StartRecording, FColor::Magenta);
-
-		if (GRenderCommandPipeMode != ERenderCommandPipeMode::All)
+		if (GRenderCommandPipeMode != ERenderCommandPipeMode::All || !GIsThreadedRendering)
 		{
 			return;
 		}
 
-		bool bAllPipes = false;
-		if (Pipes.IsEmpty())
-		{
-			bAllPipes = true;
-			Pipes = AllPipes;
-		}
+		FRenderCommandPipeBitArray PipeBits;
+		PipeBits.Init(true, AllPipes.Num());
+		StartRecording(PipeBits);
+	}
 
-		if (Pipes.IsEmpty() || !ShouldExecuteOnRenderThread())
+	void StartRecording(const FRenderCommandPipeBitArray& PipeBits)
+	{
+		if (GRenderCommandPipeMode != ERenderCommandPipeMode::All || !GIsThreadedRendering || PipeBits.IsEmpty())
 		{
 			return;
 		}
 
-		UE::TScopeLock Lock(Mutex);
+		SCOPED_NAMED_EVENT(FRenderCommandPipe_StartRecording, FColor::Magenta);
 
-		if (!bAllPipes)
-		{
-			CheckUnique(Pipes);
-		}
+		check(PipeBits.Num() == AllPipes.Num());
+
+		UE::TScopeLock Lock(Mutex);
 
 		bool bAnyPipesToStartRecording = false;
 
-		for (FRenderCommandPipe* Pipe : Pipes)
+		for (FRenderCommandPipeSetBitIterator BitIt(PipeBits); BitIt; ++BitIt)
 		{
+			FRenderCommandPipe* Pipe = AllPipes[BitIt.GetIndex()];
+
 			if (Pipe->bEnabled && !Pipe->bRecording)
 			{
 				bAnyPipesToStartRecording = true;
@@ -1675,10 +1674,12 @@ public:
 		};
 
 		TArray<FPipeToStartRecording, FConcurrentLinearArrayAllocator> PipesToStartRecording;
-		PipesToStartRecording.Reserve(Pipes.Num());
+		PipesToStartRecording.Reserve(AllPipes.Num());
 
-		for (FRenderCommandPipe* Pipe : Pipes)
+		for (FRenderCommandPipeSetBitIterator BitIt(PipeBits); BitIt; ++BitIt)
 		{
+			FRenderCommandPipe* Pipe = AllPipes[BitIt.GetIndex()];
+
 			if (Pipe->bEnabled && !Pipe->bRecording)
 			{
 				Pipe->bRecording = true;
@@ -1707,70 +1708,106 @@ public:
 		});
 	}
 
-	void StopRecording(TConstArrayView<FRenderCommandPipe*> Pipes)
+	FRenderCommandPipeBitArray StopRecording()
 	{
-		SCOPED_NAMED_EVENT(FRenderCommandPipe_StopRecording, FColor::Magenta);
+		UE::TScopeLock Lock(Mutex);
+		if (!NumPipesRecording)
+		{
+			return {};
+		}
 
-		bool bAllPipes = false;
+		FRenderCommandPipeBitArray PipeBits;
+		PipeBits.Init(false, AllPipes.Num());
 
-		if (Pipes.IsEmpty())
+		for (int32 PipeIndex = 0; PipeIndex < AllPipes.Num(); ++PipeIndex)
 		{
-			bAllPipes = true;
-			Pipes = AllPipes;
+			if (FRenderCommandPipe* Pipe = AllPipes[PipeIndex]; Pipe->bRecording)
+			{
+				PipeBits[PipeIndex] = true;
+			}
 		}
 
+		StopRecording(PipeBits);
+		return PipeBits;
+	}
+
+	FRenderCommandPipeBitArray StopRecording(TConstArrayView<FRenderCommandPipe*> Pipes)
+	{
 		if (Pipes.IsEmpty())
 		{
-			return;
+			return {};
 		}
 
 		UE::TScopeLock Lock(Mutex);
-
-		if (!bAllPipes)
+		if (!NumPipesRecording)
 		{
-			CheckUnique(Pipes);
+			return {};
 		}
 
 		bool bAnyPipesToStopRecording = false;
+		FRenderCommandPipeBitArray PipeBits;
+		PipeBits.Init(false, AllPipes.Num());
 
 		for (FRenderCommandPipe* Pipe : Pipes)
 		{
 			if (Pipe->bRecording)
 			{
+				PipeBits[Pipe->Index] = true;
 				bAnyPipesToStopRecording = true;
-				break;
 			}
 		}
 
 		if (!bAnyPipesToStopRecording)
 		{
-			return;
+			return {};
 		}
 
-		TArray<FRenderCommandPipe*, FConcurrentLinearArrayAllocator> PipesToStopRecording;
-		PipesToStopRecording.Reserve(Pipes.Num());
+		StopRecording(PipeBits);
+		return PipeBits;
+	}
 
-		for (FRenderCommandPipe* Pipe : Pipes)
+	bool IsRecording() const
+	{
+		ensureMsgf(!FTaskTagScope::IsCurrentTag(ETaskTag::EParallelRenderingThread) && !FTaskTagScope::IsCurrentTag(ETaskTag::ERenderingThread),
+			TEXT("IsRecording() is not valid from the render thread timeline."));
+
+		return NumPipesRecording > 0;
+	}
+
+	bool IsReplaying() const
+	{
+		ensure(IsInParallelRenderingThread());
+		return NumPipesReplaying > 0;
+	}
+
+private:
+	void StopRecording(const FRenderCommandPipeBitArray& PipeBits)
+	{
+		SCOPED_NAMED_EVENT(FRenderCommandPipe_StopRecording, FColor::Magenta);
+
+		uint32 NumPipesToStopRecording = 0;
+
+		for (FRenderCommandPipeSetBitIterator BitIt(PipeBits); BitIt; ++BitIt)
 		{
-			if (Pipe->bRecording)
-			{
-				Pipe->bRecording = false;
-				PipesToStopRecording.Emplace(Pipe);
+			FRenderCommandPipe* Pipe = AllPipes[BitIt.GetIndex()];
+			check(Pipe->bRecording);
+			Pipe->bRecording = false;
+			NumPipesToStopRecording++;
 
-				UE::TScopeLock PipeLock(Pipe->Mutex);
-				Pipe->Frame_GameThread = nullptr;
-			}
+			UE::TScopeLock PipeLock(Pipe->Mutex);
+			Pipe->Frame_GameThread = nullptr;
 		}
 
-		NumPipesRecording -= PipesToStopRecording.Num();
+		NumPipesRecording -= NumPipesToStopRecording;
 
-		ENQUEUE_RENDER_COMMAND(RenderCommandPipe_Stop)([this, PipesToStopRecording = MoveTemp(PipesToStopRecording)](FRHICommandListImmediate& RHICmdList)
+		ENQUEUE_RENDER_COMMAND(RenderCommandPipe_Stop)([this, PipeBits, NumPipesToStopRecording](FRHICommandListImmediate& RHICmdList)
 		{
 			TArray<FRHICommandListImmediate::FQueuedCommandList, FConcurrentLinearArrayAllocator> QueuedCommandLists;
-			QueuedCommandLists.Reserve(PipesToStopRecording.Num());
+			QueuedCommandLists.Reserve(NumPipesToStopRecording);
 
-			for (FRenderCommandPipe* Pipe : PipesToStopRecording)
+			for (FRenderCommandPipeSetBitIterator BitIt(PipeBits); BitIt; ++BitIt)
 			{
+				FRenderCommandPipe* Pipe = AllPipes[BitIt.GetIndex()];
 				FRenderCommandPipe::FFrame*& Frame_RenderThread = Pipe->Frame_RenderThread;
 				check(Frame_RenderThread);
 				Frame_RenderThread->Pipe.WaitUntilEmpty();
@@ -1785,50 +1822,17 @@ public:
 				Frame_RenderThread = nullptr;
 			}
 
-			NumPipesReplaying -= PipesToStopRecording.Num();
+			NumPipesReplaying -= NumPipesToStopRecording;
 
 			RHICmdList.QueueAsyncCommandListSubmit(QueuedCommandLists);
-			RHIResourceLifetimeReleaseRef(RHICmdList, PipesToStopRecording.Num());
+			RHIResourceLifetimeReleaseRef(RHICmdList, NumPipesToStopRecording);
 		});
 	}
 
-	bool IsRecording() const
-	{
-		ensureMsgf(!FTaskTagScope::IsCurrentTag(ETaskTag::EParallelRenderingThread) && !FTaskTagScope::IsCurrentTag(ETaskTag::ERenderingThread),
-			TEXT("IsRecording() is not valid from the render thread timeline."));
-
-		return NumPipesRecording > 0;
-	}
-
-	bool IsReplaying() const
-	{
-		ensure(IsInParallelRenderingThread());
-		return NumPipesReplaying > 0;
-	}
-
-private:
-	void CheckUnique(TConstArrayView<FRenderCommandPipe*> Pipes)
-	{
-#if DO_CHECK
-		CheckUniqueBits.Init(false, Pipes.Num());
-		for (FRenderCommandPipe* Pipe : Pipes)
-		{
-			checkf(Pipe->Index < AllPipes.Num(), TEXT("RenderCommandPipe %s has an invalid index and was likely constructed after system initialization."), Pipe->Name);
-			checkf(!CheckUniqueBits[Pipe->Index], TEXT("RenderCommandPipe %s was supplied more than once. This is invalid. Only a unique set of pipes is allowed."), Pipe->Name);
-			CheckUniqueBits[Pipe->Index] = true;
-		}
-		CheckUniqueBits.Reset();
-#endif
-	}
-
 	UE::FMutex Mutex;
 	TArray<FRenderCommandPipe*> AllPipes;
-	std::atomic_uint32_t NumPipesRecording = 0;
+	uint32 NumPipesRecording = 0;
 	uint32 NumPipesReplaying = 0;
-
-#if DO_CHECK
-	TBitArray<> CheckUniqueBits;
-#endif
 };
 
 static FRenderCommandPipeRegistry GRenderCommandPipeRegistry;
@@ -1865,28 +1869,28 @@ namespace UE::RenderCommandPipe
 	void StartRecording()
 	{
 #if !UE_SERVER
-		GRenderCommandPipeRegistry.StartRecording({});
+		GRenderCommandPipeRegistry.StartRecording();
 #endif
 	}
 
-	void StartRecording(TConstArrayView<FRenderCommandPipe*> Pipes)
+	void StartRecording(const FRenderCommandPipeBitArray& PipeBits)
 	{
 #if !UE_SERVER
-		GRenderCommandPipeRegistry.StartRecording(Pipes);
+		GRenderCommandPipeRegistry.StartRecording(PipeBits);
 #endif
 	}
 
-	void StopRecording()
+	FRenderCommandPipeBitArray StopRecording()
 	{
 #if !UE_SERVER
-		GRenderCommandPipeRegistry.StopRecording({});
+		return GRenderCommandPipeRegistry.StopRecording();
 #endif
 	}
 
-	void StopRecording(TConstArrayView<FRenderCommandPipe*> Pipes)
+	FRenderCommandPipeBitArray StopRecording(TConstArrayView<FRenderCommandPipe*> Pipes)
 	{
 #if !UE_SERVER
-		GRenderCommandPipeRegistry.StopRecording(Pipes);
+		return GRenderCommandPipeRegistry.StopRecording(Pipes);
 #endif
 	}
 }
@@ -1900,7 +1904,7 @@ FRenderCommandPipe::FRenderCommandPipe(const TCHAR* InName, ERenderCommandPipeFl
 	}))
 {
 #if !UE_SERVER
-	GlobalListLink.LinkHead(GRenderCommandPipeRegistry.GetGlobalList());
+	GlobalListLink.LinkHead(FRenderCommandPipeRegistry::GetGlobalList());
 #endif
 }
 
diff --git a/Engine/Source/Runtime/RenderCore/Public/RenderingThread.h b/Engine/Source/Runtime/RenderCore/Public/RenderingThread.h
--- a/Engine/Source/Runtime/RenderCore/Public/RenderingThread.h
+++ b/Engine/Source/Runtime/RenderCore/Public/RenderingThread.h
@@ -295,6 +295,9 @@ enum class ERenderCommandPipeFlags : uint8
 ENUM_CLASS_FLAGS(ERenderCommandPipeFlags);
 
 class FRenderCommandPipe;
+using FRenderCommandPipeBitArrayAllocator = TInlineAllocator<1, FConcurrentLinearArrayAllocator>;
+using FRenderCommandPipeBitArray = TBitArray<FRenderCommandPipeBitArrayAllocator>;
+using FRenderCommandPipeSetBitIterator = TConstSetBitIterator<FRenderCommandPipeBitArrayAllocator>;
 
 namespace UE::RenderCommandPipe
 {
@@ -309,11 +312,11 @@ namespace UE::RenderCommandPipe
 
 	// [Game Thread] Starts recording render commands into pipes. Returns whether the operation succeeded.
 	extern RENDERCORE_API void StartRecording();
-	extern RENDERCORE_API void StartRecording(TConstArrayView<FRenderCommandPipe*> Pipes);
+	extern RENDERCORE_API void StartRecording(const FRenderCommandPipeBitArray& PipeBits);
 
 	// [Game Thread] Stops recording commands into pipes and syncs all remaining pipe work to the render thread. Returns whether the operation succeeded.
-	extern RENDERCORE_API void StopRecording();
-	extern RENDERCORE_API void StopRecording(TConstArrayView<FRenderCommandPipe*> Pipes);
+	extern RENDERCORE_API FRenderCommandPipeBitArray StopRecording();
+	extern RENDERCORE_API FRenderCommandPipeBitArray StopRecording(TConstArrayView<FRenderCommandPipe*> Pipes);
 
 	// [Game Thread] Stops render command pipe recording during the duration of the scope and restarts recording once the scope is complete.
 	class FSyncScope
@@ -321,29 +324,21 @@ namespace UE::RenderCommandPipe
 	public:
 		FSyncScope()
 		{
-			StopRecording();
+			PipeBits = StopRecording();
 		}
 
-		FSyncScope(TConstArrayView<FRenderCommandPipe*> InPipes)
-			: Pipes(InPipes)
+		FSyncScope(TConstArrayView<FRenderCommandPipe*> Pipes)
 		{
-			StopRecording(Pipes);
+			PipeBits = StopRecording(Pipes);
 		}
 
 		~FSyncScope()
 		{
-			if (!Pipes.IsEmpty())
-			{
-				StartRecording(Pipes);
-			}
-			else
-			{
-				StartRecording();
-			}
+			StartRecording(PipeBits);
 		}
 
 	private:
-		TConstArrayView<FRenderCommandPipe*> Pipes;
+		FRenderCommandPipeBitArray PipeBits;
 	};
 }
 
