Fixed server build break.

diff --git a/Engine/Source/Runtime/RenderCore/Private/RenderingThread.cpp b/Engine/Source/Runtime/RenderCore/Private/RenderingThread.cpp
--- a/Engine/Source/Runtime/RenderCore/Private/RenderingThread.cpp
+++ b/Engine/Source/Runtime/RenderCore/Private/RenderingThread.cpp
@@ -1884,6 +1884,8 @@ namespace UE::RenderCommandPipe
 	{
 #if !UE_SERVER
 		return GRenderCommandPipeRegistry.StopRecording();
+#else
+		return {};
 #endif
 	}
 
@@ -1891,6 +1893,8 @@ namespace UE::RenderCommandPipe
 	{
 #if !UE_SERVER
 		return GRenderCommandPipeRegistry.StopRecording(Pipes);
+#else
+		return {};
 #endif
 	}
 }
