Defer submission of AddCommandListSetupTask to batch calls to AsyncCommandListSubmit.

diff --git a/Engine/Source/Runtime/RenderCore/Private/RenderGraphBuilder.cpp b/Engine/Source/Runtime/RenderCore/Private/RenderGraphBuilder.cpp
--- a/Engine/Source/Runtime/RenderCore/Private/RenderGraphBuilder.cpp
+++ b/Engine/Source/Runtime/RenderCore/Private/RenderGraphBuilder.cpp
@@ -1739,6 +1739,7 @@ void FRDGBuilder::Execute()
 
 		if (!ParallelSetupEvents.IsEmpty())
 		{
+			RHICmdList.QueueAsyncCommandListSubmit(ParallelSetupCommandLists);
 			UE::Tasks::Wait(ParallelSetupEvents);
 			ParallelSetupEvents.Empty();
 		}
diff --git a/Engine/Source/Runtime/RenderCore/Public/RenderGraphBuilder.h b/Engine/Source/Runtime/RenderCore/Public/RenderGraphBuilder.h
--- a/Engine/Source/Runtime/RenderCore/Public/RenderGraphBuilder.h
+++ b/Engine/Source/Runtime/RenderCore/Public/RenderGraphBuilder.h
@@ -662,9 +662,12 @@ private:
 	/** Array of all active parallel execute tasks. */
 	TArray<UE::Tasks::FTask, FRDGArrayAllocator> ParallelExecuteEvents;
 
-	/** Array of all task events requested by the user. */
+	/** Array of all task events for variants of AddSetupTask. */
 	TArray<UE::Tasks::FTask, FRDGArrayAllocator> ParallelSetupEvents;
 
+	/** Array of all command lists to submit for AddCommandListSetupTask. */
+	TArray<FRHICommandListImmediate::FQueuedCommandList, FRDGArrayAllocator> ParallelSetupCommandLists;
+
 	/** Tracks the final access used on resources in order to call SetTrackedAccess. */
 	TArray<FRHITrackedAccessInfo, FRDGArrayAllocator> EpilogueResourceAccesses;
 
diff --git a/Engine/Source/Runtime/RenderCore/Public/RenderGraphBuilder.inl b/Engine/Source/Runtime/RenderCore/Public/RenderGraphBuilder.inl
--- a/Engine/Source/Runtime/RenderCore/Public/RenderGraphBuilder.inl
+++ b/Engine/Source/Runtime/RenderCore/Public/RenderGraphBuilder.inl
@@ -614,7 +614,7 @@ UE::Tasks::FTask FRDGBuilder::AddCommandListSetupTask(
 	if (!RHICmdListTask->IsImmediate())
 	{
 		ParallelSetupEvents.Emplace(Task);
-		RHICmdList.QueueAsyncCommandListSubmit(RHICmdListTask);
+		ParallelSetupCommandLists.Emplace(RHICmdListTask);
 	}
 
 	return Task;
