Register/unregister FRayTracingGeometry with FRayTracingGeometryManager in Init/ReleaseResource
- instead doing it in Init/ReleaseRHI
- fixes RegisteredGeometries being modified while iterating over it due to InitRHIForDynamicRayTracing(...) internally calling Init/ReleaseRHI

diff --git a/Engine/Source/Runtime/RenderCore/Private/RayTracingGeometry.cpp b/Engine/Source/Runtime/RenderCore/Private/RayTracingGeometry.cpp
--- a/Engine/Source/Runtime/RenderCore/Private/RayTracingGeometry.cpp
+++ b/Engine/Source/Runtime/RenderCore/Private/RayTracingGeometry.cpp
@@ -172,14 +172,6 @@ void FRayTracingGeometry::CreateRayTracingGeometry(FRHICommandList& RHICmdList,
 			RayTracingGeometryRHI = RHICmdList.CreateRayTracingGeometry(Initializer);
 		}
 
-		// Register the geometry if it wasn't registered before and it's not using custom path
-		const bool bRegisterGeometry = (RayTracingGeometryHandle == INDEX_NONE) && (InBuildPriority != ERTAccelerationStructureBuildPriority::Immediate);
-
-		if (bRegisterGeometry)
-		{
-			RayTracingGeometryHandle = GRayTracingGeometryManager.RegisterRayTracingGeometry(this);
-		}
-
 		if (Initializer.OfflineData == nullptr)
 		{
 			// Request build if not skip
@@ -237,12 +229,6 @@ void FRayTracingGeometry::ReleaseRHI()
 {
 	RemoveBuildRequest();
 	RayTracingGeometryRHI.SafeRelease();
-
-	if (RayTracingGeometryHandle != INDEX_NONE)
-	{
-		GRayTracingGeometryManager.ReleaseRayTracingGeometryHandle(RayTracingGeometryHandle);
-		RayTracingGeometryHandle = INDEX_NONE;
-	}
 }
 
 void FRayTracingGeometry::RemoveBuildRequest()
@@ -259,12 +245,21 @@ void FRayTracingGeometry::InitResource(FRHICommandListBase& RHICmdList)
 	ensureMsgf(IsRayTracingAllowed(), TEXT("FRayTracingGeometry should only be initialized when Ray Tracing is allowed"));
 
 	FRenderResource::InitResource(RHICmdList);
+
+	check(RayTracingGeometryHandle == INDEX_NONE);
+	RayTracingGeometryHandle = GRayTracingGeometryManager.RegisterRayTracingGeometry(this);
 }
 
 void FRayTracingGeometry::ReleaseResource()
 {
 	ensureMsgf(IsRayTracingAllowed() || !IsInitialized(), TEXT("FRayTracingGeometry should only be initialized when Ray Tracing is allowed"));
 
+	if (RayTracingGeometryHandle != INDEX_NONE)
+	{
+		GRayTracingGeometryManager.ReleaseRayTracingGeometryHandle(RayTracingGeometryHandle);
+		RayTracingGeometryHandle = INDEX_NONE;
+	}
+
 	// Release any resource references held by the initializer.
 	// This includes index and vertex buffers used for building the BLAS.
 	Initializer = FRayTracingGeometryInitializer{};
diff --git a/Engine/Source/Runtime/RenderCore/Private/RayTracingGeometryManager.cpp b/Engine/Source/Runtime/RenderCore/Private/RayTracingGeometryManager.cpp
--- a/Engine/Source/Runtime/RenderCore/Private/RayTracingGeometryManager.cpp
+++ b/Engine/Source/Runtime/RenderCore/Private/RayTracingGeometryManager.cpp
@@ -128,7 +128,8 @@ void FRayTracingGeometryManager::Tick(FRHICommandList& RHICmdList, bool bHasRayT
 			FScopeLock ScopeLock(&RequestCS);
 			for (FRayTracingGeometry* Geometry : RegisteredGeometries)
 			{
-				checkf(Geometry->RayTracingGeometryRHI != nullptr, TEXT("Ray tracing geometry should be valid at this point."));
+				checkf(!EnumHasAllFlags(Geometry->GetGeometryState(), FRayTracingGeometry::EGeometryStateFlags::Valid) || Geometry->RayTracingGeometryRHI != nullptr, 
+					TEXT("Ray tracing geometry should have a valid RHI resource when ray tracing is enabled."));
 			}
 		}
 		else
@@ -136,7 +137,7 @@ void FRayTracingGeometryManager::Tick(FRHICommandList& RHICmdList, bool bHasRayT
 			FScopeLock ScopeLock(&RequestCS);
 			for (FRayTracingGeometry* Geometry : RegisteredGeometries)
 			{
-				checkf(Geometry->RayTracingGeometryRHI == nullptr, TEXT("Ray tracing geometry should not be valid at this point"));
+				checkf(Geometry->RayTracingGeometryRHI == nullptr, TEXT("Ray tracing geometry should not have a RHI resource when ray tracing is disabled."));
 			}
 		}
 #endif
@@ -159,8 +160,7 @@ void FRayTracingGeometryManager::Tick(FRHICommandList& RHICmdList, bool bHasRayT
 		FScopeLock ScopeLock(&RequestCS);
 		for (FRayTracingGeometry* Geometry : RegisteredGeometries)
 		{
-			Geometry->RemoveBuildRequest();
-			Geometry->RayTracingGeometryRHI.SafeRelease();			
+			Geometry->ReleaseRHI();			
 		}
 	}
 }
