[shaders] fix preprocessed cache versioning mechanism

diff --git a/Engine/Source/Runtime/RenderCore/Private/ShaderCore.cpp b/Engine/Source/Runtime/RenderCore/Private/ShaderCore.cpp
--- a/Engine/Source/Runtime/RenderCore/Private/ShaderCore.cpp
+++ b/Engine/Source/Runtime/RenderCore/Private/ShaderCore.cpp
@@ -3159,7 +3159,7 @@ FShaderCommonCompileJob::FInputHash FShaderCompileJob::GetInputHash()
 			Hasher << const_cast<FString&>(*VersionDirective);
 		};
 
-		PreprocessOutput.VisitDirectivesWithPrefix(TEXT("VERSION_"), HashVersion);
+		PreprocessOutput.VisitDirectivesWithPrefix(TEXT("VERSION"), HashVersion);
 		// const_cast due to serialization API requiring non-const. better than not having const correctness in the API.
 		Hasher << const_cast<FString&>(PreprocessOutput.GetSource());
 		if (SecondaryPreprocessOutput.IsValid())
