Don't register FRayTracingGeometry with GRayTracingGeometryManager if it's already registered.

diff --git a/Engine/Source/Runtime/RenderCore/Private/RayTracingGeometry.cpp b/Engine/Source/Runtime/RenderCore/Private/RayTracingGeometry.cpp
--- a/Engine/Source/Runtime/RenderCore/Private/RayTracingGeometry.cpp
+++ b/Engine/Source/Runtime/RenderCore/Private/RayTracingGeometry.cpp
@@ -242,17 +242,19 @@ void FRayTracingGeometry::RemoveBuildRequest()
 
 void FRayTracingGeometry::InitResource(FRHICommandListBase& RHICmdList)
 {
-	ensureMsgf(IsRayTracingAllowed(), TEXT("FRayTracingGeometry should only be initialized when Ray Tracing is allowed"));
+	ensureMsgf(IsRayTracingAllowed(), TEXT("FRayTracingGeometry should only be initialized when Ray Tracing is allowed."));
 
 	FRenderResource::InitResource(RHICmdList);
 
-	check(RayTracingGeometryHandle == INDEX_NONE);
-	RayTracingGeometryHandle = GRayTracingGeometryManager.RegisterRayTracingGeometry(this);
+	if (RayTracingGeometryHandle == INDEX_NONE)
+	{
+		RayTracingGeometryHandle = GRayTracingGeometryManager.RegisterRayTracingGeometry(this);
+	}
 }
 
 void FRayTracingGeometry::ReleaseResource()
 {
-	ensureMsgf(IsRayTracingAllowed() || !IsInitialized(), TEXT("FRayTracingGeometry should only be initialized when Ray Tracing is allowed"));
+	ensureMsgf(IsRayTracingAllowed() || !IsInitialized(), TEXT("FRayTracingGeometry should only be initialized when Ray Tracing is allowed."));
 
 	if (RayTracingGeometryHandle != INDEX_NONE)
 	{
