Fix FRayTracingGeometry incorrectly keeping EGeometryStateFlags::Valid when RHI resource is released.
- that state combination is only correct when RHI resource is evicted due to dynamic raytracing being disabled.

Repro:
- Run cooked build with dynamic raytracing enabled.
- Toggle r.Streaming.StressTest on/off

diff --git a/Engine/Source/Runtime/RenderCore/Private/RayTracingGeometry.cpp b/Engine/Source/Runtime/RenderCore/Private/RayTracingGeometry.cpp
--- a/Engine/Source/Runtime/RenderCore/Private/RayTracingGeometry.cpp
+++ b/Engine/Source/Runtime/RenderCore/Private/RayTracingGeometry.cpp
@@ -127,6 +127,12 @@ void FRayTracingGeometry::InitRHIForDynamicRayTracing()
 	InitRHIForDynamicRayTracing(FRHICommandListImmediate::Get());
 }
 
+void FRayTracingGeometry::Evict()
+{
+	RemoveBuildRequest();
+	RayTracingGeometryRHI.SafeRelease();
+}
+
 void FRayTracingGeometry::CreateRayTracingGeometry(FRHICommandList& RHICmdList, ERTAccelerationStructureBuildPriority InBuildPriority)
 {
 	// Release previous RHI object if any
@@ -229,6 +235,7 @@ void FRayTracingGeometry::ReleaseRHI()
 {
 	RemoveBuildRequest();
 	RayTracingGeometryRHI.SafeRelease();
+	EnumRemoveFlags(GeometryState, EGeometryStateFlags::Valid);
 }
 
 void FRayTracingGeometry::RemoveBuildRequest()
@@ -266,6 +273,8 @@ void FRayTracingGeometry::ReleaseResource()
 	// This includes index and vertex buffers used for building the BLAS.
 	Initializer = FRayTracingGeometryInitializer{};
 
+	GeometryState = EGeometryStateFlags::Invalid;
+
 	FRenderResource::ReleaseResource();
 }
 
diff --git a/Engine/Source/Runtime/RenderCore/Private/RayTracingGeometryManager.cpp b/Engine/Source/Runtime/RenderCore/Private/RayTracingGeometryManager.cpp
--- a/Engine/Source/Runtime/RenderCore/Private/RayTracingGeometryManager.cpp
+++ b/Engine/Source/Runtime/RenderCore/Private/RayTracingGeometryManager.cpp
@@ -160,7 +160,7 @@ void FRayTracingGeometryManager::Tick(FRHICommandList& RHICmdList, bool bHasRayT
 		FScopeLock ScopeLock(&RequestCS);
 		for (FRayTracingGeometry* Geometry : RegisteredGeometries)
 		{
-			Geometry->ReleaseRHI();			
+			Geometry->Evict();			
 		}
 	}
 }
diff --git a/Engine/Source/Runtime/RenderCore/Public/RayTracingGeometry.h b/Engine/Source/Runtime/RenderCore/Public/RayTracingGeometry.h
--- a/Engine/Source/Runtime/RenderCore/Public/RayTracingGeometry.h
+++ b/Engine/Source/Runtime/RenderCore/Public/RayTracingGeometry.h
@@ -120,6 +120,8 @@ public:
 
 	UE_DEPRECATED(5.4, "CreateRayTracingGeometry now requires a command list.")
 	RENDERCORE_API void CreateRayTracingGeometry(ERTAccelerationStructureBuildPriority InBuildPriority);
+
+	RENDERCORE_API void Evict();
 	
 	bool HasPendingBuildRequest() const
 	{
