Remove UBMT_RDG_UNIFORM_BLOCK_SRV, added a UB usage flag that can be used to identify "UniformView" buffers

diff --git a/Engine/Shaders/Public/ShaderVersion.ush b/Engine/Shaders/Public/ShaderVersion.ush
--- a/Engine/Shaders/Public/ShaderVersion.ush
+++ b/Engine/Shaders/Public/ShaderVersion.ush
@@ -2,4 +2,4 @@
 // Each time the console command is executed it generates a new GUID. As this file is included
 // in Platform.ush (which should be included in any shader) it allows to invalidate the shader DDC.
 // 
-#pragma message("UESHADERMETADATA_VERSION E7DBB9A8-D147-660F-A96D-9A3E24BB76F312")
+#pragma message("UESHADERMETADATA_VERSION E8C6552B-D6B5-4CA4-B26C-1E315D3A1F20")
diff --git a/Engine/Source/Runtime/OpenGLDrv/Private/OpenGLUniformBuffer.cpp b/Engine/Source/Runtime/OpenGLDrv/Private/OpenGLUniformBuffer.cpp
--- a/Engine/Source/Runtime/OpenGLDrv/Private/OpenGLUniformBuffer.cpp
+++ b/Engine/Source/Runtime/OpenGLDrv/Private/OpenGLUniformBuffer.cpp
@@ -629,20 +629,30 @@ static FOpenGLUniformBuffer* CreateUniformBufferView(FRHICommandListImmediate& R
 {
 	FOpenGLUniformBuffer* UniformBufferView = nullptr;
 	
-	if (Layout->Resources.Num() == 1 && 
-		Layout->Resources[0].MemberType == UBMT_RDG_UNIFORM_BLOCK_SRV)
+	if (Layout->bUniformView)
 	{
 		UniformBufferView = new FOpenGLUniformBuffer(Layout);
 		UniformBufferView->SetLayoutTable(Contents, EUniformBufferValidation::None);
+
+		FRHIShaderResourceView* UniformViewSRV = nullptr;
+		for (int32 Index = 0; Index < Layout->Resources.Num() && !UniformViewSRV; ++Index)
+		{
+			EUniformBufferBaseType ResourceBaseType = Layout->Resources[Index].MemberType;
+			if (ResourceBaseType == UBMT_SRV || 
+				ResourceBaseType == UBMT_RDG_BUFFER_SRV)
+			{
+				UniformViewSRV = (FRHIShaderResourceView*)GetShaderParameterResourceRHI(Contents, Layout->Resources[Index].MemberOffset, ResourceBaseType);
+			}
+		}
 		
-		FRHIShaderResourceView* SRV = (FRHIShaderResourceView*)GetShaderParameterResourceRHI(Contents, Layout->Resources[0].MemberOffset, UBMT_RDG_UNIFORM_BLOCK_SRV);
+		check(UniformViewSRV);
 
-		RHICmdList.EnqueueLambda([UniformBufferView, SRV](FRHICommandListImmediate&)
+		RHICmdList.EnqueueLambda([UniformBufferView, UniformViewSRV](FRHICommandListImmediate&)
 		{
 			VERIFY_GL_SCOPE();
 			
-			FOpenGLBuffer* UBO = FOpenGLDynamicRHI::ResourceCast(SRV->GetBuffer());
-			const FRHIViewDesc::FBufferSRV& SRVInfo = SRV->GetDesc().Buffer.SRV;
+			FOpenGLBuffer* UBO = FOpenGLDynamicRHI::ResourceCast(UniformViewSRV->GetBuffer());
+			const FRHIViewDesc::FBufferSRV& SRVInfo = UniformViewSRV->GetDesc().Buffer.SRV;
 			
 			check(UBO->Resource);
 			check(UBO->GetSize() >= PLATFORM_MAX_UNIFORM_BUFFER_RANGE);
diff --git a/Engine/Source/Runtime/RHI/Private/RHIResources.cpp b/Engine/Source/Runtime/RHI/Private/RHIResources.cpp
--- a/Engine/Source/Runtime/RHI/Private/RHIResources.cpp
+++ b/Engine/Source/Runtime/RHI/Private/RHIResources.cpp
@@ -361,5 +361,6 @@ FRHIUniformBufferLayout::FRHIUniformBufferLayout(const FRHIUniformBufferLayoutIn
 	, BindingFlags(Initializer.BindingFlags)
 	, bHasNonGraphOutputs(Initializer.bHasNonGraphOutputs)
 	, bNoEmulatedUniformBuffer(Initializer.bNoEmulatedUniformBuffer)
+	, bUniformView(Initializer.bUniformView)
 {
 }
diff --git a/Engine/Source/Runtime/RHI/Private/RHIStrings.cpp b/Engine/Source/Runtime/RHI/Private/RHIStrings.cpp
--- a/Engine/Source/Runtime/RHI/Private/RHIStrings.cpp
+++ b/Engine/Source/Runtime/RHI/Private/RHIStrings.cpp
@@ -700,8 +700,6 @@ const TCHAR* GetUniformBufferBaseTypeString(EUniformBufferBaseType BaseType)
 		return TEXT("UBMT_RDG_BUFFER_ACCESS_ARRAY");
 	case UBMT_RDG_BUFFER_SRV:
 		return TEXT("UBMT_RDG_BUFFER_SRV");
-	case UBMT_RDG_UNIFORM_BLOCK_SRV:
-		return TEXT("UBMT_RDG_UNIFORM_BLOCK_SRV");
 	case UBMT_RDG_BUFFER_UAV:
 		return TEXT("UBMT_RDG_BUFFER_UAV");
 	case UBMT_RDG_UNIFORM_BUFFER:
diff --git a/Engine/Source/Runtime/RHI/Public/RHIDefinitions.h b/Engine/Source/Runtime/RHI/Public/RHIDefinitions.h
--- a/Engine/Source/Runtime/RHI/Public/RHIDefinitions.h
+++ b/Engine/Source/Runtime/RHI/Public/RHIDefinitions.h
@@ -534,7 +534,6 @@ enum EUniformBufferBaseType : uint8
 	UBMT_RDG_BUFFER_SRV,
 	UBMT_RDG_BUFFER_UAV,
 	UBMT_RDG_UNIFORM_BUFFER,
-	UBMT_RDG_UNIFORM_BLOCK_SRV,
 
 	// Nested structure.
 	UBMT_NESTED_STRUCT,
@@ -1275,8 +1274,7 @@ inline bool IsRDGBufferReferenceShaderParameterType(EUniformBufferBaseType BaseT
 		BaseType == UBMT_RDG_BUFFER_SRV ||
 		BaseType == UBMT_RDG_BUFFER_UAV ||
 		BaseType == UBMT_RDG_BUFFER_ACCESS ||
-		BaseType == UBMT_RDG_BUFFER_ACCESS_ARRAY ||
-		BaseType == UBMT_RDG_UNIFORM_BLOCK_SRV;
+		BaseType == UBMT_RDG_BUFFER_ACCESS_ARRAY;
 }
 
 /** Returns whether the shader parameter type is for RDG access and not actually for shaders. */
@@ -1304,7 +1302,6 @@ inline bool IsShaderParameterTypeForUniformBufferLayout(EUniformBufferBaseType B
 		BaseType == UBMT_SRV ||
 		BaseType == UBMT_SAMPLER ||
 		BaseType == UBMT_UAV ||
-		BaseType == UBMT_RDG_UNIFORM_BLOCK_SRV ||
 
 		// RHI is able to access RHI resources from RDG.
 		IsRDGResourceReferenceShaderParameterType(BaseType) ||
diff --git a/Engine/Source/Runtime/RHI/Public/RHIResources.h b/Engine/Source/Runtime/RHI/Public/RHIResources.h
--- a/Engine/Source/Runtime/RHI/Public/RHIResources.h
+++ b/Engine/Source/Runtime/RHI/Public/RHIResources.h
@@ -1051,6 +1051,9 @@ struct FRHIUniformBufferLayout : public FRHIResource
 	/** Used for platforms which use emulated ub's, forces a real uniform buffer instead */
 	const bool bNoEmulatedUniformBuffer;
 
+	/** This struct is a view into uniform buffer object, on platforms that support UBO */
+	const bool bUniformView;
+
 	/** Compare two uniform buffer layouts. */
 	friend inline bool operator==(const FRHIUniformBufferLayout& A, const FRHIUniformBufferLayout& B)
 	{
diff --git a/Engine/Source/Runtime/RHI/Public/RHIUniformBufferLayoutInitializer.h b/Engine/Source/Runtime/RHI/Public/RHIUniformBufferLayoutInitializer.h
--- a/Engine/Source/Runtime/RHI/Public/RHIUniformBufferLayoutInitializer.h
+++ b/Engine/Source/Runtime/RHI/Public/RHIUniformBufferLayoutInitializer.h
@@ -137,6 +137,7 @@ struct FRHIUniformBufferLayoutInitializer
 		Ar << Ref.Name;
 		Ar << Ref.Hash;
 		Ar << Ref.bNoEmulatedUniformBuffer;
+		Ar << Ref.bUniformView;
 		return Ar;
 	}
 
@@ -185,6 +186,9 @@ public:
 	/** Used for platforms which use emulated ub's, forces a real uniform buffer instead */
 	LAYOUT_FIELD_INITIALIZED(bool, bNoEmulatedUniformBuffer, false);
 
+	/** This struct is a view into uniform buffer object, on platforms that support UBO */
+	LAYOUT_FIELD_INITIALIZED(bool, bUniformView, false);
+	
 	/** Compare two uniform buffer layout initializers. */
 	friend inline bool operator==(const FRHIUniformBufferLayoutInitializer& A, const FRHIUniformBufferLayoutInitializer& B)
 	{
diff --git a/Engine/Source/Runtime/RHICore/Internal/RHIUniformBufferDataShared.h b/Engine/Source/Runtime/RHICore/Internal/RHIUniformBufferDataShared.h
--- a/Engine/Source/Runtime/RHICore/Internal/RHIUniformBufferDataShared.h
+++ b/Engine/Source/Runtime/RHICore/Internal/RHIUniformBufferDataShared.h
@@ -52,7 +52,6 @@ namespace UE::RHICore
 		break;
 		case UBMT_RDG_TEXTURE_SRV:
 		case UBMT_RDG_BUFFER_SRV:
-		case UBMT_RDG_UNIFORM_BLOCK_SRV:
 		{
 			FRDGShaderResourceView* RDGShaderResourceView = Reader.Read<FRDGShaderResourceView*>(Resource);
 			if (RDGShaderResourceView)
diff --git a/Engine/Source/Runtime/RenderCore/Private/RenderGraphBuilder.cpp b/Engine/Source/Runtime/RenderCore/Private/RenderGraphBuilder.cpp
--- a/Engine/Source/Runtime/RenderCore/Private/RenderGraphBuilder.cpp
+++ b/Engine/Source/Runtime/RenderCore/Private/RenderGraphBuilder.cpp
@@ -277,7 +277,6 @@ void EnumerateBufferAccess(FRDGParameterStruct PassParameters, ERDGPassFlags Pas
 		}
 		break;
 		case UBMT_RDG_BUFFER_SRV:
-		case UBMT_RDG_UNIFORM_BLOCK_SRV:
 			if (FRDGBufferSRVRef SRV = Parameter.GetAsBufferSRV())
 			{
 				FRDGBufferRef Buffer = SRV->GetParent();
@@ -710,7 +709,6 @@ FRHIUniformBuffer* FRDGBuilder::ConvertToExternalUniformBuffer(FRDGUniformBuffer
 			}
 			break;
 			case UBMT_RDG_BUFFER_SRV:
-			case UBMT_RDG_UNIFORM_BLOCK_SRV:
 			{
 				ConvertBuffer(this, Param.GetAsBufferSRV()->Desc.Buffer);
 				InitRHI(Param.GetAsView());
diff --git a/Engine/Source/Runtime/RenderCore/Private/RenderGraphUtils.cpp b/Engine/Source/Runtime/RenderCore/Private/RenderGraphUtils.cpp
--- a/Engine/Source/Runtime/RenderCore/Private/RenderGraphUtils.cpp
+++ b/Engine/Source/Runtime/RenderCore/Private/RenderGraphUtils.cpp
@@ -34,8 +34,7 @@ void ClearUnusedGraphResourcesImpl(
 			Type == UBMT_RDG_TEXTURE_SRV ||
 			Type == UBMT_RDG_TEXTURE_UAV ||
 			Type == UBMT_RDG_BUFFER_SRV ||
-			Type == UBMT_RDG_BUFFER_UAV ||
-			Type == UBMT_RDG_UNIFORM_BLOCK_SRV)
+			Type == UBMT_RDG_BUFFER_UAV)
 		{
 			const TMemoryImageArray<FShaderParameterBindings::FResourceParameter>& ResourceParameters = ShaderBindings.ResourceParameters;
 			const int32 ShaderResourceCount = ResourceParameters.Num();
@@ -119,8 +118,7 @@ void ClearUnusedGraphResourcesImpl(
 			Type == UBMT_RDG_TEXTURE_SRV ||
 			Type == UBMT_RDG_TEXTURE_UAV ||
 			Type == UBMT_RDG_BUFFER_SRV ||
-			Type == UBMT_RDG_BUFFER_UAV ||
-			Type == UBMT_RDG_UNIFORM_BLOCK_SRV)
+			Type == UBMT_RDG_BUFFER_UAV)
 		{
 			for (int32 Index = 0; Index < ShaderBindingsList.Num(); ++Index)
 			{
diff --git a/Engine/Source/Runtime/RenderCore/Private/RenderGraphValidation.cpp b/Engine/Source/Runtime/RenderCore/Private/RenderGraphValidation.cpp
--- a/Engine/Source/Runtime/RenderCore/Private/RenderGraphValidation.cpp
+++ b/Engine/Source/Runtime/RenderCore/Private/RenderGraphValidation.cpp
@@ -882,7 +882,6 @@ void FRDGUserValidation::ValidateAddPass(const FRDGPass* Pass, bool bSkipPassAcc
 		}
 		break;
 		case UBMT_RDG_BUFFER_SRV:
-		case UBMT_RDG_UNIFORM_BLOCK_SRV:
 		{
 			if (FRDGBufferSRVRef SRV = Parameter.GetAsBufferSRV())
 			{
diff --git a/Engine/Source/Runtime/RenderCore/Private/ShaderParameterMetadata.cpp b/Engine/Source/Runtime/RenderCore/Private/ShaderParameterMetadata.cpp
--- a/Engine/Source/Runtime/RenderCore/Private/ShaderParameterMetadata.cpp
+++ b/Engine/Source/Runtime/RenderCore/Private/ShaderParameterMetadata.cpp
@@ -197,7 +197,6 @@ const TCHAR* const kShaderParameterMacroNames[] = {
 	TEXT("SHADER_PARAMETER_RDG_BUFFER_SRV"), // UBMT_RDG_BUFFER_SRV,
 	TEXT("SHADER_PARAMETER_RDG_BUFFER_UAV"), // UBMT_RDG_BUFFER_UAV,
 	TEXT("SHADER_PARAMETER_RDG_UNIFORM_BUFFER"), // UBMT_RDG_UNIFORM_BUFFER,
-	TEXT("SHADER_PARAMETER_RDG_UNIFORM_BLOCK_SRV"), // UBMT_RDG_UNIFORM_BLOCK_SRV,
 
 	// Nested structure.
 	TEXT("SHADER_PARAMETER_STRUCT"), // UBMT_NESTED_STRUCT,
@@ -463,8 +462,9 @@ void FShaderParametersMetadata::InitializeLayout(FRHIUniformBufferLayoutInitiali
 	FRHIUniformBufferLayoutInitializer LocalLayoutInitializer(LayoutName);
 	FRHIUniformBufferLayoutInitializer& LayoutInitializer = OutLayoutInitializer ? *OutLayoutInitializer : LocalLayoutInitializer;
 	LayoutInitializer.ConstantBufferSize = Size;
-	LayoutInitializer.bNoEmulatedUniformBuffer = UsageFlags & (uint32)EUsageFlags::NoEmulatedUniformBuffer;
-
+	LayoutInitializer.bUniformView = UsageFlags & (uint32)EUsageFlags::UniformView;
+	LayoutInitializer.bNoEmulatedUniformBuffer = LayoutInitializer.bUniformView || (UsageFlags & (uint32)EUsageFlags::NoEmulatedUniformBuffer);
+	
 	if (StaticSlotName)
 	{
 		checkf(EnumHasAnyFlags(BindingFlags, EUniformBufferBindingFlags::Static), TEXT("Uniform buffer of type '%s' and shader name '%s' attempted to reference static slot '%s', but the binding model does not contain 'Global'."),
@@ -624,8 +624,7 @@ void FShaderParametersMetadata::InitializeLayout(FRHIUniformBufferLayoutInitiali
 				{
 					bIsValidBindingType = (
 						BaseType == UBMT_SRV ||
-						BaseType == UBMT_RDG_BUFFER_SRV ||
-						BaseType == UBMT_RDG_UNIFORM_BLOCK_SRV);
+						BaseType == UBMT_RDG_BUFFER_SRV);
 				}
 				else if (BindingType == EShaderCodeResourceBindingType::RaytracingAccelerationStructure)
 				{
@@ -958,7 +957,7 @@ void FShaderParametersMetadata::AddResourceTableEntries(FShaderResourceTableMap&
 	UniformBufferEntry.StaticSlotName = StaticSlotName;
 	UniformBufferEntry.LayoutHash = IsLayoutInitialized() ? GetLayout().GetHash() : 0;
 	UniformBufferEntry.BindingFlags = BindingFlags;
-	UniformBufferEntry.bNoEmulatedUniformBuffer = UsageFlags & (uint32)EUsageFlags::NoEmulatedUniformBuffer;
+	UniformBufferEntry.bNoEmulatedUniformBuffer = (UsageFlags & (uint32)EUsageFlags::NoEmulatedUniformBuffer) || (UsageFlags & (uint32)EUsageFlags::UniformView);
 	UniformBufferEntry.MemberNameBuffer = MemberNameBuffer;
 	UniformBufferMap.AddByHash(ShaderVariableNameHash, ShaderVariableName, UniformBufferEntry);
 }
diff --git a/Engine/Source/Runtime/RenderCore/Private/ShaderParameterStruct.cpp b/Engine/Source/Runtime/RenderCore/Private/ShaderParameterStruct.cpp
--- a/Engine/Source/Runtime/RenderCore/Private/ShaderParameterStruct.cpp
+++ b/Engine/Source/Runtime/RenderCore/Private/ShaderParameterStruct.cpp
@@ -558,7 +558,6 @@ void ValidateShaderParameters(const TShaderRef<FShader>& Shader, const FShaderPa
 			case UBMT_RDG_BUFFER_SRV:
 			case UBMT_RDG_TEXTURE_UAV:
 			case UBMT_RDG_BUFFER_UAV:
-			case UBMT_RDG_UNIFORM_BLOCK_SRV:
 			{
 				const FRDGResource* GraphResource = Reader.Read<const FRDGResource*>(Parameter);
 
@@ -621,8 +620,7 @@ void ValidateShaderParameterResourcesRHI(const void* Contents, const FRHIUniform
 
 		const bool bSRV =
 			Parameter.MemberType == UBMT_SRV ||
-			Parameter.MemberType == UBMT_RDG_TEXTURE_SRV ||
-			Parameter.MemberType == UBMT_RDG_BUFFER_SRV;
+			Parameter.MemberType == UBMT_RDG_TEXTURE_SRV;
 
 		// Allow null SRV's in uniform buffers for feature levels that don't support SRV's in shaders
 		if (GMaxRHIFeatureLevel <= ERHIFeatureLevel::ES3_1 && bSRV)
@@ -702,7 +700,6 @@ FRHIShaderParameterResource ExtractShaderParameterResource(FShaderParameterReade
 	}
 	case UBMT_RDG_TEXTURE_SRV:
 	case UBMT_RDG_BUFFER_SRV:
-	case UBMT_RDG_UNIFORM_BLOCK_SRV:
 	{
 		FRDGShaderResourceView* RDGShaderResourceView = Reader.Read<FRDGShaderResourceView*>(Parameter);
 		checkSlow(RDGShaderResourceView);
@@ -900,7 +897,6 @@ void SetShaderParameters(
 		break;
 		case UBMT_RDG_TEXTURE_SRV:
 		case UBMT_RDG_BUFFER_SRV:
-		case UBMT_RDG_UNIFORM_BLOCK_SRV:
 		{
 			FRDGShaderResourceView* RDGShaderResourceView = Reader.Read<FRDGShaderResourceView*>(Parameter);
 			checkSlow(RDGShaderResourceView);
diff --git a/Engine/Source/Runtime/RenderCore/Private/ShaderParameters.cpp b/Engine/Source/Runtime/RenderCore/Private/ShaderParameters.cpp
--- a/Engine/Source/Runtime/RenderCore/Private/ShaderParameters.cpp
+++ b/Engine/Source/Runtime/RenderCore/Private/ShaderParameters.cpp
@@ -156,6 +156,21 @@ static void CreateHLSLUniformBufferStructMembersDeclaration(
 {
 	const TArray<FShaderParametersMetadata::FMember>& StructMembers = UniformBufferStruct.GetMembers();
 	
+	const bool bUniformView = (UniformBufferStruct.GetUsageFlags() & (uint32)FShaderParametersMetadata::EUsageFlags::UniformView) != 0;
+	if (bUniformView)
+	{
+		// UniformView struct is expected to have a single SRV member which serves as a uniform view
+		check(StructMembers.Num() == 1);
+		const FShaderParametersMetadata::FMember& Member = StructMembers[0];
+		check(Member.GetBaseType() == UBMT_SRV || Member.GetBaseType() == UBMT_RDG_BUFFER_SRV);
+		
+		FString ParameterName = FString::Printf(TEXT("%s%s"),*NamePrefix,Member.GetName());
+		Decl.ConstantBufferMembers += FString::Printf(TEXT("UB_MEMBER_UNIFORM_BLOCK(%s);\r\n"), *ParameterName);
+		Decl.StructMembers += FString::Printf(TEXT("UB_MEMBER_UNIFORM_BLOCK(%s);\r\n"), Member.GetName());
+		Decl.Initializer += FString::Printf(TEXT("%s,"),*ParameterName);
+		return;
+	}
+		
 	const TCHAR* PreviousBaseTypeName = TEXT("UB_FLOAT");
 	for (int32 MemberIndex = 0; MemberIndex < StructMembers.Num(); ++MemberIndex)
 	{
@@ -176,13 +191,6 @@ static void CreateHLSLUniformBufferStructMembersDeclaration(
 			Decl.Initializer += TEXT("},");
 			Decl.StructMembers += FString::Printf(TEXT("} %s%s;\r\n"),Member.GetName(),*ArrayDim);
 		}
-		else if (Member.GetBaseType() == UBMT_RDG_UNIFORM_BLOCK_SRV)
-		{
-			FString ParameterName = FString::Printf(TEXT("%s%s"),*NamePrefix,Member.GetName());
-			Decl.ConstantBufferMembers += FString::Printf(TEXT("UB_MEMBER_UNIFORM_BLOCK(%s);\r\n"), *ParameterName);
-			Decl.StructMembers += FString::Printf(TEXT("UB_MEMBER_UNIFORM_BLOCK(%s);\r\n"), Member.GetName());
-			Decl.Initializer += FString::Printf(TEXT("%s,"),*ParameterName);
-		}
 		else if (Member.GetBaseType() == UBMT_INCLUDED_STRUCT)
 		{
 			CreateHLSLUniformBufferStructMembersDeclaration(*Member.GetStructMetadata(), NamePrefix, StructOffset + Member.GetOffset(), Decl, HLSLBaseOffset);
@@ -314,10 +322,6 @@ static void CreateHLSLUniformBufferStructMembersDeclaration(
 				Decl.StructMembers += FString::Printf(TEXT("\t%s %s;\r\n"), Member.GetShaderType(), Member.GetName());
 				Decl.Initializer += FString::Printf(TEXT("%s,"), *ParameterName);
 			}
-			else if (Member.GetBaseType() == UBMT_RDG_UNIFORM_BLOCK_SRV)
-			{
-				// do nothing
-			}
 			else
 			{
 				Decl.ResourceMembers += FString::Printf(TEXT("UB_RESOURCE_MEMBER_RESOURCE(%s, %s);\r\n"), Member.GetShaderType(), *ParameterName);
diff --git a/Engine/Source/Runtime/RenderCore/Public/RenderGraphParameter.h b/Engine/Source/Runtime/RenderCore/Public/RenderGraphParameter.h
--- a/Engine/Source/Runtime/RenderCore/Public/RenderGraphParameter.h
+++ b/Engine/Source/Runtime/RenderCore/Public/RenderGraphParameter.h
@@ -18,7 +18,7 @@ public:
 
 	bool IsSRV() const
 	{
-		return MemberType == UBMT_RDG_TEXTURE_SRV || MemberType == UBMT_RDG_BUFFER_SRV || MemberType == UBMT_RDG_UNIFORM_BLOCK_SRV;
+		return MemberType == UBMT_RDG_TEXTURE_SRV || MemberType == UBMT_RDG_BUFFER_SRV;
 	}
 
 	bool IsUAV() const
@@ -168,7 +168,7 @@ public:
 
 	FRDGBufferSRVRef GetAsBufferSRV() const
 	{
-		check(MemberType == UBMT_RDG_BUFFER_SRV || MemberType == UBMT_RDG_UNIFORM_BLOCK_SRV);
+		check(MemberType == UBMT_RDG_BUFFER_SRV);
 		return *GetAs<FRDGBufferSRVRef>();
 	}
 
diff --git a/Engine/Source/Runtime/RenderCore/Public/ShaderParameterMacros.h b/Engine/Source/Runtime/RenderCore/Public/ShaderParameterMacros.h
--- a/Engine/Source/Runtime/RenderCore/Public/ShaderParameterMacros.h
+++ b/Engine/Source/Runtime/RenderCore/Public/ShaderParameterMacros.h
@@ -1627,11 +1627,6 @@ private:
 #define SHADER_PARAMETER_SRV_ARRAY(ShaderType,MemberName, ArrayDecl) \
 	INTERNAL_SHADER_PARAMETER_EXPLICIT(UBMT_SRV, TShaderResourceParameterTypeInfo<FRHIShaderResourceView* ArrayDecl>, FRHIShaderResourceView*,MemberName,ArrayDecl,,EShaderPrecisionModifier::Float,TEXT(#ShaderType),false)
 
-
-#define SHADER_PARAMETER_UNIFORM_BLOCK_SRV(MemberName) \
-	INTERNAL_SHADER_PARAMETER_EXPLICIT(UBMT_RDG_UNIFORM_BLOCK_SRV, TShaderResourceParameterTypeInfo<FRDGBufferSRV*>, FRDGBufferSRV*,MemberName,, = nullptr,EShaderPrecisionModifier::Float,TEXT("Buffer<float4>"),false)
-
-
 /** Adds an unordered access view.
  *
  * Example:
diff --git a/Engine/Source/Runtime/RenderCore/Public/ShaderParameterMetadata.h b/Engine/Source/Runtime/RenderCore/Public/ShaderParameterMetadata.h
--- a/Engine/Source/Runtime/RenderCore/Public/ShaderParameterMetadata.h
+++ b/Engine/Source/Runtime/RenderCore/Public/ShaderParameterMetadata.h
@@ -155,6 +155,9 @@ public:
 
 		/** On platforms that support emulated uniform buffers, disable them for this uniform buffer */
 		NoEmulatedUniformBuffer = 1 << 0,
+		
+		/** This struct is a view into uniform buffer object, on platforms that support UBO */
+		UniformView = 1 << 1,
 	};
 
 	/** Shader binding name of the uniform buffer that contains the root shader parameters. */
diff --git a/Engine/Source/Runtime/RenderCore/Public/ShaderParameterStruct.h b/Engine/Source/Runtime/RenderCore/Public/ShaderParameterStruct.h
--- a/Engine/Source/Runtime/RenderCore/Public/ShaderParameterStruct.h
+++ b/Engine/Source/Runtime/RenderCore/Public/ShaderParameterStruct.h
@@ -149,8 +149,7 @@ inline void CollectSRVsToUnset(FRHIBatchedShaderUnbinds& BatchedUnbinds, const T
 	{
 		if (Parameter.BaseType == UBMT_SRV ||
 			Parameter.BaseType == UBMT_RDG_TEXTURE_SRV ||
-			Parameter.BaseType == UBMT_RDG_BUFFER_SRV ||
-			Parameter.BaseType == UBMT_RDG_UNIFORM_BLOCK_SRV)
+			Parameter.BaseType == UBMT_RDG_BUFFER_SRV)
 		{
 			BatchedUnbinds.UnsetSRV(GetParameterIndex(Parameter));
 		}
diff --git a/Engine/Source/Runtime/Renderer/Private/InstanceCulling/InstanceCullingContext.cpp b/Engine/Source/Runtime/Renderer/Private/InstanceCulling/InstanceCullingContext.cpp
--- a/Engine/Source/Runtime/Renderer/Private/InstanceCulling/InstanceCullingContext.cpp
+++ b/Engine/Source/Runtime/Renderer/Private/InstanceCulling/InstanceCullingContext.cpp
@@ -42,7 +42,7 @@ IMPLEMENT_STATIC_UNIFORM_BUFFER_SLOT(InstanceCullingUbSlot);
 IMPLEMENT_STATIC_UNIFORM_BUFFER_STRUCT(FInstanceCullingGlobalUniforms, "InstanceCulling", InstanceCullingUbSlot);
 
 IMPLEMENT_STATIC_UNIFORM_BUFFER_SLOT(BatchedPrimitive);
-IMPLEMENT_STATIC_AND_SHADER_UNIFORM_BUFFER_STRUCT_EX(FBatchedPrimitiveParameters, "BatchedPrimitive", BatchedPrimitive, FShaderParametersMetadata::EUsageFlags::NoEmulatedUniformBuffer);
+IMPLEMENT_STATIC_AND_SHADER_UNIFORM_BUFFER_STRUCT_EX(FBatchedPrimitiveParameters, "BatchedPrimitive", BatchedPrimitive, FShaderParametersMetadata::EUsageFlags::UniformView);
 
 static const TCHAR* BatchProcessingModeStr[] =
 {
diff --git a/Engine/Source/Runtime/Renderer/Public/InstanceCulling/InstanceCullingContext.h b/Engine/Source/Runtime/Renderer/Public/InstanceCulling/InstanceCullingContext.h
--- a/Engine/Source/Runtime/Renderer/Public/InstanceCulling/InstanceCullingContext.h
+++ b/Engine/Source/Runtime/Renderer/Public/InstanceCulling/InstanceCullingContext.h
@@ -26,7 +26,7 @@ BEGIN_GLOBAL_SHADER_PARAMETER_STRUCT(FInstanceCullingGlobalUniforms, RENDERER_AP
 END_GLOBAL_SHADER_PARAMETER_STRUCT()
 
 BEGIN_GLOBAL_SHADER_PARAMETER_STRUCT(FBatchedPrimitiveParameters,RENDERER_API)
-	SHADER_PARAMETER_UNIFORM_BLOCK_SRV(Data)
+	SHADER_PARAMETER_RDG_BUFFER_SRV(Buffer<float4>, Data)
 END_GLOBAL_SHADER_PARAMETER_STRUCT()
 
 BEGIN_SHADER_PARAMETER_STRUCT(FInstanceCullingDrawParams, )
diff --git a/Engine/Source/Runtime/VulkanRHI/Private/VulkanUniformBuffer.cpp b/Engine/Source/Runtime/VulkanRHI/Private/VulkanUniformBuffer.cpp
--- a/Engine/Source/Runtime/VulkanRHI/Private/VulkanUniformBuffer.cpp
+++ b/Engine/Source/Runtime/VulkanRHI/Private/VulkanUniformBuffer.cpp
@@ -105,15 +105,23 @@ bool FVulkanUniformBuffer::SetupUniformBufferView(const FRHIUniformBufferLayout*
 {
 	bUniformView = false;
 
-	if (InLayout->Resources.Num() == 1 &&
-		InLayout->Resources[0].MemberType == UBMT_RDG_UNIFORM_BLOCK_SRV)
+	if (InLayout->bUniformView)
 	{
-		FRHIShaderResourceView* SRV = (FRHIShaderResourceView*)GetShaderParameterResourceRHI(Contents, InLayout->Resources[0].MemberOffset, UBMT_RDG_UNIFORM_BLOCK_SRV);
-		ResourceTable.Empty(1);
-		ResourceTable.Add(SRV);
-		
-		FVulkanResourceMultiBuffer* Buffer = ResourceCast(SRV->GetBuffer());
-		const FRHIViewDesc::FBufferSRV& SRVInfo = SRV->GetDesc().Buffer.SRV;
+		FRHIShaderResourceView* UniformViewSRV = nullptr;
+		for (int32 Index = 0; Index < InLayout->Resources.Num() && !UniformViewSRV; ++Index)
+		{
+			EUniformBufferBaseType ResourceBaseType = InLayout->Resources[Index].MemberType;
+			if (ResourceBaseType == UBMT_SRV || 
+				ResourceBaseType == UBMT_RDG_BUFFER_SRV)
+			{
+				UniformViewSRV = (FRHIShaderResourceView*)GetShaderParameterResourceRHI(Contents, InLayout->Resources[Index].MemberOffset, ResourceBaseType);
+			}
+		}
+
+		check(UniformViewSRV)
+				
+		FVulkanResourceMultiBuffer* Buffer = ResourceCast(UniformViewSRV->GetBuffer());
+		const FRHIViewDesc::FBufferSRV& SRVInfo = UniformViewSRV->GetDesc().Buffer.SRV;
 		
 		Allocation.Reference(Buffer->GetCurrentAllocation());
 		//Adjust Allocation.Size ???
@@ -140,11 +148,6 @@ FVulkanUniformBuffer::FVulkanUniformBuffer(FVulkanDevice& InDevice, const FRHIUn
 	check(InLayout->Resources.Num() > 0 || InLayout->ConstantBufferSize > 0);
 	const uint32 NumResources = InLayout->Resources.Num();
 
-	if (SetupUniformBufferView(InLayout, Contents))
-	{
-		return;
-	}
-		
 	// Setup resource table
 	if (NumResources > 0)
 	{
@@ -161,6 +164,11 @@ FVulkanUniformBuffer::FVulkanUniformBuffer(FVulkanDevice& InDevice, const FRHIUn
 		}
 	}
 
+	if (SetupUniformBufferView(InLayout, Contents))
+	{
+		return;
+	}
+
 	if (InLayout->ConstantBufferSize > 0)
 	{
 		const bool bInRenderingThread = IsInRenderingThread();
