Make Glint/SpecularLUT/SheenQuality be set per platform.

diff --git a/Engine/Config/Android/AndroidEngine.ini b/Engine/Config/Android/AndroidEngine.ini
--- a/Engine/Config/Android/AndroidEngine.ini
+++ b/Engine/Config/Android/AndroidEngine.ini
@@ -93,6 +93,12 @@ r.Vulkan.MemoryBacktrace=0
 [/Script/Engine.RendererSettings]
 r.DistanceFields=1
 r.TemporalAA.Mobile.UseCompute=0
+r.Substrate.RoughDiffuse=0
+r.Substrate.ShadingQuality=2
+r.Substrate.TileCoord8bits=1
+r.Substrate.SheenQuality=2
+r.Substrate.Glints=0
+r.Substrate.SpecularProfile=0
 ; Enable GPU culling for landscape
 landscape.SupportGPUCulling=1
 
diff --git a/Engine/Config/BaseDeviceProfiles.ini b/Engine/Config/BaseDeviceProfiles.ini
--- a/Engine/Config/BaseDeviceProfiles.ini
+++ b/Engine/Config/BaseDeviceProfiles.ini
@@ -865,11 +865,7 @@ BaseProfileName=Android
 +CVars=r.DistanceFields=1
 +CVars=r.Vulkan.RayTracing.AllowCompaction=0
 +CVars=r.Vulkan.RayTracing.TLASPreferFastTraceTLAS=0
-; Disable rough diffuse for saving lighting cost
-+CVars=r.Substrate.RoughDiffuse=0
-; Reduce shading quality for performance
-+CVars=r.Substrate.ShadingQuality=2
-+CVars=r.Substrate.TileCoord8bits=1
+
 
 [Android_Vulkan_SM5_Mali DeviceProfile]
 DeviceType=Android
diff --git a/Engine/Config/IOS/IOSEngine.ini b/Engine/Config/IOS/IOSEngine.ini
--- a/Engine/Config/IOS/IOSEngine.ini
+++ b/Engine/Config/IOS/IOSEngine.ini
@@ -80,6 +80,12 @@ fx.Niagara.QualityLevel.Max=1
 [/Script/Engine.RendererSettings]
 r.DistanceFields=0
 r.TemporalAA.Mobile.UseCompute=0
+r.Substrate.RoughDiffuse=0
+r.Substrate.ShadingQuality=2
+r.Substrate.TileCoord8bits=1
+r.Substrate.SheenQuality=2
+r.Substrate.Glints=0
+r.Substrate.SpecularProfile=0
 ; Enable GPU culling for landscape
 landscape.SupportGPUCulling=1
 
diff --git a/Engine/Source/Runtime/Engine/Private/ShaderCompiler/ShaderCompiler.cpp b/Engine/Source/Runtime/Engine/Private/ShaderCompiler/ShaderCompiler.cpp
--- a/Engine/Source/Runtime/Engine/Private/ShaderCompiler/ShaderCompiler.cpp
+++ b/Engine/Source/Runtime/Engine/Private/ShaderCompiler/ShaderCompiler.cpp
@@ -8000,10 +8000,10 @@ void GlobalBeginCompileShader(
 			const bool bSubstrateDBufferPass = Substrate::IsDBufferPassEnabled(Target.GetPlatform());
 			SET_SHADER_DEFINE(Input.Environment, SUBSTRATE_USE_DBUFFER_PASS, bSubstrateDBufferPass ? 1 : 0);
 
-			const bool bSubstrateGlints = Substrate::IsGlintEnabled();
+			const bool bSubstrateGlints = Substrate::IsGlintEnabled(Target.GetPlatform());
 			SET_SHADER_DEFINE(Input.Environment, PLATFORM_ENABLES_SUBSTRATE_GLINTS, bSubstrateGlints ? 1 : 0);
 
-			const bool bSpecularProfileEnabled = Substrate::IsSpecularProfileEnabled();
+			const bool bSpecularProfileEnabled = Substrate::IsSpecularProfileEnabled(Target.GetPlatform());
 			SET_SHADER_DEFINE(Input.Environment, PLATFORM_ENABLES_SUBSTRATE_SPECULAR_PROFILE, bSpecularProfileEnabled ? 1 : 0);
 		}
 
diff --git a/Engine/Source/Runtime/RenderCore/Private/RenderUtils.cpp b/Engine/Source/Runtime/RenderCore/Private/RenderUtils.cpp
--- a/Engine/Source/Runtime/RenderCore/Private/RenderUtils.cpp
+++ b/Engine/Source/Runtime/RenderCore/Private/RenderUtils.cpp
@@ -1802,6 +1802,12 @@ namespace Substrate
 		return IsSubstrateEnabled() && CVarSubstrateGlints.GetValueOnAnyThread() > 0;
 	}
 
+	bool IsGlintEnabled(EShaderPlatform InPlatform)
+	{
+		static FShaderPlatformCachedIniValue<int32> CVar(TEXT("r.Substrate.Glints"));
+		return IsSubstrateEnabled() && CVar.Get(InPlatform);
+	}
+
 	uint32 GlintLUTIndex()
 	{
 		return CVarSubstrateGlintsLUT.GetValueOnAnyThread()<= 0 ? 0u : 1u;
@@ -1823,6 +1829,12 @@ namespace Substrate
 		return IsSubstrateEnabled() && CVarSubstrateSpecularProfile.GetValueOnAnyThread() > 0;
 	}
 
+	bool IsSpecularProfileEnabled(EShaderPlatform InPlatform)
+	{
+		static FShaderPlatformCachedIniValue<int32> CVar(TEXT("r.Substrate.SpecularProfile"));
+		return IsSubstrateEnabled() && CVar.Get(InPlatform) > 0;
+	}
+
 	bool Is8bitTileCoordEnabled()
 	{
 		return CVarSubstrateTileCoord8Bits.GetValueOnAnyThread() > 0 ? 1 : 0;
@@ -1833,6 +1845,12 @@ namespace Substrate
 		return CVarSubstrateSheenQuality.GetValueOnAnyThread();
 	}
 
+	uint32 GetSheenQuality(EShaderPlatform InPlatform)
+	{
+		static FShaderPlatformCachedIniValue<int32> CVar(TEXT("r.Substrate.SheenQuality"));
+		return CVar.Get(InPlatform);
+	}
+
 	uint32 GetShadingQuality()
 	{
 		return CVarSubstrateShadingQuality.GetValueOnAnyThread();
diff --git a/Engine/Source/Runtime/RenderCore/Private/Shader.cpp b/Engine/Source/Runtime/RenderCore/Private/Shader.cpp
--- a/Engine/Source/Runtime/RenderCore/Private/Shader.cpp
+++ b/Engine/Source/Runtime/RenderCore/Private/Shader.cpp
@@ -2032,15 +2032,15 @@ void ShaderMapAppendKeyString(EShaderPlatform Platform, FString& KeyString)
 		}
 
 		{
-			KeyString += FString::Printf(TEXT("_SSHEEN%u"), Substrate::GetSheenQuality());
+			KeyString += FString::Printf(TEXT("_SSHEEN%u"), Substrate::GetSheenQuality(Platform));
 		}
 
-		if (Substrate::IsGlintEnabled())
+		if (Substrate::IsGlintEnabled(Platform))
 		{
 			KeyString += FString::Printf(TEXT("_STRTGLT"));
 		}
 
-		if (Substrate::IsSpecularProfileEnabled())
+		if (Substrate::IsSpecularProfileEnabled(Platform))
 		{
 			KeyString += FString::Printf(TEXT("_STRTSP"));
 		}
diff --git a/Engine/Source/Runtime/RenderCore/Public/RenderUtils.h b/Engine/Source/Runtime/RenderCore/Public/RenderUtils.h
--- a/Engine/Source/Runtime/RenderCore/Public/RenderUtils.h
+++ b/Engine/Source/Runtime/RenderCore/Public/RenderUtils.h
@@ -489,10 +489,12 @@ namespace Substrate
 	RENDERCORE_API bool IsSubstrateEnabled();
 	RENDERCORE_API bool IsRoughDiffuseEnabled();
 	RENDERCORE_API bool IsGlintEnabled();
+	RENDERCORE_API bool IsGlintEnabled(EShaderPlatform InPlatform);
 	RENDERCORE_API uint32 GlintLUTIndex();
 	RENDERCORE_API float GlintLevelBias();
 	RENDERCORE_API float GlintLevelMin();
 	RENDERCORE_API bool IsSpecularProfileEnabled();
+	RENDERCORE_API bool IsSpecularProfileEnabled(EShaderPlatform InPlatform);
 	RENDERCORE_API bool IsBackCompatibilityEnabled();
 	RENDERCORE_API bool IsDBufferPassEnabled(EShaderPlatform InPlatform);
 	RENDERCORE_API bool IsOpaqueRoughRefractionEnabled();
@@ -508,6 +510,7 @@ namespace Substrate
 	RENDERCORE_API uint32 GetNormalQuality();
 
 	RENDERCORE_API uint32 GetSheenQuality();
+	RENDERCORE_API uint32 GetSheenQuality(EShaderPlatform InPlatform);
 
 	RENDERCORE_API uint32 GetShadingQuality();
 	RENDERCORE_API uint32 GetShadingQuality(EShaderPlatform InPlatform);
diff --git a/Engine/Source/Runtime/Renderer/Private/IndirectLightRendering.cpp b/Engine/Source/Runtime/Renderer/Private/IndirectLightRendering.cpp
--- a/Engine/Source/Runtime/Renderer/Private/IndirectLightRendering.cpp
+++ b/Engine/Source/Runtime/Renderer/Private/IndirectLightRendering.cpp
@@ -1737,7 +1737,7 @@ static void AddSkyReflectionPass(
 		PassParameters->PS.ForwardLightData = View.ForwardLightingResources.ForwardLightUniformBuffer;
 
 		PassParameters->PS.Substrate = Substrate::BindSubstrateGlobalUniformParameters(View);
-		if (Substrate::IsGlintEnabled())
+		if (Substrate::IsGlintEnabled(View.GetShaderPlatform()))
 		{
 			FBlueNoise BlueNoise = GetBlueNoiseGlobalParameters();
 			PassParameters->PS.BlueNoise = CreateUniformBufferImmediate(BlueNoise, EUniformBufferUsage::UniformBuffer_SingleDraw);
diff --git a/Engine/Source/Runtime/Renderer/Private/SceneRendering.cpp b/Engine/Source/Runtime/Renderer/Private/SceneRendering.cpp
--- a/Engine/Source/Runtime/Renderer/Private/SceneRendering.cpp
+++ b/Engine/Source/Runtime/Renderer/Private/SceneRendering.cpp
@@ -2572,12 +2572,12 @@ FSceneRenderer::FSceneRenderer(const FSceneViewFamily* InViewFamily, FHitProxyCo
 			ShouldRenderLumenReflections(*ViewInfo) || 
 			ShouldRenderVolumetricCloudWithBlueNoise_GameThread(Scene, *ViewInfo) || 
 			UseVirtualShadowMaps(Scene->GetShaderPlatform(), Scene->GetFeatureLevel()) ||
-			Substrate::IsGlintEnabled())
+			Substrate::IsGlintEnabled(ViewInfo->GetShaderPlatform()))
 		{
 			GEngine->LoadBlueNoiseTexture();
 		}
 
-		if (Substrate::IsGlintEnabled())
+		if (Substrate::IsGlintEnabled(ViewInfo->GetShaderPlatform()))
 		{
 			GEngine->LoadGlintTextures();
 		}
diff --git a/Engine/Source/Runtime/Renderer/Private/Substrate/Glint/GlintShadingLUTs.cpp b/Engine/Source/Runtime/Renderer/Private/Substrate/Glint/GlintShadingLUTs.cpp
--- a/Engine/Source/Runtime/Renderer/Private/Substrate/Glint/GlintShadingLUTs.cpp
+++ b/Engine/Source/Runtime/Renderer/Private/Substrate/Glint/GlintShadingLUTs.cpp
@@ -18,7 +18,7 @@ void FGlintShadingLUTsStateData::Init(FRDGBuilder& GraphBuilder, FViewInfo& View
 {
 #if OVERRIDE_GLINT_LUT_ASSSET_WITH_EMBEDDED_DATA
 
-	if (Substrate::IsGlintEnabled())
+	if (Substrate::IsGlintEnabled(View->GetShaderPlatform()))
 	{
 		if (View.ViewState && View.ViewState->GlintShadingLUTsData.GlintShadingLUTs == nullptr)
 		{
