Removing extra include from ShaderCompilerCore.h

diff --git a/Engine/Source/Runtime/RenderCore/Public/ShaderCompilerCore.h b/Engine/Source/Runtime/RenderCore/Public/ShaderCompilerCore.h
--- a/Engine/Source/Runtime/RenderCore/Public/ShaderCompilerCore.h
+++ b/Engine/Source/Runtime/RenderCore/Public/ShaderCompilerCore.h
@@ -18,7 +18,9 @@
 #include "Misc/CoreStats.h"
 #include "ShaderCore.h"
 #include "ShaderParameterMetadata.h"
+#if UE_ENABLE_INCLUDE_ORDER_DEPRECATED_IN_5_4
 #include "ShaderParameterParser.h"
+#endif
 
 class Error;
 class IShaderFormat;
