Reworking UniformBuffer code generation to not use an intermediate struct
- When generating Uniform Buffer HLSL code, a 'static const struct {} Name' was created which has been causing compiler issues off and on for years.
- When RemoveUniformBuffersFromSource removed said struct, it parsed the struct's structure to figure out how to translate UniformBuffer.Member into UniformBuffer_Member in code that accessed Uniform Buffers.
- Now, instead of a struct, a "UniformBuffer <Name>" declaration is created and filled with mappings from "UniformBuffer.Member" into the correct code.
- This allows ConstantBuffer objects to be used differently from cbuffer where the former generates "UniformBuffer.Member" and the latter generates "UniformBuffer_Member"
- This also allows resources to always be translated to "UniformBuffer_Member"
- Shader define 'PLATFORM_SUPPORTS_UB_STRUCT' is now 'PLATFORM_SUPPORTS_CONSTANTBUFFER_OBJECT' and will be set on shader formats that support ConstantBuffer<> objects.
- Removed lots of unused code from BindlessResources.ush

diff --git a/Engine/Shaders/Public/BindlessResources.ush b/Engine/Shaders/Public/BindlessResources.ush
--- a/Engine/Shaders/Public/BindlessResources.ush
+++ b/Engine/Shaders/Public/BindlessResources.ush
@@ -5,106 +5,31 @@
 #include "Platform.ush"
 
 #ifndef ENABLE_BINDLESS_RESOURCES
-#define ENABLE_BINDLESS_RESOURCES 0
+	#define ENABLE_BINDLESS_RESOURCES 0
 #endif
 
 #ifndef ENABLE_BINDLESS_SAMPLERS
-#define ENABLE_BINDLESS_SAMPLERS 0
+	#define ENABLE_BINDLESS_SAMPLERS 0
 #endif
 
-// Indices can always be used as-is. This is important for things like Uniform Buffer layouts.
+#if defined(UB_RESOURCE_MEMBER_RESOURCE)
+	// Using platform defined bindless macros
 
-#define DEFINE_RESOURCE_INDEX(Name)  uint BindlessResource_##Name
-#define DEFINE_SAMPLER_INDEX(Name)   uint BindlessSampler_##Name
-
-//
-
-#if ENABLE_BINDLESS_RESOURCES && PLATFORM_SUPPORTS_BINDLESS
-	#define DEFINE_RESOURCE_VARIABLE(Type, Name) \
-		Type GetBindlessResource##Name() \
-		{ \
-			return GetResourceFromHeap(Type, BindlessResource_##Name); \
-		}
-
-	#define GET_RESOURCE(Name) GetBindlessResource##Name()
-
-	#define DEFINE_RESOURCE(Type, Name) \
-		DEFINE_RESOURCE_INDEX(Name); \
-		DEFINE_RESOURCE_VARIABLE(Type, Name); \
-		static const Type Name = GET_RESOURCE(Name);
-
-	// NOTE: These macros equivalent to FShaderParameterParser::ApplyBindlessModifications and must be kept in sync
-	#define AUTO_BINDLESS_RESOURCE_INDEX(Name)          uint BindlessResource_##Name
-	#define AUTO_BINDLESS_RESOURCE_VARIABLE(Type, Name) Type GetBindlessResource##Name() { return GetResourceFromHeap(Type, BindlessResource_##Name); }; static const Type Name = GetBindlessResource##Name()
-
-	#ifndef UB_RESOURCE_MEMBER_RESOURCE  // allow platforms to override ub member resource decl
-		#define UB_RESOURCE_MEMBER_RESOURCE(Type, Name) AUTO_BINDLESS_RESOURCE_VARIABLE(Type, Name)
-	#endif
+#elif ENABLE_BINDLESS_RESOURCES && PLATFORM_SUPPORTS_BINDLESS
+	#define UB_RESOURCE_MEMBER_RESOURCE(Type, UBName, Name) \
+		Type GetBindlessResource##UBName##_##Name() { return GetResourceFromHeap(Type, UB_CB_MEMBER_ACCESS(UBName, BindlessResource_##Name)); } \
+		static const Type UBName##_##Name = GetBindlessResource##UBName##_##Name();
 #else
-	#define DEFINE_RESOURCE_VARIABLE(Type, Name)  Type Name
-	#define DEFINE_RESOURCE(Type, Name)           Type Name
-	#define GET_RESOURCE(Name)                    Name
-
-	#define UB_RESOURCE_MEMBER_RESOURCE(Type, Name) Type Name
+	#define UB_RESOURCE_MEMBER_RESOURCE(Type, UBName, Name) Type UBName##_##Name
 #endif
 
-#if ENABLE_BINDLESS_SAMPLERS && PLATFORM_SUPPORTS_BINDLESS
-	#define DEFINE_SAMPLER_VARIABLE(Type, Name) \
-		Type GetBindlessSampler##Name() \
-		{ \
-			return GetSamplerFromHeap(Type, BindlessSampler_##Name); \
-		}
-
-	#define GET_SAMPLER(Name)  GetBindlessSampler##Name()
-
-	#define DEFINE_SAMPLER(Type, Name) \
-		DEFINE_SAMPLER_INDEX(Name); \
-		DEFINE_SAMPLER_VARIABLE(Type, Name); \
-		static const Type Name = GET_SAMPLER(Name);
+#if defined(UB_RESOURCE_MEMBER_SAMPLER)
+	// Using platform defined bindless macros
 
-	// NOTE: These macros equivalent to FShaderParameterParser::ApplyBindlessModifications and must be kept in sync
-	#define AUTO_BINDLESS_SAMPLER_INDEX(Name)          uint BindlessSampler_##Name
-	#define AUTO_BINDLESS_SAMPLER_VARIABLE(Type, Name) Type GetBindlessSampler##Name() { return GetSamplerFromHeap(Type, BindlessSampler_##Name); }; static const Type Name = GetBindlessSampler##Name()
-
-	#ifndef UB_RESOURCE_MEMBER_SAMPLER  // allow platforms to override ub member sampler decl
-		#define UB_RESOURCE_MEMBER_SAMPLER(Type, Name) AUTO_BINDLESS_SAMPLER_VARIABLE(Type, Name)
-	#endif
+#elif ENABLE_BINDLESS_SAMPLERS && PLATFORM_SUPPORTS_BINDLESS
+	#define UB_RESOURCE_MEMBER_SAMPLER(Type, UBName, Name) \
+		Type GetBindlessSampler##UBName##_##Name() { return GetSamplerFromHeap(Type, UB_CB_MEMBER_ACCESS(UBName, BindlessSampler_##Name)); } \
+		static const Type UBName##_##Name = GetBindlessSampler##UBName##_##Name();
 #else
-	#define DEFINE_SAMPLER_VARIABLE(Type, Name)  Type Name
-	#define DEFINE_SAMPLER(Type, Name)           Type Name
-	#define GET_SAMPLER(Name)                    Name
-
-	#define UB_RESOURCE_MEMBER_SAMPLER(Type, Name) Type Name
+	#define UB_RESOURCE_MEMBER_SAMPLER(Type, UBName, Name) Type UBName##_##Name
 #endif
-
-/*
-// Uniform buffers will do:
-cbuffer UniformBuffer_MyUniformBuffer
-{
-	DEFINE_RESOURCE_INDEX(MyUniformBuffer_MyTexture);
-	DEFINE_SAMPLER_INDEX(MyUniformBuffer_MySampler);
-};
-
-UB_RESOURCE_MEMBER_RESOURCE(Texture2D, MyUniformBuffer_MyTexture);
-UB_RESOURCE_MEMBER_SAMPLER(SamplerState, MyUniformBuffer_MySampler);
-
-static const struct
-{
-	Texture2D MyTexture;
-	SamplerState MySampler;
-}
-MyUniformBuffer = { MyUniformBuffer_MyTexture, MyUniformBuffer_MySampler };
-
-{	// No wrappers necessary to use
-	MyUniformBuffer.MyTexture.Sample(MyUniformBuffer.MySampler, ...);
-}
-
-// Normal uses will look like:
-DEFINE_RESOURCE(Texture2D, MyTexture);
-DEFINE_SAMPLER(SamplerState, MySampler);
-
-{	// No wrappers necessary to use
-	MyTexture.Sample(MySampler, ...);
-}
-
-*/
diff --git a/Engine/Shaders/Public/Platform.ush b/Engine/Shaders/Public/Platform.ush
--- a/Engine/Shaders/Public/Platform.ush
+++ b/Engine/Shaders/Public/Platform.ush
@@ -249,20 +249,18 @@
 #define COMPILER_SUPPORTS_HLSL2021 0
 #endif
 
-#ifndef PLATFORM_SUPPORTS_UB_STRUCT
-#define PLATFORM_SUPPORTS_UB_STRUCT 0
+#ifndef PLATFORM_SUPPORTS_CONSTANTBUFFER_OBJECT
+#define PLATFORM_SUPPORTS_CONSTANTBUFFER_OBJECT 0
+#endif
+
+#ifndef PLATFORM_SUPPORTS_UNIFORM_BUFFER_OBJECTS
+#define PLATFORM_SUPPORTS_UNIFORM_BUFFER_OBJECTS 0
 #endif
 
 #ifndef PLATFORM_GPU_ARCH
 #define PLATFORM_GPU_ARCH PLATFORM_GPU_ARCH_UNKNOWN
 #endif 
 
-#if PLATFORM_SUPPORTS_UB_STRUCT
-	#define UB_CB_NAME(X) UniformBufferConstants_##X
-#else
-	#define UB_CB_NAME(X) X
-#endif
-
 #ifndef PLATFORM_NEEDS_DEPTH_TEXTURE_READS
 #define PLATFORM_NEEDS_DEPTH_TEXTURE_READS 0
 #endif
@@ -452,6 +450,32 @@
 	#define UB_FIXED_FLOAT(Dim) float##Dim
 #endif
 
+// 16KB by default. Must match Platform.h
+#ifndef PLATFORM_MAX_UNIFORM_BUFFER_RANGE
+	#define PLATFORM_MAX_UNIFORM_BUFFER_RANGE (16u*1024u)
+#endif
+#define PLATFORM_MAX_UNIFORM_BUFFER_RANGE_FLOAT4 (PLATFORM_MAX_UNIFORM_BUFFER_RANGE/16u)
+
+#if PLATFORM_SUPPORTS_CONSTANTBUFFER_OBJECT
+	#define UB_CB_DEFINITION_START(UBName            ) struct F##UBName##Constants {
+	#define UB_CB_MEMBER_NAME(     UBName, MemberName) MemberName
+	#define UB_CB_DEFINITION_END(  UBName            ) }; ConstantBuffer<F##UBName##Constants> UBName;
+	#define UB_CB_MEMBER_ACCESS(   UBName, MemberName) UBName##.##MemberName
+
+	#define UB_DECL_PARAMETER(UBName, StructName, GlobalName) UBName##.##StructName = UBName##.##GlobalName
+	#define UB_DECL_RESOURCE(UBName, StructName, GlobalName)  UBName##.##StructName = UBName##_##GlobalName
+#else
+	#define UB_CB_DEFINITION_START(UBName            ) cbuffer UBName {
+	#define UB_CB_MEMBER_NAME(     UBName, MemberName) UBName##_##MemberName
+	#define UB_CB_DEFINITION_END(  UBName            ) }
+	#define UB_CB_MEMBER_ACCESS(   UBName, MemberName) UBName##_##MemberName
+
+	#define UB_DECL_PARAMETER(UBName, StructName, GlobalName) UBName##.##StructName = UBName##_##GlobalName
+	#define UB_DECL_RESOURCE(UBName, StructName, GlobalName)  UBName##.##StructName = UBName##_##GlobalName
+#endif
+
+#define UB_CB_UNIFORM_BLOCK(UBName, MemberName) UB_FLOAT(4) UB_CB_MEMBER_NAME(UBName,MemberName)[PLATFORM_MAX_UNIFORM_BUFFER_RANGE_FLOAT4]
+
 // Use interpolator for platforms that do not support ClipDistance
 #if !PLATFORM_SUPPORTS_CLIP_DISTANCE
 #define SV_ClipDistance OUTCLIPDIST
@@ -1291,15 +1315,3 @@ uint WaveGetActiveLaneIndexLast()
 #define GET_SCALAR_ARRAY_ELEMENT(PackedArray, ElementIndex) PackedArray[(uint)ElementIndex>>2u][(uint)ElementIndex&3u]
 #endif
 
-#ifndef PLATFORM_SUPPORTS_UNIFORM_BUFFER_OBJECTS
-#define PLATFORM_SUPPORTS_UNIFORM_BUFFER_OBJECTS 0
-#endif 
-
-// 16KB by default. Must match Platform.h
-#ifndef PLATFORM_MAX_UNIFORM_BUFFER_RANGE
-	#define PLATFORM_MAX_UNIFORM_BUFFER_RANGE (16u*1024)
-#endif
-#define PLATFORM_MAX_UNIFORM_BUFFER_RANGE_FLOAT4 (PLATFORM_MAX_UNIFORM_BUFFER_RANGE/16u)
-
-#define UB_MEMBER_UNIFORM_BLOCK(Name) float4 Name[PLATFORM_MAX_UNIFORM_BUFFER_RANGE_FLOAT4]
-
diff --git a/Engine/Shaders/Public/Platform/Vulkan/VulkanCommon.ush b/Engine/Shaders/Public/Platform/Vulkan/VulkanCommon.ush
--- a/Engine/Shaders/Public/Platform/Vulkan/VulkanCommon.ush
+++ b/Engine/Shaders/Public/Platform/Vulkan/VulkanCommon.ush
@@ -42,17 +42,21 @@ void ImageInterlockedMaxUInt64(RWTexture2D<UlongType> Dest, uint2 Coord, UlongTy
 	COMPILER_DEFINE VULKAN_SAMPLER_HEAP(Type) VULKAN_HEAP_NAME(VULKAN_BINDLESS_SAMPLER_ARRAY_PREFIX,Type)
 	COMPILER_DEFINE VULKAN_RESOURCE_HEAP(Type) VULKAN_HEAP_NAME(VULKAN_BINDLESS_RESOURCE_ARRAY_PREFIX,Type)
 
-	// Bindless resource heaps are filled in by the shader compiler
-	COMPILER_DEFINE GetResourceFromHeap(Type, Index)  /* filled by compiler */
-
-	// Override the UB_RESOURCE_MEMBER_RESOURCE to use the Vulkan resource heaps
-	#define UB_RESOURCE_MEMBER_RESOURCE(Type, Name) typedef Type SafeType##Name; SafeType##Name VULKAN_RESOURCE_HEAP(SafeType##Name)[]; static const Type Name = VULKAN_RESOURCE_HEAP(SafeType##Name)[BindlessResource_##Name]
-
-	// bindless sampler heaps (specify a register so AUTO_BINDLESS macros aren't applied on them)
-	COMPILER_DEFINE GetSamplerFromHeap(Type, Index)  /* filled by compiler */
-
-	// Override the UB_RESOURCE_MEMBER_SAMPLER to use the Vulkan sampler heaps
-	#define UB_RESOURCE_MEMBER_SAMPLER(Type, Name) typedef Type SafeType##Name; SafeType##Name VULKAN_SAMPLER_HEAP(SafeType##Name)[]; static const Type Name = VULKAN_SAMPLER_HEAP(SafeType##Name)[BindlessSampler_##Name]
+	#define UB_RESOURCE_MEMBER_RESOURCE(Type, UBName, Name) \
+		typedef Type SafeType##UBName##_##Name; \
+		SafeType##UBName##_##Name VULKAN_RESOURCE_HEAP(SafeType##UBName##_##Name)[]; \
+		Type GetBindlessResource##UBName##_##Name() { \
+			return VULKAN_RESOURCE_HEAP(SafeType##UBName##_##Name)[UB_CB_MEMBER_ACCESS(UBName, BindlessResource_##Name)]; \
+		} \
+		static const Type UBName##_##Name = GetBindlessResource##UBName##_##Name();
+
+	#define UB_RESOURCE_MEMBER_SAMPLER(Type, UBName, Name) \
+		typedef Type SafeType##UBName##_##Name; \
+		SafeType##UBName##_##Name VULKAN_RESOURCE_HEAP(SafeType##UBName##_##Name)[]; \
+		Type GetBindlessSampler##UBName##_##Name() { \
+			return VULKAN_RESOURCE_HEAP(SafeType##UBName##_##Name)[UB_CB_MEMBER_ACCESS(UBName, BindlessSampler_##Name)]; \
+		} \
+		static const Type UBName##_##Name = GetBindlessSampler##UBName##_##Name();
 #endif
 
 #if VULKAN_SUPPORTS_SUBGROUP_SIZE_CONTROL
diff --git a/Engine/Source/Developer/Apple/MetalShaderFormat/Private/MetalShaderCompiler.cpp b/Engine/Source/Developer/Apple/MetalShaderFormat/Private/MetalShaderCompiler.cpp
--- a/Engine/Source/Developer/Apple/MetalShaderFormat/Private/MetalShaderCompiler.cpp
+++ b/Engine/Source/Developer/Apple/MetalShaderFormat/Private/MetalShaderCompiler.cpp
@@ -878,7 +878,7 @@ bool PreprocessMetalShader(const FShaderCompilerInput& Input, const FShaderCompi
 		return false;
 	}
 
-	RemoveUniformBuffersFromSource(Input.Environment, PreprocessOutput.EditSource());
+	CleanupUniformBufferCode(Input.Environment, PreprocessOutput.EditSource());
 
 	// Process TEXT macro.
 	TransformStringIntoCharacterArray(PreprocessOutput.EditSource());
diff --git a/Engine/Source/Developer/ShaderCompilerCommon/Private/HlslParser.cpp b/Engine/Source/Developer/ShaderCompilerCommon/Private/HlslParser.cpp
--- a/Engine/Source/Developer/ShaderCompilerCommon/Private/HlslParser.cpp
+++ b/Engine/Source/Developer/ShaderCompilerCommon/Private/HlslParser.cpp
@@ -1919,6 +1919,22 @@ Done:
 		return EParseResult::Matched;
 	}
 
+	bool MatchConstantBufferNoTemplate(FHlslParser& Parser)
+	{
+		const FHlslToken* Peek = Parser.Scanner.PeekToken();
+		if (Peek && Peek->Token == EHlslToken::ConstantBuffer)
+		{
+			const FHlslToken* Peek1 = Parser.Scanner.PeekToken(1);
+			if (!Peek1 || Peek1->Token != EHlslToken::Lower)
+			{
+				// it's a standard one, make sure we match it
+				return Parser.Scanner.MatchToken(EHlslToken::ConstantBuffer);
+			}
+		}
+
+		return false;
+	}
+
 	EParseResult TryTranslationUnit(FHlslParser& Parser, FLinearAllocator* Allocator, AST::FNode** OutNode)
 	{
 		if (MatchPragma(Parser, Allocator, OutNode))
@@ -1926,7 +1942,7 @@ Done:
 			return EParseResult::Matched;
 		}
 
-		if (Parser.Scanner.MatchToken(EHlslToken::ConstantBuffer) || Parser.Scanner.MatchToken(EHlslToken::CBuffer))
+		if (Parser.Scanner.MatchToken(EHlslToken::CBuffer) || MatchConstantBufferNoTemplate(Parser))
 		{
 			auto Result = ParseCBuffer(Parser, Allocator, OutNode);
 			if (Result == EParseResult::Error || Result == EParseResult::Matched)
diff --git a/Engine/Source/Developer/ShaderCompilerCommon/Private/ShaderCompilerCommon.cpp b/Engine/Source/Developer/ShaderCompilerCommon/Private/ShaderCompilerCommon.cpp
--- a/Engine/Source/Developer/ShaderCompilerCommon/Private/ShaderCompilerCommon.cpp
+++ b/Engine/Source/Developer/ShaderCompilerCommon/Private/ShaderCompilerCommon.cpp
@@ -232,15 +232,12 @@ const TCHAR* FindNextWhitespace(const TCHAR* StringPtr)
 
 const TCHAR* FindNextNonWhitespace(const TCHAR* StringPtr)
 {
-	bool bFoundWhitespace = false;
-
-	while (*StringPtr && (FChar::IsWhitespace(*StringPtr) || !bFoundWhitespace))
+	while (*StringPtr && FChar::IsWhitespace(*StringPtr))
 	{
-		bFoundWhitespace = true;
 		StringPtr++;
 	}
 
-	if (bFoundWhitespace && *StringPtr && !FChar::IsWhitespace(*StringPtr))
+	if (*StringPtr && !FChar::IsWhitespace(*StringPtr))
 	{
 		return StringPtr;
 	}
@@ -323,7 +320,7 @@ void ParseHLSLTypeName(const TCHAR* SearchString, const TCHAR*& TypeNameStartPtr
 	check(TypeNameEndPtr);
 }
 
-const TCHAR* ParseHLSLSymbolName(const TCHAR* SearchString, FString& SymboName)
+FStringView ParseHLSLSymbolName(const TCHAR* SearchString)
 {
 	const TCHAR* SymbolNameStartPtr = FindNextNonWhitespace(SearchString);
 	check(SymbolNameStartPtr);
@@ -334,101 +331,16 @@ const TCHAR* ParseHLSLSymbolName(const TCHAR* SearchString, FString& SymboName)
 		SymbolNameEndPtr++;
 	}
 
-	SymboName = FString(SymbolNameEndPtr - SymbolNameStartPtr, SymbolNameStartPtr);
-
-	return SymbolNameEndPtr;
+	return FStringView(SymbolNameStartPtr, SymbolNameEndPtr - SymbolNameStartPtr);
 }
 
-class FUniformBufferMemberInfo
-{
-public:
-	// eg View.WorldToClip
-	FString NameAsStructMember;
-	// eg View_WorldToClip
-	FString GlobalName;
-};
-
-const TCHAR* ParseStructRecursive(
-	const TCHAR* StructStartPtr,
-	FString& UniformBufferName,
-	int32 StructDepth,
-	const FString& StructNamePrefix, 
-	const FString& GlobalNamePrefix, 
-	TMap<FString, TArray<FUniformBufferMemberInfo>>& UniformBufferNameToMembers)
+const TCHAR* ParseHLSLSymbolName(const TCHAR* SearchString, FString& SymbolName)
 {
-	const TCHAR* OpeningBracePtr = FCString::Strstr(StructStartPtr, TEXT("{"));
-	check(OpeningBracePtr);
-
-	const TCHAR* ClosingBracePtr = FindMatchingClosingBrace(OpeningBracePtr + 1);
-	check(ClosingBracePtr);
-
-	FString StructName;
-	const TCHAR* StructNameEndPtr = ParseHLSLSymbolName(ClosingBracePtr + 1, StructName);
-	check(StructName.Len() > 0);
-
-	FString NestedStructNamePrefix = StructNamePrefix + StructName + TEXT(".");
-	FString NestedGlobalNamePrefix = GlobalNamePrefix + StructName + TEXT("_");
-
-	if (StructDepth == 0)
-	{
-		UniformBufferName = StructName;
-	}
-
-	const TCHAR* LastMemberSemicolon = ClosingBracePtr;
-
-	// Search backward to find the last member semicolon so we know when to stop parsing members
-	while (LastMemberSemicolon > OpeningBracePtr && *LastMemberSemicolon != ';')
-	{
-		LastMemberSemicolon--;
-	}
-
-	const TCHAR* MemberSearchPtr = OpeningBracePtr + 1;
-
-	do
-	{
-		const TCHAR* MemberTypeStartPtr = nullptr;
-		const TCHAR* MemberTypeEndPtr = nullptr;
-		ParseHLSLTypeName(MemberSearchPtr, MemberTypeStartPtr, MemberTypeEndPtr);
-		FString MemberTypeName(MemberTypeEndPtr - MemberTypeStartPtr, MemberTypeStartPtr);
-
-		if (FCString::Strcmp(*MemberTypeName, TEXT("struct")) == 0)
-		{
-			MemberSearchPtr = ParseStructRecursive(MemberTypeStartPtr, UniformBufferName, StructDepth + 1, NestedStructNamePrefix, NestedGlobalNamePrefix, UniformBufferNameToMembers);
-		}
-		else
-		{
-			FString MemberName;
-			const TCHAR* SymbolEndPtr = ParseHLSLSymbolName(MemberTypeEndPtr, MemberName);
-			check(MemberName.Len() > 0);
-			
-			MemberSearchPtr = SymbolEndPtr;
+	FStringView Result = ParseHLSLSymbolName(SearchString);
 
-			// Skip over trailing tokens '[1];'
-			while (*MemberSearchPtr && *MemberSearchPtr != ';')
-			{
-				MemberSearchPtr++;
-			}
-
-			// Add this member to the map
-			TArray<FUniformBufferMemberInfo>& UniformBufferMembers = UniformBufferNameToMembers.FindOrAdd(UniformBufferName);
-
-			FUniformBufferMemberInfo NewMemberInfo;
-			NewMemberInfo.NameAsStructMember = NestedStructNamePrefix + MemberName;
-			NewMemberInfo.GlobalName = NestedGlobalNamePrefix + MemberName;
-			UniformBufferMembers.Add(MoveTemp(NewMemberInfo));
-		}
-	} 
-	while (MemberSearchPtr < LastMemberSemicolon);
+	SymbolName = FString(Result);
 
-	const TCHAR* StructEndPtr = StructNameEndPtr;
-
-	// Skip over trailing tokens '[1];'
-	while (*StructEndPtr && *StructEndPtr != ';')
-	{
-		StructEndPtr++;
-	}
-
-	return StructEndPtr;
+	return Result.GetData() + Result.Len();
 }
 
 bool MatchStructMemberName(const FString& SymbolName, const TCHAR* SearchPtr, const FString& PreprocessedShaderSource)
@@ -875,35 +787,113 @@ void AddUnboundShaderParameterError(
 	AddNoteToDisplayShaderParameterStructureOnCppSide(CompilerInput.RootParametersStructure, CompilerOutput);
 }
 
+void RemoveUniformBuffersFromSource(const FShaderCompilerEnvironment& Environment, FString& PreprocessedShaderSource)
+{
+	CleanupUniformBufferCode(Environment, PreprocessedShaderSource);
+}
+
+struct FUniformBufferMemberInfo
+{
+	// eg View.WorldToClip
+	FString NameAsStructMember;
+	// eg View_WorldToClip
+	FString GlobalName;
+};
+
+const TCHAR* ParseUniformBufferDefinition(const TCHAR* ReadStart, TMap<FString, TArray<FUniformBufferMemberInfo>>& UniformBufferNameToMembers)
+{
+	const FString UniformBufferName(ParseHLSLSymbolName(ReadStart));
+
+	const TCHAR* OpeningBrace = FCString::Strstr(ReadStart, TEXT("{"));
+	const TCHAR* ClosingBrace = FindMatchingClosingBrace(OpeningBrace + 1);
+
+	const TCHAR* CurrentParseStart = OpeningBrace + 1;
+	const TCHAR* NextSemicolon = FCString::Strstr(CurrentParseStart, TEXT(";"));
+	
+	while (NextSemicolon < ClosingBrace)
+	{
+		const TCHAR* NextEquals = FCString::Strstr(CurrentParseStart, TEXT("="));
+		if (NextEquals < NextSemicolon)
+		{
+			const TCHAR* StructStart = CurrentParseStart;
+			const TCHAR* StructEnd = NextEquals - 1;
+
+			const TCHAR* GlobalStart = NextEquals + 1;
+			const TCHAR* GlobalEnd = NextSemicolon - 1;
+
+			FString StructName(StructEnd - StructStart, StructStart);
+			StructName.TrimStartAndEndInline();
+			StructName.ReplaceInline(TEXT(" "), TEXT(""));
+
+			FString GlobalName(GlobalEnd - GlobalStart, GlobalStart);
+			GlobalName.TrimStartAndEndInline();
+			GlobalName.ReplaceInline(TEXT(" "), TEXT(""));
+
+			// Avoid unnecessary conversions
+			if (StructName != GlobalName)
+			{
+				FUniformBufferMemberInfo NewMemberInfo;
+				NewMemberInfo.NameAsStructMember = StructName;
+				NewMemberInfo.GlobalName = GlobalName;
+
+				// Add this member to the map
+				TArray<FUniformBufferMemberInfo>& UniformBufferMembers = UniformBufferNameToMembers.FindOrAdd(UniformBufferName);
+				UniformBufferMembers.Add(MoveTemp(NewMemberInfo));
+			}
+		}
+
+		CurrentParseStart = NextSemicolon + 1;
+		NextSemicolon = FCString::Strstr(CurrentParseStart, TEXT(";"));
+	}
+
+	const TCHAR* EndPtr = ClosingBrace;
+
+	// Skip to the end of the UniformBuffer
+	while (*EndPtr && *EndPtr != ';')
+	{
+		EndPtr++;
+	}
+
+	return EndPtr;
+}
+
 // The cross compiler doesn't yet support struct initializers needed to construct static structs for uniform buffers
 // Replace all uniform buffer struct member references (View.WorldToClip) with a flattened name that removes the struct dependency (View_WorldToClip)
-void RemoveUniformBuffersFromSource(const FShaderCompilerEnvironment& Environment, FString& PreprocessedShaderSource)
+void CleanupUniformBufferCode(const FShaderCompilerEnvironment& Environment, FString& PreprocessedShaderSource)
 {
-	TRACE_CPUPROFILER_EVENT_SCOPE(RemoveUniformBuffersFromSource);
+	TRACE_CPUPROFILER_EVENT_SCOPE(CleanupUniformBufferCode);
 
 	TMap<FString, TArray<FUniformBufferMemberInfo>> UniformBufferNameToMembers;
 	UniformBufferNameToMembers.Reserve(Environment.UniformBufferMap.Num());
 
 	// Build a mapping from uniform buffer name to its members
 	{
-		const TCHAR* UniformBufferStructIdentifier = TEXT("static const struct");
-		const int32 StructPrefixLen = FCString::Strlen(TEXT("static const "));
+		const TCHAR* UniformBufferStructIdentifier = TEXT("UniformBuffer");
 		const int32 StructIdentifierLen = FCString::Strlen(UniformBufferStructIdentifier);
+
+		const TCHAR* SourceStart = &PreprocessedShaderSource[0];
 		TCHAR* SearchPtr = FCString::Strstr(&PreprocessedShaderSource[0], UniformBufferStructIdentifier);
 
 		while (SearchPtr)
 		{
-			FString UniformBufferName;
-			const TCHAR* ConstStructEndPtr = ParseStructRecursive(SearchPtr + StructPrefixLen, UniformBufferName, 0, TEXT(""), TEXT(""), UniformBufferNameToMembers);
-			TCHAR* StructEndPtr = &PreprocessedShaderSource[ConstStructEndPtr - &PreprocessedShaderSource[0]];
+			const ptrdiff_t Offset = SearchPtr - SourceStart;
+			if (Offset > 0 && FChar::IsWhitespace(*(SearchPtr - 1)) && FChar::IsWhitespace(*(SearchPtr + StructIdentifierLen)))
+			{
+				const TCHAR* ConstStructEndPtr = ParseUniformBufferDefinition(SearchPtr + StructIdentifierLen, UniformBufferNameToMembers);
+				TCHAR* StructEndPtr = &PreprocessedShaderSource[ConstStructEndPtr - &PreprocessedShaderSource[0]];
 
-			// Comment out the uniform buffer struct and initializer
-			*SearchPtr = '/';
-			*(SearchPtr + 1) = '*';
-			*(StructEndPtr - 1) = '*';
-			*StructEndPtr = '/';
+				// Comment out the uniform buffer struct and initializer
+				*SearchPtr = '/';
+				*(SearchPtr + 1) = '*';
+				*(StructEndPtr - 1) = '*';
+				*StructEndPtr = '/';
 
-			SearchPtr = FCString::Strstr(StructEndPtr, UniformBufferStructIdentifier);
+				SearchPtr = FCString::Strstr(StructEndPtr, UniformBufferStructIdentifier);
+			}
+			else
+			{
+				SearchPtr = FCString::Strstr(SearchPtr + 1, UniformBufferStructIdentifier);
+			}
 		}
 	}
 
diff --git a/Engine/Source/Developer/ShaderCompilerCommon/Public/ShaderCompilerCommon.h b/Engine/Source/Developer/ShaderCompilerCommon/Public/ShaderCompilerCommon.h
--- a/Engine/Source/Developer/ShaderCompilerCommon/Public/ShaderCompilerCommon.h
+++ b/Engine/Source/Developer/ShaderCompilerCommon/Public/ShaderCompilerCommon.h
@@ -322,8 +322,12 @@ extern SHADERCOMPILERCOMMON_API void AddUnboundShaderParameterError(
 	const FString& ParameterBindingName,
 	FShaderCompilerOutput& CompilerOutput);
 
+// Convert generated UniformBuffer code and references into something the shader compilers can use.
+extern SHADERCOMPILERCOMMON_API void CleanupUniformBufferCode(const FShaderCompilerEnvironment& Environment, FString& PreprocessedShaderSource);
+
 // The cross compiler doesn't yet support struct initializers needed to construct static structs for uniform buffers
 // Replace all uniform buffer struct member references (View.WorldToClip) with a flattened name that removes the struct dependency (View_WorldToClip)
+UE_DEPRECATED(5.4, "RemoveUniformBuffersFromSource was renamed CleanupUniformBufferCode")
 extern SHADERCOMPILERCOMMON_API void RemoveUniformBuffersFromSource(const FShaderCompilerEnvironment& Environment, FString& PreprocessedShaderSource);
 extern SHADERCOMPILERCOMMON_API const TCHAR* FindMatchingClosingBrace(const TCHAR* OpeningCharPtr);
 extern SHADERCOMPILERCOMMON_API const TCHAR* ParseHLSLSymbolName(const TCHAR* SearchString, FString& SymboName);
diff --git a/Engine/Source/Developer/ShaderFormatOpenGL/Private/OpenGLShaderCompiler.cpp b/Engine/Source/Developer/ShaderFormatOpenGL/Private/OpenGLShaderCompiler.cpp
--- a/Engine/Source/Developer/ShaderFormatOpenGL/Private/OpenGLShaderCompiler.cpp
+++ b/Engine/Source/Developer/ShaderFormatOpenGL/Private/OpenGLShaderCompiler.cpp
@@ -3298,7 +3298,7 @@ bool PreprocessOpenGLShader(const FShaderCompilerInput& Input, const FShaderComp
 	}
 
 	// This requires removing the HLSLCC_NoPreprocess flag later on!
-	RemoveUniformBuffersFromSource(Environment, PreprocessOutput.EditSource());
+	CleanupUniformBufferCode(Input.Environment, PreprocessOutput.EditSource());
 
 	// Process TEXT macro.
 	TransformStringIntoCharacterArray(PreprocessOutput.EditSource());
diff --git a/Engine/Source/Developer/VulkanShaderFormat/Private/VulkanShaderCompiler.cpp b/Engine/Source/Developer/VulkanShaderFormat/Private/VulkanShaderCompiler.cpp
--- a/Engine/Source/Developer/VulkanShaderFormat/Private/VulkanShaderCompiler.cpp
+++ b/Engine/Source/Developer/VulkanShaderFormat/Private/VulkanShaderCompiler.cpp
@@ -2242,7 +2242,7 @@ bool PreprocessVulkanShader(const FShaderCompilerInput& Input, const FShaderComp
 		return false;
 	}
 
-	RemoveUniformBuffersFromSource(Input.Environment, PreprocessedShaderSource);
+	CleanupUniformBufferCode(Input.Environment, PreprocessedShaderSource);
 
 	// Process TEXT macro.
 	TransformStringIntoCharacterArray(PreprocessedShaderSource);
diff --git a/Engine/Source/Developer/Windows/ShaderFormatD3D/Private/D3DShaderCompiler.cpp b/Engine/Source/Developer/Windows/ShaderFormatD3D/Private/D3DShaderCompiler.cpp
--- a/Engine/Source/Developer/Windows/ShaderFormatD3D/Private/D3DShaderCompiler.cpp
+++ b/Engine/Source/Developer/Windows/ShaderFormatD3D/Private/D3DShaderCompiler.cpp
@@ -1194,11 +1194,7 @@ bool PreprocessD3DShader(
 		return false;
 	}
 
-	// Only use UniformBuffer structs on SM6 until we can fully vet SM5
-	if (!IsUsingSM66(Input, Language))
-	{
-		RemoveUniformBuffersFromSource(Environment, PreprocessedSource);
-	}
+	CleanupUniformBufferCode(Input.Environment, PreprocessedSource);
 
 	// Process TEXT macro.
 	TransformStringIntoCharacterArray(PreprocessedSource, &Output.EditDiagnosticDatas());
diff --git a/Engine/Source/Developer/Windows/ShaderFormatD3D/Private/ShaderFormatD3D.cpp b/Engine/Source/Developer/Windows/ShaderFormatD3D/Private/ShaderFormatD3D.cpp
--- a/Engine/Source/Developer/Windows/ShaderFormatD3D/Private/ShaderFormatD3D.cpp
+++ b/Engine/Source/Developer/Windows/ShaderFormatD3D/Private/ShaderFormatD3D.cpp
@@ -133,7 +133,7 @@ public:
 		{
 			Input.Environment.SetDefine(TEXT("SM6_PROFILE"), 1);
 			Input.Environment.SetDefine(TEXT("COMPILER_DXC"), 1);
-			Input.Environment.SetDefine(TEXT("PLATFORM_SUPPORTS_UB_STRUCT"), 1);
+			Input.Environment.SetDefine(TEXT("PLATFORM_SUPPORTS_CONSTANTBUFFER_OBJECT"), 1);
 
 			AddShaderTargetDefines(Input, 6, 6);
 		}
@@ -144,6 +144,7 @@ public:
 				Input.Environment.CompilerFlags.Contains(CFLAG_WaveOperations)
 				|| Input.Environment.CompilerFlags.Contains(CFLAG_ForceDXC);
 			Input.Environment.SetDefine(TEXT("COMPILER_DXC"), bUseDXC);
+			Input.Environment.SetDefine(TEXT("PLATFORM_SUPPORTS_CONSTANTBUFFER_OBJECT"), bUseDXC);
 
 			if (bUseDXC)
 			{
diff --git a/Engine/Source/Runtime/RenderCore/Private/ShaderParameterParser.cpp b/Engine/Source/Runtime/RenderCore/Private/ShaderParameterParser.cpp
--- a/Engine/Source/Runtime/RenderCore/Private/ShaderParameterParser.cpp
+++ b/Engine/Source/Runtime/RenderCore/Private/ShaderParameterParser.cpp
@@ -176,7 +176,6 @@ static const TCHAR* const s_AllSRVTypes[] =
 	TEXT("Buffer"),
 	TEXT("ByteAddressBuffer"),
 	TEXT("StructuredBuffer"),
-	TEXT("ConstantBuffer"),
 	TEXT("RaytracingAccelerationStructure"),
 };
 
@@ -869,7 +868,6 @@ void FShaderParameterParser::RemoveMovingParametersFromSource(FString& Preproces
 
 FString FShaderParameterParser::GenerateBindlessParameterDeclaration(const FParsedShaderParameter& ParsedParameter) const
 {
-	// NOTE: Macros AUTO_BINDLESS_SAMPLER_INDEX/VARIABLE and AUTO_BINDLESS_RESOURCE_INDEX/VARIABLE in BindlessResources.ush must be kept in sync with this function
 	const bool bIsSampler = (ParsedParameter.BindlessConversionType == EBindlessConversionType::Sampler);
 	const FStringView Name = ParsedParameter.ParsedName;
 	const FStringView Type = ParsedParameter.ParsedType;
diff --git a/Engine/Source/Runtime/RenderCore/Private/ShaderParameters.cpp b/Engine/Source/Runtime/RenderCore/Private/ShaderParameters.cpp
--- a/Engine/Source/Runtime/RenderCore/Private/ShaderParameters.cpp
+++ b/Engine/Source/Runtime/RenderCore/Private/ShaderParameters.cpp
@@ -8,6 +8,7 @@
 #include "Containers/List.h"
 #include "ShaderCore.h"
 #include "Shader.h"
+#include "ShaderParameterParser.h"
 #include "VertexFactory.h"
 #include "ShaderCodeLibrary.h"
 #include "Misc/ScopeRWLock.h"
@@ -137,68 +138,91 @@ void FShaderUniformBufferParameter::Bind(const FShaderParameterMap& ParameterMap
 struct FUniformBufferDecl
 {
 	/** Members to place in the constant buffer. */
-	FString ConstantBufferMembers;
+	FStringBuilderBase ConstantBufferMembers;
 	/** Members to place in the resource table. */
-	FString ResourceMembers;
+	FStringBuilderBase ResourceMembers;
 	/** Members in the struct HLSL shader code will access. */
-	FString StructMembers;
-	/** The HLSL initializer that will copy constants and resources in to the struct. */
-	FString Initializer;
+	FStringBuilderBase StructMembers;
 };
 
 /** Generates a HLSL struct declaration for a uniform buffer struct. */
 static void CreateHLSLUniformBufferStructMembersDeclaration(
-	const FShaderParametersMetadata& UniformBufferStruct, 
-	const FString& NamePrefix, 
-	uint32 StructOffset, 
-	FUniformBufferDecl& Decl, 
+	const FShaderParametersMetadata& UniformBufferStruct,
+	const FString& UniformBufferName,
+	const FString& StructPrefix,
+	const FString& GlobalPrefix,
+	uint32 StructOffset,
+	FUniformBufferDecl& Decl,
 	uint32& HLSLBaseOffset)
 {
 	const TArray<FShaderParametersMetadata::FMember>& StructMembers = UniformBufferStruct.GetMembers();
-	
-	const bool bUniformView = (UniformBufferStruct.GetUsageFlags() & (uint32)FShaderParametersMetadata::EUsageFlags::UniformView) != 0;
-	if (bUniformView)
+
+	auto AddStructMember = [&](const FShaderParametersMetadata::FMember& Member, bool bResource)
+	{
+		// UB_DECL_PARAMETER(UniformBufferName, StructPrefix_Name, GlobalPrefix_Name)
+
+		Decl.StructMembers << (bResource ? TEXT("UB_DECL_RESOURCE(") : TEXT("UB_DECL_PARAMETER(")) << UniformBufferName << TEXT(",");
+		if (!StructPrefix.IsEmpty())
+		{
+			Decl.StructMembers << StructPrefix << TEXT(".");
+		}
+		Decl.StructMembers << Member.GetName() << TEXT(",");
+		if (!GlobalPrefix.IsEmpty())
+		{
+			Decl.StructMembers << GlobalPrefix << TEXT("_");
+		}
+		Decl.StructMembers << Member.GetName() << TEXT(");\r\n");
+	};
+
+	if ((UniformBufferStruct.GetUsageFlags() & (uint32)FShaderParametersMetadata::EUsageFlags::UniformView) != 0)
 	{
 		// UniformView struct is expected to have a single SRV member which serves as a uniform view
 		check(StructMembers.Num() == 1);
 		const FShaderParametersMetadata::FMember& Member = StructMembers[0];
 		check(Member.GetBaseType() == UBMT_SRV || Member.GetBaseType() == UBMT_RDG_BUFFER_SRV);
-		
-		FString ParameterName = FString::Printf(TEXT("%s%s"),*NamePrefix,Member.GetName());
-		Decl.ConstantBufferMembers += FString::Printf(TEXT("UB_MEMBER_UNIFORM_BLOCK(%s);\r\n"), *ParameterName);
-		Decl.StructMembers += FString::Printf(TEXT("UB_MEMBER_UNIFORM_BLOCK(%s);\r\n"), Member.GetName());
-		Decl.Initializer += FString::Printf(TEXT("%s,"),*ParameterName);
+
+		TStringBuilder<256> ParameterName;
+		if (!GlobalPrefix.IsEmpty())
+		{
+			ParameterName << GlobalPrefix << TEXT("_");
+		}
+		ParameterName << Member.GetName();
+
+		Decl.ConstantBufferMembers.Appendf(TEXT("UB_CB_UNIFORM_BLOCK(%s, %s);\r\n"), *UniformBufferName, *ParameterName);
+
+		AddStructMember(Member, false);
 		return;
 	}
-		
+
 	const TCHAR* PreviousBaseTypeName = TEXT("UB_FLOAT");
 	for (int32 MemberIndex = 0; MemberIndex < StructMembers.Num(); ++MemberIndex)
 	{
 		const FShaderParametersMetadata::FMember& Member = StructMembers[MemberIndex];
-		
-		FString ArrayDim;
-		if(Member.GetNumElements() > 0)
+
+		TStringBuilder<8> ArrayDim;
+		if (Member.GetNumElements() > 0)
 		{
-			ArrayDim = FString::Printf(TEXT("[%u]"),Member.GetNumElements());
+			ArrayDim.Appendf(TEXT("[%u]"), Member.GetNumElements());
 		}
 
-		if(Member.GetBaseType() == UBMT_NESTED_STRUCT)
+		if (Member.GetBaseType() == UBMT_NESTED_STRUCT)
 		{
 			checkf(Member.GetNumElements() == 0, TEXT("SHADER_PARAMETER_STRUCT_ARRAY() is not supported in uniform buffer yet."));
-			Decl.StructMembers += TEXT("struct {\r\n");
-			Decl.Initializer += TEXT("{");
-			CreateHLSLUniformBufferStructMembersDeclaration(*Member.GetStructMetadata(), FString::Printf(TEXT("%s%s_"), *NamePrefix, Member.GetName()), StructOffset + Member.GetOffset(), Decl, HLSLBaseOffset);
-			Decl.Initializer += TEXT("},");
-			Decl.StructMembers += FString::Printf(TEXT("} %s%s;\r\n"),Member.GetName(),*ArrayDim);
+
+			FString NewStructPrefix = StructPrefix.IsEmpty() ? FString(Member.GetName()) : FString::Printf(TEXT("%s.%s"), *StructPrefix, Member.GetName());
+			FString NewGlobalPrefix = GlobalPrefix.IsEmpty() ? FString(Member.GetName()) : FString::Printf(TEXT("%s_%s"), *GlobalPrefix, Member.GetName());
+
+			CreateHLSLUniformBufferStructMembersDeclaration(*Member.GetStructMetadata(), UniformBufferName, NewStructPrefix, NewGlobalPrefix, StructOffset + Member.GetOffset(), Decl, HLSLBaseOffset);
 		}
 		else if (Member.GetBaseType() == UBMT_INCLUDED_STRUCT)
 		{
-			CreateHLSLUniformBufferStructMembersDeclaration(*Member.GetStructMetadata(), NamePrefix, StructOffset + Member.GetOffset(), Decl, HLSLBaseOffset);
+			CreateHLSLUniformBufferStructMembersDeclaration(*Member.GetStructMetadata(), UniformBufferName, StructPrefix, GlobalPrefix, StructOffset + Member.GetOffset(), Decl, HLSLBaseOffset);
 		}
 		else if (IsShaderParameterTypeForUniformBufferLayout(Member.GetBaseType()))
 		{
-			// Generate the type dimensions for vectors and matrices.
-			uint32 HLSLMemberSize = 4;
+			// Add the constant buffer entry for bindless resource indices
+
+			constexpr uint32 HLSLMemberSize = 4;
 
 			const uint32 AbsoluteMemberOffset = StructOffset + Member.GetOffset();
 
@@ -208,7 +232,7 @@ static void CreateHLSLUniformBufferStructMembersDeclaration(
 				check(HLSLBaseOffset < AbsoluteMemberOffset);
 				while (HLSLBaseOffset < AbsoluteMemberOffset)
 				{
-					Decl.ConstantBufferMembers += FString::Printf(TEXT("\t%s() PrePadding_%s%u;\r\n"), PreviousBaseTypeName, *NamePrefix, HLSLBaseOffset);
+					Decl.ConstantBufferMembers.Appendf(TEXT("\t%s() UB_CB_MEMBER_NAME(%s, Padding%u);\r\n"), PreviousBaseTypeName, *UniformBufferName, HLSLBaseOffset);
 					HLSLBaseOffset += 4;
 				};
 				check(HLSLBaseOffset == AbsoluteMemberOffset);
@@ -217,73 +241,73 @@ static void CreateHLSLUniformBufferStructMembersDeclaration(
 			HLSLBaseOffset = AbsoluteMemberOffset + HLSLMemberSize;
 
 			// Generate the member declaration.
-			const FString ParameterName = FString::Printf(TEXT("%s%s"), *NamePrefix, Member.GetName());
 
-			if (Member.GetBaseType() == UBMT_SAMPLER)
-			{
-				Decl.ConstantBufferMembers += FString::Printf(TEXT("\tDEFINE_SAMPLER_INDEX(%s);\r\n"), *ParameterName);
-			}
-			else
+			TStringBuilder<256> ParameterName;
+			ParameterName << (Member.GetBaseType() == UBMT_SAMPLER ? FShaderParameterParser::kBindlessSamplerPrefix : FShaderParameterParser::kBindlessResourcePrefix);
+			if (!GlobalPrefix.IsEmpty())
 			{
-				Decl.ConstantBufferMembers += FString::Printf(TEXT("\tDEFINE_RESOURCE_INDEX(%s);\r\n"), *ParameterName);
+				ParameterName << GlobalPrefix << TEXT("_");
 			}
+			ParameterName << Member.GetName();
+
+			Decl.ConstantBufferMembers.Appendf(TEXT("\tUB_UINT() UB_CB_MEMBER_NAME(%s, %s);\r\n"), *UniformBufferName, *ParameterName);
 		}
-		else 
+		else
 		{
 			// Generate the base type name.
 			const TCHAR* BaseTypeName = TEXT("");
-			switch(Member.GetBaseType())
+			switch (Member.GetBaseType())
 			{
 			case UBMT_INT32:   BaseTypeName = TEXT("UB_INT"); break;
 			case UBMT_UINT32:  BaseTypeName = TEXT("UB_UINT"); break;
-			case UBMT_FLOAT32: 
+			case UBMT_FLOAT32:
 				if (Member.GetPrecision() == EShaderPrecisionModifier::Float)
 				{
-					BaseTypeName = TEXT("UB_FLOAT"); 
+					BaseTypeName = TEXT("UB_FLOAT");
 				}
 				else if (Member.GetPrecision() == EShaderPrecisionModifier::Half)
 				{
-					BaseTypeName = TEXT("UB_HALF_FLOAT"); 
+					BaseTypeName = TEXT("UB_HALF_FLOAT");
 				}
 				else if (Member.GetPrecision() == EShaderPrecisionModifier::Fixed)
 				{
-					BaseTypeName = TEXT("UB_FIXED_FLOAT"); 
+					BaseTypeName = TEXT("UB_FIXED_FLOAT");
 				}
 				break;
-			default:           UE_LOG(LogShaders, Fatal,TEXT("Unrecognized uniform buffer struct member base type."));
+			default:           UE_LOG(LogShaders, Fatal, TEXT("Unrecognized uniform buffer struct member base type."));
 			};
 
 			// Generate the type dimensions for vectors and matrices.
-			FString TypeDim;
+			TStringBuilder<16> TypeDim;
 			uint32 HLSLMemberSize = 4;
-			if(Member.GetNumRows() > 1)
+			if (Member.GetNumRows() > 1)
 			{
-				TypeDim = FString::Printf(TEXT("%ux%u"),Member.GetNumRows(),Member.GetNumColumns());
+				TypeDim.Appendf(TEXT("%ux%u"), Member.GetNumRows(), Member.GetNumColumns());
 
 				// Each row of a matrix is 16 byte aligned.
 				HLSLMemberSize = (Member.GetNumRows() - 1) * 16 + Member.GetNumColumns() * 4;
 			}
-			else if(Member.GetNumColumns() > 1)
+			else if (Member.GetNumColumns() > 1)
 			{
-				TypeDim = FString::Printf(TEXT("%u"),Member.GetNumColumns());
+				TypeDim.Appendf(TEXT("%u"), Member.GetNumColumns());
 				HLSLMemberSize = Member.GetNumColumns() * 4;
 			}
 
 			// Array elements are 16 byte aligned.
-			if(Member.GetNumElements() > 0)
+			if (Member.GetNumElements() > 0)
 			{
-				HLSLMemberSize = (Member.GetNumElements() - 1) * Align(HLSLMemberSize,16) + HLSLMemberSize;
+				HLSLMemberSize = (Member.GetNumElements() - 1) * Align(HLSLMemberSize, 16) + HLSLMemberSize;
 			}
 
 			const uint32 AbsoluteMemberOffset = StructOffset + Member.GetOffset();
 
 			// If the HLSL offset doesn't match the C++ offset, generate padding to fix it.
-			if(HLSLBaseOffset != AbsoluteMemberOffset)
+			if (HLSLBaseOffset != AbsoluteMemberOffset)
 			{
 				check(HLSLBaseOffset < AbsoluteMemberOffset);
-				while(HLSLBaseOffset < AbsoluteMemberOffset)
+				while (HLSLBaseOffset < AbsoluteMemberOffset)
 				{
-					Decl.ConstantBufferMembers += FString::Printf(TEXT("\t%s() PrePadding_%s%u;\r\n"), PreviousBaseTypeName, *NamePrefix, HLSLBaseOffset);
+					Decl.ConstantBufferMembers.Appendf(TEXT("\t%s() UB_CB_MEMBER_NAME(%s, Padding%u);\r\n"), PreviousBaseTypeName, *UniformBufferName, HLSLBaseOffset);
 					HLSLBaseOffset += 4;
 				};
 				check(HLSLBaseOffset == AbsoluteMemberOffset);
@@ -291,11 +315,16 @@ static void CreateHLSLUniformBufferStructMembersDeclaration(
 			PreviousBaseTypeName = BaseTypeName;
 			HLSLBaseOffset = AbsoluteMemberOffset + HLSLMemberSize;
 
-			// Generate the member declaration.
-			FString ParameterName = FString::Printf(TEXT("%s%s"),*NamePrefix,Member.GetName());
-			Decl.ConstantBufferMembers += FString::Printf(TEXT("\t%s(%s) %s%s;\r\n"),BaseTypeName,*TypeDim,*ParameterName,*ArrayDim);
-			Decl.StructMembers += FString::Printf(TEXT("\t%s(%s) %s%s;\r\n"),BaseTypeName,*TypeDim,Member.GetName(),*ArrayDim);
-			Decl.Initializer += FString::Printf(TEXT("%s,"),*ParameterName);
+			TStringBuilder<256> ParameterName;
+			if (!GlobalPrefix.IsEmpty())
+			{
+				ParameterName << GlobalPrefix << TEXT("_");
+			}
+			ParameterName << Member.GetName();
+
+			Decl.ConstantBufferMembers.Appendf(TEXT("\t%s(%s) UB_CB_MEMBER_NAME(%s, %s%s);\r\n"), BaseTypeName, *TypeDim, *UniformBufferName, *ParameterName, *ArrayDim);
+
+			AddStructMember(Member, false);
 		}
 	}
 
@@ -308,71 +337,67 @@ static void CreateHLSLUniformBufferStructMembersDeclaration(
 			// TODO: handle arrays?
 			checkf(!IsRDGResourceAccessType(Member.GetBaseType()), TEXT("RDG access parameter types (e.g. RDG_TEXTURE_ACCESS) are not allowed in uniform buffers."));
 
-			const FString ParameterName = FString::Printf(TEXT("%s%s"), *NamePrefix, Member.GetName());
+			TStringBuilder<256> ParameterName;
+			if (!GlobalPrefix.IsEmpty())
+			{
+				ParameterName << GlobalPrefix << TEXT("_");
+			}
+			ParameterName << Member.GetName();
 
 			if (Member.GetBaseType() == UBMT_SAMPLER)
 			{
-				Decl.ResourceMembers += FString::Printf(TEXT("UB_RESOURCE_MEMBER_SAMPLER(%s, %s);\r\n"), Member.GetShaderType(), *ParameterName);
-				Decl.StructMembers += FString::Printf(TEXT("\t%s %s;\r\n"), Member.GetShaderType(), Member.GetName());
-				Decl.Initializer += FString::Printf(TEXT("%s,"), *ParameterName);
+				Decl.ResourceMembers.Appendf(TEXT("UB_RESOURCE_MEMBER_SAMPLER(%s, %s, %s);\r\n"), Member.GetShaderType(), *UniformBufferName, *ParameterName);
+				AddStructMember(Member, true);
 			}
 			else if (Member.GetBaseType() == UBMT_SRV)
 			{
-				Decl.ResourceMembers += FString::Printf(TEXT("UB_RESOURCE_MEMBER_RESOURCE(%s, %s);\r\n"), Member.GetShaderType(), *ParameterName);
-				Decl.StructMembers += FString::Printf(TEXT("\t%s %s;\r\n"), Member.GetShaderType(), Member.GetName());
-				Decl.Initializer += FString::Printf(TEXT("%s,"), *ParameterName);
+				Decl.ResourceMembers.Appendf(TEXT("UB_RESOURCE_MEMBER_RESOURCE(%s, %s, %s);\r\n"), Member.GetShaderType(), *UniformBufferName, *ParameterName);
+				AddStructMember(Member, true);
 			}
 			else
 			{
-				Decl.ResourceMembers += FString::Printf(TEXT("UB_RESOURCE_MEMBER_RESOURCE(%s, %s);\r\n"), Member.GetShaderType(), *ParameterName);
-				Decl.StructMembers += FString::Printf(TEXT("\t%s %s;\r\n"), Member.GetShaderType(), Member.GetName());
-				Decl.Initializer += FString::Printf(TEXT("%s,"), *ParameterName);
+				Decl.ResourceMembers.Appendf(TEXT("UB_RESOURCE_MEMBER_RESOURCE(%s, %s, %s);\r\n"), Member.GetShaderType(), *UniformBufferName, *ParameterName);
+				AddStructMember(Member, true);
 			}
 		}
 	}
 }
 
 /** Creates a HLSL declaration of a uniform buffer with the given structure. */
-static FString CreateHLSLUniformBufferDeclaration(const TCHAR* Name, const FShaderParametersMetadata& UniformBufferStruct)
+static FString CreateHLSLUniformBufferDeclaration(const TCHAR* UniformBufferName, const FShaderParametersMetadata& UniformBufferStruct)
 {
 	// If the uniform buffer has no members, we don't want to write out anything.  Shader compilers throw errors when faced with empty cbuffers and structs.
 	if (UniformBufferStruct.GetMembers().Num() > 0)
 	{
-		FString NamePrefix(FString(Name) + FString(TEXT("_")));
 		FUniformBufferDecl Decl;
 		uint32 HLSLBaseOffset = 0;
-		CreateHLSLUniformBufferStructMembersDeclaration(UniformBufferStruct, NamePrefix, 0, Decl, HLSLBaseOffset);
+		CreateHLSLUniformBufferStructMembersDeclaration(UniformBufferStruct, UniformBufferName, TEXT(""), TEXT(""), 0, Decl, HLSLBaseOffset);
 
 		return FString::Printf(
-			TEXT("#ifndef __UniformBuffer_%s_Definition__\r\n")
-			TEXT("#define __UniformBuffer_%s_Definition__\r\n")
-			TEXT("cbuffer UB_CB_NAME(%s)\r\n")
+			TEXT("#pragma once\r\n")
+			TEXT("UB_CB_DEFINITION_START(%s)\r\n")
+			TEXT("%s")
+			TEXT("UB_CB_DEFINITION_END(%s)\r\n")
+			TEXT("%s")
+			TEXT("UniformBuffer %s\r\n")
 			TEXT("{\r\n")
 			TEXT("%s")
-			TEXT("}\r\n")
-				TEXT("%s")
-				TEXT("static const struct\r\n")
-				TEXT("{\r\n")
-				TEXT("%s")
-				TEXT("} %s = {%s};\r\n")
-			TEXT("#endif\r\n"),
-			Name,
-			Name,
-			Name,
+			TEXT("};\r\n"),
+			UniformBufferName,
 			*Decl.ConstantBufferMembers,
+			UniformBufferName,
 			*Decl.ResourceMembers,
-			*Decl.StructMembers,
-			Name,
-			*Decl.Initializer
-			);
+			UniformBufferName,
+			*Decl.StructMembers
+		);
 	}
 
 	return FString(TEXT("\n"));
 }
 
-FString UE::ShaderParameters::CreateUniformBufferShaderDeclaration(const TCHAR* Name, const FShaderParametersMetadata& UniformBufferStruct)
+FString UE::ShaderParameters::CreateUniformBufferShaderDeclaration(const TCHAR* UniformBufferName, const FShaderParametersMetadata& UniformBufferStruct)
 {
-	return CreateHLSLUniformBufferDeclaration(Name, UniformBufferStruct);
+	return CreateHLSLUniformBufferDeclaration(UniformBufferName, UniformBufferStruct);
 }
 
 static FShaderParametersMetadata* FindShaderParametersMetadataWithVariableName(uint32 InVariableNameHash, FStringView InVariableNameView)
