Clear FRayTracingGeometry state in ReleaseRHI().

diff --git a/Engine/Source/Runtime/RenderCore/Private/RayTracingGeometry.cpp b/Engine/Source/Runtime/RenderCore/Private/RayTracingGeometry.cpp
--- a/Engine/Source/Runtime/RenderCore/Private/RayTracingGeometry.cpp
+++ b/Engine/Source/Runtime/RenderCore/Private/RayTracingGeometry.cpp
@@ -161,15 +161,11 @@ void FRayTracingGeometry::CreateRayTracingGeometry(FRHICommandList& RHICmdList,
 		}
 	}
 
-	const bool bWithoutNativeResource = Initializer.Type == ERayTracingGeometryInitializerType::StreamingDestination;
 	if (bAllSegmentsAreValid)
 	{
-		// Only geometries with StreamingDestination type are initially created in invalid state until they are streamed in (see InitRHIForStreaming).
-		if (bWithoutNativeResource)
-		{
-			EnumRemoveFlags(GeometryState, EGeometryStateFlags::Valid);
-		}
-		else
+		// Geometries with StreamingDestination type are initially created in invalid state until they are streamed in (see InitRHIForStreaming).
+		const bool bWithNativeResource = Initializer.Type != ERayTracingGeometryInitializerType::StreamingDestination;
+		if (bWithNativeResource)
 		{
 			EnumAddFlags(GeometryState, EGeometryStateFlags::Valid);
 		}
@@ -236,7 +232,7 @@ void FRayTracingGeometry::ReleaseRHI()
 {
 	RemoveBuildRequest();
 	RayTracingGeometryRHI.SafeRelease();
-	EnumRemoveFlags(GeometryState, EGeometryStateFlags::Valid);
+	GeometryState = EGeometryStateFlags::Invalid;
 }
 
 void FRayTracingGeometry::RemoveBuildRequest()
@@ -274,8 +270,6 @@ void FRayTracingGeometry::ReleaseResource()
 	// This includes index and vertex buffers used for building the BLAS.
 	Initializer = FRayTracingGeometryInitializer{};
 
-	GeometryState = EGeometryStateFlags::Invalid;
-
 	FRenderResource::ReleaseResource();
 }
 
