Refactored render command fence implementation to fence render command pipes.
 - Reimplemented render command fence bundler flushing to keep the bundler active during a GC.
 - Added insights regions to track render command pipe recording and fence bundler activity when the render commands channel is active.
 - Cleaned up render command pipe recording so that it works properly with the fence bundler.

diff --git a/Engine/Source/Runtime/Engine/Private/UnrealEngine.cpp b/Engine/Source/Runtime/Engine/Private/UnrealEngine.cpp
--- a/Engine/Source/Runtime/Engine/Private/UnrealEngine.cpp
+++ b/Engine/Source/Runtime/Engine/Private/UnrealEngine.cpp
@@ -1648,9 +1648,6 @@ void UEngine::PreGarbageCollect()
 	TRACE_CPUPROFILER_EVENT_SCOPE(UEngine::PreGarbageCollect);
 
 	SendWorldEndOfFrameUpdates();
-
-	// Sync recording render command pipes prior to calling BeginDestroy on resources so that render thread fences are serialized with other render commands.
-	UE::RenderCommandPipe::StopRecording();
 }
 
 float GetLowMemoryGCTimer(const float DefaultTimeBetweenGC)
diff --git a/Engine/Source/Runtime/Launch/Private/LaunchEngineLoop.cpp b/Engine/Source/Runtime/Launch/Private/LaunchEngineLoop.cpp
--- a/Engine/Source/Runtime/Launch/Private/LaunchEngineLoop.cpp
+++ b/Engine/Source/Runtime/Launch/Private/LaunchEngineLoop.cpp
@@ -3937,7 +3937,6 @@ int32 FEngineLoop::PreInitPostStartupScreen(const TCHAR* CmdLine)
 		}
 	}
 
-	if(!GIsEditor)
 	{
 		FCoreUObjectDelegates::PreGarbageCollectConditionalBeginDestroy.AddStatic(StartRenderCommandFenceBundler);
 		FCoreUObjectDelegates::PostGarbageCollectConditionalBeginDestroy.AddStatic(StopRenderCommandFenceBundler);
diff --git a/Engine/Source/Runtime/RenderCore/Private/RenderingThread.cpp b/Engine/Source/Runtime/RenderCore/Private/RenderingThread.cpp
--- a/Engine/Source/Runtime/RenderCore/Private/RenderingThread.cpp
+++ b/Engine/Source/Runtime/RenderCore/Private/RenderingThread.cpp
@@ -30,6 +30,7 @@
 #include "ProfilingDebugging/CsvProfiler.h"
 #include "Async/TaskTrace.h"
 #include "DataDrivenShaderPlatformInfo.h"
+#include "ProfilingDebugging/CountersTrace.h"
 
 //
 // Globals
@@ -968,7 +969,34 @@ bool IsRenderingThreadHealthy()
 	return GIsRenderingThreadHealthy;
 }
 
-static TOptional<UE::Tasks::FTaskEvent> BundledCompletionEvent;
+static struct FRenderCommandFenceBundlerState
+{
+	TOptional<UE::Tasks::FTaskEvent> Event;
+	FRenderCommandPipeBitArray RenderCommandPipeBits;
+	int32 RecursionDepth = 0;
+
+} GRenderCommandFenceBundlerState; 
+
+#define UE_RENDER_COMMAND_FENCE_BUNDLER_REGION TEXT("Render Command Fence Bundler")
+#define UE_RENDER_COMMAND_PIPE_RECORD_REGION TEXT("Render Command Pipe Recording")
+#define UE_RENDER_COMMAND_PIPE_SYNC_REGION TEXT("Render Command Pipe Synced")
+
+#if UE_TRACE_ENABLED
+#define UE_RENDER_COMMAND_BEGIN_REGION(Region) \
+	if (RenderCommandsChannel) \
+	{ \
+		TRACE_BEGIN_REGION(Region) \
+	}
+
+#define UE_RENDER_COMMAND_END_REGION(Region) \
+	if (RenderCommandsChannel) \
+	{ \
+		TRACE_END_REGION(Region) \
+	}
+#else
+#define UE_RENDER_COMMAND_BEGIN_REGION(Region)
+#define UE_RENDER_COMMAND_END_REGION(Region)
+#endif
 
 void StartRenderCommandFenceBundler()
 {
@@ -977,28 +1005,79 @@ void StartRenderCommandFenceBundler()
 		return;
 	}
 
-	check(IsInGameThread() && !BundledCompletionEvent); // can't use this in a nested fashion
-	BundledCompletionEvent.Emplace(TEXT("RenderCommandFenceBundlerEvent"));
+	check(IsInGameThread());
+	check(!GRenderCommandFenceBundlerState.Event.IsSet() == !GRenderCommandFenceBundlerState.RecursionDepth);
+
+	++GRenderCommandFenceBundlerState.RecursionDepth;
+
+	if (GRenderCommandFenceBundlerState.RecursionDepth > 1)
+	{
+		return;
+	}
+
+	GRenderCommandFenceBundlerState.Event.Emplace(TEXT("RenderCommandFenceBundlerEvent"));
+
+	// Stop render command pipes so that the bundled render command fence is serialized with other render commands.
+	GRenderCommandFenceBundlerState.RenderCommandPipeBits = UE::RenderCommandPipe::StopRecording();
 
 	StartBatchedRelease();
+
+	UE_RENDER_COMMAND_BEGIN_REGION(UE_RENDER_COMMAND_FENCE_BUNDLER_REGION);
+}
+
+void FlushRenderCommandFenceBundler()
+{
+	if (GRenderCommandFenceBundlerState.Event)
+	{
+		EndBatchedRelease();
+
+		ENQUEUE_RENDER_COMMAND(InsertFence)(
+			[CompletionEvent = MoveTemp(*GRenderCommandFenceBundlerState.Event)](FRHICommandListBase&) mutable
+		{
+			CompletionEvent.Trigger();
+		});
+
+		GRenderCommandFenceBundlerState.Event.Emplace(TEXT("RenderCommandFenceBundlerEvent"));
+
+		StartBatchedRelease();
+	}
 }
 
 void StopRenderCommandFenceBundler()
 {
-	if (!GIsThreadedRendering || !BundledCompletionEvent)
+	if (!GIsThreadedRendering || !GRenderCommandFenceBundlerState.Event)
 	{
 		return;
 	}
 
+	TOptional<UE::Tasks::FTaskEvent>& CompletionEvent = GRenderCommandFenceBundlerState.Event;
+
+	check(CompletionEvent);
+	check(!CompletionEvent->IsCompleted());
+	check(GRenderCommandFenceBundlerState.RecursionDepth > 0);
+
+	--GRenderCommandFenceBundlerState.RecursionDepth;
+
+	if (GRenderCommandFenceBundlerState.RecursionDepth > 0)
+	{
+		return;
+	}
+
+	UE_RENDER_COMMAND_END_REGION(UE_RENDER_COMMAND_FENCE_BUNDLER_REGION);
+
 	EndBatchedRelease();
-	checkf(IsInGameThread() && BundledCompletionEvent && !BundledCompletionEvent->IsCompleted(), TEXT("IsInGameThread: %d, BundledCompletionEvent is completed: %d"), IsInGameThread(), BundledCompletionEvent->IsCompleted()); // can't use this in a nested fashion
 
 	ENQUEUE_RENDER_COMMAND(InsertFence)(
-		[CompletionEvent = MoveTemp(*BundledCompletionEvent)](FRHICommandListImmediate&) mutable
+		[CompletionEvent = MoveTemp(*CompletionEvent)](FRHICommandListBase&) mutable
 	{
 		CompletionEvent.Trigger();
 	});
-	BundledCompletionEvent.Reset();
+
+	CompletionEvent.Reset();
+
+	// Restart render command pipes that were previously recording.
+	UE::RenderCommandPipe::StartRecording(GRenderCommandFenceBundlerState.RenderCommandPipeBits);
+	GRenderCommandFenceBundlerState.RenderCommandPipeBits.Empty();
 }
 
 std::atomic<int> GTimeoutSuspendCount;
@@ -1067,52 +1146,128 @@ void FRenderCommandFence::BeginFence(bool bSyncToRHIAndGPU)
 	{
 		return;
 	}
-	else
+	
+	if (GRenderCommandFenceBundlerState.Event && IsInGameThread())
+	{
+		CompletionTask = *GRenderCommandFenceBundlerState.Event;
+		return;
+	}
+
+	TRACE_CPUPROFILER_EVENT_SCOPE(FRenderCommandFence::BeginFence);
+
+	struct FRenderCommandPipeFence : public TConcurrentLinearObject<FRenderCommandPipeFence>
 	{
-		if (BundledCompletionEvent && IsInGameThread())
+		FRenderCommandPipeFence(int32 InNumRefs)
+			: NumRefs(InNumRefs)
+		{}
+
+		void Trigger(int32 NumTriggerRefs = 1)
 		{
-			CompletionTask = *BundledCompletionEvent;
-			return;
+			if (NumRefs.fetch_sub(NumTriggerRefs, std::memory_order_release) == 1)
+			{
+				std::atomic_thread_fence(std::memory_order_acquire);
+				CompletionTaskEvent.Trigger();
+				delete this;
+			}
 		}
 
-		int32 GTSyncType = CVarGTSyncType.GetValueOnAnyThread();
-		if (bSyncToRHIAndGPU)
+		UE::Tasks::FTaskEvent CompletionTaskEvent{ UE_SOURCE_LOCATION };
+		std::atomic_int32_t NumRefs;
+	};
+
+	TConstArrayView<FRenderCommandPipe*> Pipes = GRenderCommandPipeMode == ERenderCommandPipeMode::All
+		? UE::RenderCommandPipe::GetPipes()
+		: TConstArrayView<FRenderCommandPipe*>{};
+
+	FRenderCommandPipeBitArray ActivePipeBits;
+	int32 NumActivePipes = 0;
+
+	for (FRenderCommandPipe* Pipe : Pipes)
+	{
+		// Skip pipes that aren't recording or replaying any work.
+		const bool bIsActive = Pipe->IsRecording() && !Pipe->IsEmpty();
+		ActivePipeBits.Add(bIsActive);
+		NumActivePipes += bIsActive ? 1 : 0;
+	}
+
+	FRenderCommandPipeFence* Fence = nullptr;
+
+	if (NumActivePipes > 0)
+	{
+		Fence = new FRenderCommandPipeFence(NumActivePipes + 1);
+
+		for (FRenderCommandPipeSetBitIterator BitIt(ActivePipeBits); BitIt; ++BitIt)
 		{
-			// Don't sync to the RHI and GPU if GtSyncType is disabled, or we're not vsyncing
-			//@TODO: do this logic in the caller?
-			static auto CVarVsync = IConsoleManager::Get().FindConsoleVariable(TEXT("r.VSync"));
-			check(CVarVsync != nullptr);
+			FRenderCommandPipe* Pipe = Pipes[BitIt.GetIndex()];
 
-			if ( GTSyncType == 0 || CVarVsync->GetInt() == 0 )
+			ENQUEUE_RENDER_COMMAND(BeginFence)(Pipe, [Fence]
 			{
-				bSyncToRHIAndGPU = false;
-			}
+				Fence->Trigger();
+			});
 		}
+	}
+
+	const int32 GTSyncType = CVarGTSyncType.GetValueOnAnyThread();
 
+	if (bSyncToRHIAndGPU)
+	{
+		// Don't sync to the RHI and GPU if GtSyncType is disabled, or we're not vsyncing
+		//@TODO: do this logic in the caller?
+		static auto CVarVsync = IConsoleManager::Get().FindConsoleVariable(TEXT("r.VSync"));
+		check(CVarVsync != nullptr);
+
+		if (GTSyncType == 0 || CVarVsync->GetInt() == 0)
+		{
+			bSyncToRHIAndGPU = false;
+		}
+	}
+
+	if (bSyncToRHIAndGPU)
+	{
 		UE::Tasks::FTaskEvent CompletionTaskEvent{ UE_SOURCE_LOCATION };
 
+		if (Fence)
+		{
+			// RHI frame sync command requires a task event, so connect it to the ref-counted fence event.
+			CompletionTaskEvent.AddPrerequisites(Fence->CompletionTaskEvent);
+			Fence->Trigger();
+		}
+
 		ENQUEUE_RENDER_COMMAND(FSyncFrameCommand)(
 			[CompletionTaskEvent, GTSyncType, bSyncToRHIAndGPU](FRHICommandListImmediate& RHICmdList) mutable
 		{
-			if (bSyncToRHIAndGPU)
+			if (IsRHIThreadRunning())
 			{
-				if (IsRHIThreadRunning())
-				{
-					ALLOC_COMMAND_CL(RHICmdList, FRHISyncFrameCommand)(MoveTemp(CompletionTaskEvent), GTSyncType);
-					RHICmdList.ImmediateFlush(EImmediateFlushType::DispatchToRHIThread);
-				}
-				else
-				{
-					FRHISyncFrameCommand Command(MoveTemp(CompletionTaskEvent), GTSyncType);
-					Command.Execute(RHICmdList);
-				}
+				ALLOC_COMMAND_CL(RHICmdList, FRHISyncFrameCommand)(MoveTemp(CompletionTaskEvent), GTSyncType);
+				RHICmdList.ImmediateFlush(EImmediateFlushType::DispatchToRHIThread);
 			}
 			else
 			{
-				CompletionTaskEvent.Trigger();
+				FRHISyncFrameCommand Command(MoveTemp(CompletionTaskEvent), GTSyncType);
+				Command.Execute(RHICmdList);
 			}
 		});
 
+		CompletionTask = MoveTemp(CompletionTaskEvent);
+	}
+	else if (Fence)
+	{
+		CompletionTask = Fence->CompletionTaskEvent;
+
+		ENQUEUE_RENDER_COMMAND(BeginFence)([Fence](FRHICommandListBase& RHICmdList)
+		{
+			Fence->Trigger();
+		});
+	}
+	else
+	{
+		UE::Tasks::FTaskEvent CompletionTaskEvent{ UE_SOURCE_LOCATION };
+
+		ENQUEUE_RENDER_COMMAND(BeginFence)([CompletionTaskEvent](FRHICommandListBase& RHICmdList) mutable
+		{
+			CompletionTaskEvent.Trigger();
+		});
+
 		CompletionTask = MoveTemp(CompletionTaskEvent);
 	}
 }
@@ -1265,7 +1420,7 @@ void FRenderCommandFence::Wait(bool bProcessGameThreadTasks) const
 {
 	if (!IsFenceComplete())
 	{
-		StopRenderCommandFenceBundler();
+		FlushRenderCommandFenceBundler();
 		GameThreadWaitForTask(CompletionTask, bProcessGameThreadTasks);
 		CompletionTask = {}; // release the internal memory as soon as it's not needed anymore
 	}
@@ -1372,7 +1527,6 @@ FPendingCleanupObjects::~FPendingCleanupObjects()
 	if (CleanupArray.Num())
 	{
 		QUICK_SCOPE_CYCLE_COUNTER(STAT_FPendingCleanupObjects_Destruct);
-		UE::RenderCommandPipe::StopRecording();
 
 		const bool bBatchingEnabled = bEnablePendingCleanupObjectsCommandBatching;
 		if (bBatchingEnabled)
@@ -1414,7 +1568,6 @@ FPendingCleanupObjects::~FPendingCleanupObjects()
 	if (CleanupArray.Num())
 	{
 		QUICK_SCOPE_CYCLE_COUNTER(STAT_FPendingCleanupObjects_Destruct);
-		UE::RenderCommandPipe::StopRecording();
 
 		const bool bBatchingEnabled = bEnablePendingCleanupObjectsCommandBatching;
 		if (bBatchingEnabled)
@@ -1552,6 +1705,7 @@ FAutoConsoleVariable CVarRenderCommandPipeMode(
 	TEXT(" 2: Render commands are enqueued into a render command pipe for all declared pipes.;\n"),
 	FConsoleVariableDelegate::CreateLambda([](IConsoleVariable* Variable)
 	{
+		UE::RenderCommandPipe::StopRecording();
 		GRenderCommandPipeMode = GetValidatedRenderCommandPipeMode(Variable->GetInt());
 	}));
 
@@ -1590,8 +1744,6 @@ void FRenderThreadCommandPipe::EnqueueAndLaunch(const TCHAR* Name, uint32& SpecI
 	}
 }
 
-#if !UE_SERVER
-
 class FRenderCommandPipeRegistry
 {
 public:
@@ -1660,6 +1812,8 @@ public:
 			return;
 		}
 
+		UE_RENDER_COMMAND_BEGIN_REGION(UE_RENDER_COMMAND_PIPE_RECORD_REGION);
+
 		UE::Tasks::FTaskEvent TaskEvent{ UE_SOURCE_LOCATION };
 
 		struct FPipeToStartRecording
@@ -1766,6 +1920,11 @@ public:
 		return PipeBits;
 	}
 
+	TConstArrayView<FRenderCommandPipe*> GetPipes() const
+	{
+		return AllPipes;
+	}
+
 	bool IsRecording() const
 	{
 		ensureMsgf(!FTaskTagScope::IsCurrentTag(ETaskTag::EParallelRenderingThread) && !FTaskTagScope::IsCurrentTag(ETaskTag::ERenderingThread),
@@ -1827,6 +1986,8 @@ private:
 			RHICmdList.QueueAsyncCommandListSubmit(QueuedCommandLists);
 			RHIResourceLifetimeReleaseRef(RHICmdList, NumPipesToStopRecording);
 		});
+
+		UE_RENDER_COMMAND_END_REGION(UE_RENDER_COMMAND_PIPE_RECORD_REGION);
 	}
 
 	UE::FMutex Mutex;
@@ -1837,66 +1998,95 @@ private:
 
 static FRenderCommandPipeRegistry GRenderCommandPipeRegistry;
 
-#endif // !UE_SERVER
+inline bool HasBitsSet(const FRenderCommandPipeBitArray& Bits)
+{
+	for (FRenderCommandPipeBitArray::FConstWordIterator It(Bits); It; ++It)
+	{
+		if (It.GetWord() != 0)
+		{
+			return true;
+		}
+	}
+	return false;
+}
 
 namespace UE::RenderCommandPipe
 {
 	void Initialize()
 	{
-#if !UE_SERVER
 		GRenderCommandPipeRegistry.Initialize();
-#endif
 	}
 
 	bool IsRecording()
 	{
-#if !UE_SERVER
 		return GRenderCommandPipeRegistry.IsRecording();
-#else
-		return false;
-#endif
 	}
 
 	bool IsReplaying()
 	{
-#if !UE_SERVER
 		return GRenderCommandPipeRegistry.IsReplaying();
-#else
-		return false;
-#endif
 	}
 
 	void StartRecording()
 	{
-#if !UE_SERVER
 		GRenderCommandPipeRegistry.StartRecording();
-#endif
 	}
 
 	void StartRecording(const FRenderCommandPipeBitArray& PipeBits)
 	{
-#if !UE_SERVER
 		GRenderCommandPipeRegistry.StartRecording(PipeBits);
-#endif
 	}
 
 	FRenderCommandPipeBitArray StopRecording()
 	{
-#if !UE_SERVER
 		return GRenderCommandPipeRegistry.StopRecording();
-#else
-		return {};
-#endif
 	}
 
 	FRenderCommandPipeBitArray StopRecording(TConstArrayView<FRenderCommandPipe*> Pipes)
 	{
-#if !UE_SERVER
 		return GRenderCommandPipeRegistry.StopRecording(Pipes);
-#else
-		return {};
+	}
+
+	TConstArrayView<FRenderCommandPipe*> GetPipes()
+	{
+		return GRenderCommandPipeRegistry.GetPipes();
+	}
+
+	FSyncScope::FSyncScope()
+	{
+		PipeBits = StopRecording();
+
+#if UE_TRACE_ENABLED
+		if (HasBitsSet(PipeBits))
+		{
+			UE_RENDER_COMMAND_BEGIN_REGION(UE_RENDER_COMMAND_PIPE_SYNC_REGION);
+		}
+#endif
+	}
+
+	FSyncScope::FSyncScope(TConstArrayView<FRenderCommandPipe*> Pipes)
+	{
+		PipeBits = StopRecording(Pipes);
+
+#if UE_TRACE_ENABLED
+		if (HasBitsSet(PipeBits))
+		{
+			UE_RENDER_COMMAND_BEGIN_REGION(UE_RENDER_COMMAND_PIPE_SYNC_REGION);
+		}
 #endif
 	}
+
+	FSyncScope::~FSyncScope()
+	{
+#if UE_TRACE_ENABLED
+		if (HasBitsSet(PipeBits))
+		{
+			UE_RENDER_COMMAND_END_REGION(UE_RENDER_COMMAND_PIPE_SYNC_REGION);
+		}
+#endif
+
+		StartRecording(PipeBits);
+	}
 }
 
 FRenderCommandPipe::FRenderCommandPipe(const TCHAR* InName, ERenderCommandPipeFlags Flags, const TCHAR* CVarName, const TCHAR* CVarDescription)
@@ -1924,14 +2114,16 @@ void FRenderCommandPipe::EnqueueAndLaunch(FFunctionVariant&& FunctionVariant, co
 {
 	bool bWasEmpty = Frame_GameThread->Queue.IsEmpty();
 	Frame_GameThread->Queue.Emplace(MoveTemp(FunctionVariant), CommandName, CommandSpecId);
+	NumInFlightCommands.fetch_add(1, std::memory_order_relaxed);
 
 	if (bWasEmpty)
 	{
-		SCOPED_NAMED_EVENT(FRenderCommandPipe_EnqueueAndLaunch_LaunchTask, FColor::Magenta);
+		TRACE_CPUPROFILER_EVENT_SCOPE_ON_CHANNEL_STR("RenderCommandPipe LaunchTask", RenderCommandsChannel)
 
 		Frame_GameThread->Pipe.Launch(Name, [this]
 		{
 			check(Frame_RenderThread);
+			TRACE_CPUPROFILER_EVENT_SCOPE_ON_CHANNEL_STR("RenderCommandPipe ReplayCommands", RenderCommandsChannel)
 			SCOPED_NAMED_EVENT_TCHAR(Name, FColor::Magenta);
 			FOptionalTaskTagScope Scope(ETaskTag::EParallelRenderingThread);
 
@@ -1944,7 +2136,7 @@ void FRenderCommandPipe::EnqueueAndLaunch(FFunctionVariant&& FunctionVariant, co
 
 			for (FCommand& Command : PoppedQueue)
 			{
-				TRACE_CPUPROFILER_EVENT_SCOPE_USE_ON_CHANNEL(*Command.SpecId, Command.Name, EventScope, RenderCommandsChannel, true);
+				TRACE_CPUPROFILER_EVENT_SCOPE_USE_ON_CHANNEL(*Command.SpecId, Command.Name, CommandEventScope, RenderCommandsChannel, true);
 
 				if (FCommandListFunction* Function = Command.Function.TryGet<FCommandListFunction>())
 				{
@@ -1965,6 +2157,8 @@ void FRenderCommandPipe::EnqueueAndLaunch(FFunctionVariant&& FunctionVariant, co
 				Command.Function = {};
 			}
 
+			NumInFlightCommands.fetch_sub(PoppedQueue.Num(), std::memory_order_release);
+
 		}, Frame_GameThread->TaskEvent);
 	}
 }
\ No newline at end of file
diff --git a/Engine/Source/Runtime/RenderCore/Public/RenderingThread.h b/Engine/Source/Runtime/RenderCore/Public/RenderingThread.h
--- a/Engine/Source/Runtime/RenderCore/Public/RenderingThread.h
+++ b/Engine/Source/Runtime/RenderCore/Public/RenderingThread.h
@@ -295,7 +295,7 @@ enum class ERenderCommandPipeFlags : uint8
 ENUM_CLASS_FLAGS(ERenderCommandPipeFlags);
 
 class FRenderCommandPipe;
-using FRenderCommandPipeBitArrayAllocator = TInlineAllocator<1, FConcurrentLinearArrayAllocator>;
+using FRenderCommandPipeBitArrayAllocator = TInlineAllocator<1, FConcurrentLinearBitArrayAllocator>;
 using FRenderCommandPipeBitArray = TBitArray<FRenderCommandPipeBitArrayAllocator>;
 using FRenderCommandPipeSetBitIterator = TConstSetBitIterator<FRenderCommandPipeBitArrayAllocator>;
 
@@ -318,24 +318,16 @@ namespace UE::RenderCommandPipe
 	extern RENDERCORE_API FRenderCommandPipeBitArray StopRecording();
 	extern RENDERCORE_API FRenderCommandPipeBitArray StopRecording(TConstArrayView<FRenderCommandPipe*> Pipes);
 
+	// Returns the list of all registered pipes.
+	extern TConstArrayView<FRenderCommandPipe*> GetPipes();
+
 	// [Game Thread] Stops render command pipe recording during the duration of the scope and restarts recording once the scope is complete.
-	class FSyncScope
+	class RENDERCORE_API FSyncScope
 	{
 	public:
-		FSyncScope()
-		{
-			PipeBits = StopRecording();
-		}
-
-		FSyncScope(TConstArrayView<FRenderCommandPipe*> Pipes)
-		{
-			PipeBits = StopRecording(Pipes);
-		}
-
-		~FSyncScope()
-		{
-			StartRecording(PipeBits);
-		}
+		FSyncScope();
+		FSyncScope(TConstArrayView<FRenderCommandPipe*> Pipes);
+		~FSyncScope();
 
 	private:
 		FRenderCommandPipeBitArray PipeBits;
@@ -437,6 +429,11 @@ public:
 		return bRecording;
 	}
 
+	FORCEINLINE bool IsEmpty() const
+	{
+		return NumInFlightCommands.load(std::memory_order_relaxed) == 0;
+	}
+
 	void SetEnabled(bool bInIsEnabled)
 	{
 		check(IsInGameThread());
@@ -462,17 +459,7 @@ public:
 	template <typename RenderCommandTag>
 	static void Enqueue(FRenderCommandPipe& Pipe, FCommandListFunction&& Function)
 	{
-		if (GRenderCommandPipeMode == ERenderCommandPipeMode::All)
-		{
-			UE::TScopeLock Lock(Pipe.Mutex);
-			if (Pipe.Frame_GameThread)
-			{
-				Pipe.EnqueueAndLaunch(MoveTemp(Function), RenderCommandTag::GetName(), RenderCommandTag::GetSpecId());
-				return;
-			}
-		}
-		
-		EnqueueUniqueRenderCommand<RenderCommandTag>([Function = MoveTemp(Function)](FRHICommandListImmediate& RHICmdList) { Function(RHICmdList); });
+		Enqueue<RenderCommandTag>(&Pipe, MoveTemp(Function));
 	}
 
 	template <typename RenderCommandTag>
@@ -494,17 +481,7 @@ public:
 	template <typename RenderCommandTag>
 	static void Enqueue(FRenderCommandPipe& Pipe, FEmptyFunction&& Function)
 	{
-		if (Pipe.bEnabled)
-		{
-			UE::TScopeLock Lock(Pipe.Mutex);
-			if (Pipe.Frame_GameThread)
-			{
-				Pipe.EnqueueAndLaunch(MoveTemp(Function), RenderCommandTag::GetName(), RenderCommandTag::GetSpecId());
-				return;
-			}
-		}
-
-		EnqueueUniqueRenderCommand<RenderCommandTag>([Function = MoveTemp(Function)](FRHICommandListImmediate&) { Function(); });
+		Enqueue<RenderCommandTag>(&Pipe, MoveTemp(Function));
 	}
 
 	template <typename RenderCommandTag, typename LambdaType>
@@ -562,6 +539,7 @@ private:
 	FFrame* Frame_RenderThread = nullptr;
 	TLinkedList<FRenderCommandPipe*> GlobalListLink;
 	FAutoConsoleVariable ConsoleVariable;
+	std::atomic_int32_t NumInFlightCommands{ 0 };
 	uint16 Index = uint16(-1);
 	bool bRecording = false;
 	bool bEnabled = true;
