Fixed RDG AddCommandListSetupTask to support single tasks or task collections. Added IsParallelSetupEnabled() external call.

diff --git a/Engine/Source/Runtime/RenderCore/Private/RenderGraphBuilder.cpp b/Engine/Source/Runtime/RenderCore/Private/RenderGraphBuilder.cpp
--- a/Engine/Source/Runtime/RenderCore/Private/RenderGraphBuilder.cpp
+++ b/Engine/Source/Runtime/RenderCore/Private/RenderGraphBuilder.cpp
@@ -570,7 +570,7 @@ FRDGBuilder::FRDGBuilder(FRHICommandListImmediate& InRHICmdList, FRDGEventName I
 #endif
 	, GPUScopeStacks(Allocator)
 	, bParallelExecuteEnabled(IsParallelExecuteEnabled() && EnumHasAnyFlags(InFlags, ERDGBuilderFlags::AllowParallelExecute))
-	, bParallelSetupEnabled(IsParallelSetupEnabled() && EnumHasAnyFlags(InFlags, ERDGBuilderFlags::AllowParallelExecute))
+	, bParallelSetupEnabled(::IsParallelSetupEnabled() && EnumHasAnyFlags(InFlags, ERDGBuilderFlags::AllowParallelExecute))
 #if RDG_ENABLE_DEBUG
 	, UserValidation(Allocator, bParallelExecuteEnabled)
 	, BarrierValidation(&Passes, BuilderName)
diff --git a/Engine/Source/Runtime/RenderCore/Public/RenderGraphBuilder.h b/Engine/Source/Runtime/RenderCore/Public/RenderGraphBuilder.h
--- a/Engine/Source/Runtime/RenderCore/Public/RenderGraphBuilder.h
+++ b/Engine/Source/Runtime/RenderCore/Public/RenderGraphBuilder.h
@@ -252,6 +252,12 @@ public:
 	template <typename TaskLambda, typename PrerequisitesCollectionType>
 	UE::Tasks::FTask AddCommandListSetupTask(TaskLambda&& Task, UE::Tasks::FPipe* Pipe, PrerequisitesCollectionType&& Prerequisites, UE::Tasks::ETaskPriority Priority = UE::Tasks::ETaskPriority::Normal, bool bCondition = true);
 
+	/** Whether RDG will launch async tasks when AddSetup{CommandList}Task is called. */
+	inline bool IsParallelSetupEnabled() const
+	{
+		return bParallelSetupEnabled;
+	}
+
 	/** Tells the builder to delete unused RHI resources. The behavior of this method depends on whether RDG immediate mode is enabled:
 	 *   Deferred:  RHI resource flushes are performed prior to execution.
 	 *   Immediate: RHI resource flushes are performed immediately.
diff --git a/Engine/Source/Runtime/RenderCore/Public/RenderGraphBuilder.inl b/Engine/Source/Runtime/RenderCore/Public/RenderGraphBuilder.inl
--- a/Engine/Source/Runtime/RenderCore/Public/RenderGraphBuilder.inl
+++ b/Engine/Source/Runtime/RenderCore/Public/RenderGraphBuilder.inl
@@ -468,6 +468,40 @@ FORCEINLINE UE::Tasks::FTask FRDGBuilder::AddSetupTask(
 	return AddSetupTask(MoveTemp(TaskLambda), nullptr, Forward<PrerequisitesCollectionType&&>(Prerequisites), Priority, bCondition);
 }
 
+namespace UE::RDG
+{
+	template<typename TaskCollectionType, decltype(std::declval<TaskCollectionType>().begin())* = nullptr>
+	inline bool IsCompleted(const TaskCollectionType& Tasks)
+	{
+		for (const UE::Tasks::FTask& Task : Tasks)
+		{
+			if (!Task.IsCompleted())
+			{
+				return false;
+			}
+		}
+		return true;
+	}
+
+	template<typename TaskType, decltype(std::declval<TaskType>().IsCompleted())* = nullptr>
+	inline bool IsCompleted(const TaskType& Task)
+	{
+		return Task.IsCompleted();
+	}
+
+	template<typename TaskCollectionType, decltype(std::declval<TaskCollectionType>().begin())* = nullptr>
+	inline void Wait(const TaskCollectionType& Tasks)
+	{
+		UE::Tasks::Wait(Tasks);
+	}
+
+	template<typename TaskType, decltype(std::declval<TaskType>().Wait())* = nullptr>
+	inline void Wait(const TaskType& Task)
+	{
+		Task.Wait();
+	}
+}
+
 template <typename TaskLambdaType, typename PrerequisitesCollectionType>
 UE::Tasks::FTask FRDGBuilder::AddSetupTask(
 	TaskLambdaType&& TaskLambda,
@@ -490,7 +524,7 @@ UE::Tasks::FTask FRDGBuilder::AddSetupTask(
 
 	if (IsImmediateMode())
 	{
-		UE::Tasks::Wait(Prerequisites);
+		UE::RDG::Wait(Prerequisites);
 		OuterLambda();
 	}
 	else if (Pipe)
@@ -556,17 +590,7 @@ UE::Tasks::FTask FRDGBuilder::AddCommandListSetupTask(
 
 	FRHICommandList* RHICmdListTask = nullptr;
 
-	bool bUseSeparateCommandList = bCondition;
-
-	for (const UE::Tasks::FTask& Prerequisite : Prerequisites)
-	{
-		// Always create a command list when a prerequisite task is present, as the inline task can be scheduled on any thread.
-		if (!Prerequisite.IsCompleted())
-		{
-			bUseSeparateCommandList = true;
-			break;
-		}
-	}
+	const bool bUseSeparateCommandList = bCondition || !UE::RDG::IsCompleted(Prerequisites);
 
 	// When using prerequisites we have to 
 	if (bUseSeparateCommandList)
@@ -599,7 +623,7 @@ UE::Tasks::FTask FRDGBuilder::AddCommandListSetupTask(
 
 	if (IsImmediateMode())
 	{
-		UE::Tasks::Wait(Prerequisites);
+		UE::RDG::Wait(Prerequisites);
 		OuterLambda();
 	}
 	else if (Pipe)
