Instead of assuming the RootUniformBuffer is assigned to a specific slot, we should force it at code generation time.

diff --git a/Engine/Source/Runtime/RenderCore/Private/ShaderParameterParser.cpp b/Engine/Source/Runtime/RenderCore/Private/ShaderParameterParser.cpp
--- a/Engine/Source/Runtime/RenderCore/Private/ShaderParameterParser.cpp
+++ b/Engine/Source/Runtime/RenderCore/Private/ShaderParameterParser.cpp
@@ -1083,12 +1083,13 @@ bool FShaderParameterParser::MoveShaderParametersToRootConstantBuffer(
 		});
 
 		FString CBufferCodeBlock = FString::Printf(
-			TEXT("%s %s\r\n")
+			TEXT("%s %s : register(b%d)\r\n")
 			TEXT("{\r\n")
 			TEXT("%s")
 			TEXT("}\r\n\r\n"),
 			ConstantBufferType,
 			FShaderParametersMetadata::kRootUniformBufferBindingName,
+			FShaderParametersMetadata::kRootCBufferBindingIndex,
 			*RootCBufferContent);
 
 		FString NewShaderCode = (
