ShaderCompiler:  Preprocessor optimizations, first pass.  Saves 16.7% overall on ConditionalPreprocessShader.
* Inline array memory allocation added to low level preprocessor for output and various temporary buffers to reduce dynamic memory allocation and reallocation overhead.  Saved 4.6%.
* FShaderPreprocessOutput::StripCode optimized to write to FString as TCHAR array, rather than using AppendChar (over 4x speedup).  Saved 2.9%
* Shader source file cache now also stores stripped and ANSI converted source, to avoid need to convert and strip the source, plus allocating a copy is avoided.  Saved 4.3%
* Uniform buffer structure declarations stored as ANSI converted source, avoiding convert and copy.  Saved 4.9%

#rnx

diff --git a/Engine/Source/Developer/ShaderPreprocessor/Private/ShaderPreprocessor.cpp b/Engine/Source/Developer/ShaderPreprocessor/Private/ShaderPreprocessor.cpp
--- a/Engine/Source/Developer/ShaderPreprocessor/Private/ShaderPreprocessor.cpp
+++ b/Engine/Source/Developer/ShaderPreprocessor/Private/ShaderPreprocessor.cpp
@@ -12,6 +12,13 @@
 #include "stb_preprocess/stb_alloc.h"
 #include "stb_preprocess/stb_ds.h"
 
+static TAutoConsoleVariable<int32> CVarShaderCompilerThreadLocalPreprocessBuffer(
+	TEXT("r.ShaderCompiler.ThreadLocalPreprocessBuffer"),
+	1280 * 1024,
+	TEXT("Amount to preallocate for preprocess output per worker thread, to save reallocation overhead in the preprocessor."),
+	ECVF_Default
+);
+
 namespace
 {
 	const FString PlatformHeader = TEXT("/Engine/Public/Platform.ush");
@@ -52,6 +59,7 @@ public:
 
 	static void PopulateDefines(const FShaderCompilerEnvironment& Environment, const FShaderCompilerDefinitions& AdditionalDefines, stb_arena* MacroArena, macro_definition**& OutDefines)
 	{
+		arrsetcap(OutDefines, Environment.Definitions->Num() + AdditionalDefines.Num());
 		AddStbDefines(MacroArena, OutDefines, *Environment.Definitions);
 		AddStbDefines(MacroArena, OutDefines, AdditionalDefines);
 	}
@@ -91,11 +99,19 @@ extern "C"
 	}
 }
 
+struct FStbLoadedInclude
+{
+	const ANSICHAR* Data = nullptr;				// Points to SharedData, LocalData, or data from FShaderCompilerEnvironment
+	size_t DataLength = 0;
+	FShaderSharedAnsiStringPtr SharedData;
+	TArray<ANSICHAR> LocalData;
+};
+
 struct FStbPreprocessContext
 {
 	const FShaderCompilerInput& ShaderInput;
 	const FShaderCompilerEnvironment& Environment;
-	TMap<FString, TArray<ANSICHAR>> LoadedIncludesCache;
+	TMap<FString, FStbLoadedInclude> LoadedIncludesCache;
 	TMap<FString, TUniquePtr<ANSICHAR[]>> SeenPathsLowerCase;
 
 	bool HasIncludedMandatoryHeaders()
@@ -104,137 +120,48 @@ struct FStbPreprocessContext
 	}
 };
 
-inline bool IsEndOfLine(TCHAR C)
-{
-	return C == TEXT('\r') || C == TEXT('\n');
-}
-
-inline bool CommentStripNeedsHandling(TCHAR C)
-{
-	return IsEndOfLine(C) || C == TEXT('/') || C == 0;
-}
-
-inline int NewlineCharCount(TCHAR First, TCHAR Second)
-{
-	return ((First + Second) == TEXT('\r') + TEXT('\n')) ? 2 : 1;
-}
-
-// Given an FString containing the contents of a shader source file, populates the given array with contents of
-// that source file with all comments stripped. This is needed since the STB preprocessor itself does not strip 
-// comments.
-void ConvertAndStripComments(const FString& ShaderSource, TArray<ANSICHAR>& OutStripped)
+static const ANSICHAR* StbLoadFile(const ANSICHAR* Filename, void* RawContext, size_t* OutLength)
 {
-	// STB preprocessor does not strip comments, so we do so here before returning the loaded source
-	// Doing so is barely more costly than the memcopy we require anyways so has negligible overhead.
-	// Reserve worst case (i.e. assuming there are no comments at all) to avoid reallocation
-	int32 BufferSize = ShaderSource.Len() + 1; // +1 to append null terminator
-	OutStripped.SetNumUninitialized(BufferSize);
-
-	ANSICHAR* CurrentOut = OutStripped.GetData();
+	FStbPreprocessContext& Context = *reinterpret_cast<FStbPreprocessContext*>(RawContext);
+	FString FilenameConverted = StringCast<TCHAR>(Filename).Get();
+	uint32 FilenameConvertedHash = GetTypeHash(FilenameConverted);
+	FStbLoadedInclude* ContentsCached = Context.LoadedIncludesCache.FindByHash(FilenameConvertedHash, FilenameConverted);
+	if (!ContentsCached)
+	{
+		ContentsCached = &Context.LoadedIncludesCache.AddByHash(FilenameConvertedHash, FilenameConverted);
 
-	const TCHAR* const End = ShaderSource.GetCharArray().GetData() + ShaderSource.Len();
+		const FString* InMemorySource = Context.Environment.IncludeVirtualPathToContentsMap.FindByHash(FilenameConvertedHash, FilenameConverted);
 
-	// We rely on null termination to avoid the need to check Current < End in some cases
-	check(*End == TEXT('\0'));
-	for (const TCHAR* Current = ShaderSource.GetCharArray().GetData(); Current < End;)
-	{
-		// sanity check that we're not overrunning the buffer
-		check(CurrentOut < (OutStripped.GetData() + BufferSize));
-		// CommentStripNeedsHandling returns true when *Current == '\0';
-		while (!CommentStripNeedsHandling(*Current))
+		if (InMemorySource)
 		{
-			// straight cast to ansichar; since this is a character in hlsl source that's not in a comment
-			// we assume that it must be valid to do so. if this assumption is not valid the shader source was
-			// broken/corrupt anyways.
-			*CurrentOut++ = (ANSICHAR)(*Current++);
-		}
+			check(!InMemorySource->IsEmpty());
+			ShaderConvertAndStripComments(*InMemorySource, ContentsCached->LocalData);
 
-		if (IsEndOfLine(*Current))
-		{
-			*CurrentOut++ = '\n';
-			Current += NewlineCharCount(Current[0], Current[1]);
+			ContentsCached->Data = ContentsCached->LocalData.GetData();
+			ContentsCached->DataLength = ContentsCached->LocalData.Num();
 		}
-		else if (Current[0] == '/')
+		else
 		{
-			if (Current[1] == '/')
+			const FThreadSafeSharedAnsiStringPtr* InMemorySourceAnsi = Context.Environment.IncludeVirtualPathToSharedContentsMap.FindByHash(FilenameConvertedHash, FilenameConverted);
+
+			if (InMemorySourceAnsi)
 			{
-				while (!IsEndOfLine(*Current) && Current < End)
-				{
-					++Current;
-				}
-			}
-			else if (Current[1] == '*')
-			{
-				Current += 2;
-				while (Current < End)
-				{
-					if (Current[0] == '*' && Current[1] == '/')
-					{
-						Current += 2;
-						break;
-					}
-					else if (IsEndOfLine(*Current))
-					{
-						*CurrentOut++ = '\n';
-						Current += NewlineCharCount(Current[0], Current[1]);
-					}
-					else
-					{
-						++Current;
-					}
-				}
+				ContentsCached->Data = InMemorySourceAnsi->Get()->GetData();
+				ContentsCached->DataLength = InMemorySourceAnsi->Get()->Num();
 			}
 			else
 			{
-				*CurrentOut++ = (ANSICHAR)(*Current++);
-			}
-		}
-	}
-	// Null terminate after comment-stripped copy
-	check(CurrentOut < (OutStripped.GetData() + BufferSize));
-	*CurrentOut++ = 0;
-
-	// Set correct length after stripping but don't bother shrinking/reallocating, minor memory overhead to save time
-	OutStripped.SetNum(CurrentOut - OutStripped.GetData(), /* bAllowShrinking */false);
-}
-
-const FString* FindInMemorySource(const FShaderCompilerEnvironment& Environment, const FString& FilenameConverted)
-{
-	const FString* InMemorySource = Environment.IncludeVirtualPathToContentsMap.Find(FilenameConverted);
-	if (!InMemorySource)
-	{
-		const FThreadSafeSharedStringPtr* SharedPtr = Environment.IncludeVirtualPathToExternalContentsMap.Find(FilenameConverted);
-		InMemorySource = SharedPtr ? SharedPtr->Get() : nullptr;
-	}
-	return InMemorySource;
-}
+				CheckShaderHashCacheInclude(FilenameConverted, Context.ShaderInput.Target.GetPlatform(), Context.ShaderInput.ShaderFormat.ToString());
+				LoadShaderSourceFile(*FilenameConverted, Context.ShaderInput.Target.GetPlatform(), nullptr, nullptr, nullptr, &ContentsCached->SharedData);
 
-static const ANSICHAR* StbLoadFile(const ANSICHAR* Filename, void* RawContext, size_t* OutLength)
-{
-	FStbPreprocessContext& Context = *reinterpret_cast<FStbPreprocessContext*>(RawContext);
-	FString FilenameConverted = StringCast<TCHAR>(Filename).Get();
-	TArray<ANSICHAR>* ContentsCached = Context.LoadedIncludesCache.Find(FilenameConverted);
-	if (!ContentsCached)
-	{
-		ContentsCached = &Context.LoadedIncludesCache.Add(FilenameConverted);
-		// Local FString used for the LoadShaderSourceFile path; we should consider retrieving source from the shader file cache as a reference
-		// (avoid an extra alloc+copy)
-		FString SourceCopy;
-
-		const FString* InMemorySource = FindInMemorySource(Context.Environment, FilenameConverted);
-
-		if (!InMemorySource)
-		{
-			CheckShaderHashCacheInclude(FilenameConverted, Context.ShaderInput.Target.GetPlatform(), Context.ShaderInput.ShaderFormat.ToString());
-			LoadShaderSourceFile(*FilenameConverted, Context.ShaderInput.Target.GetPlatform(), &SourceCopy, nullptr);
-			InMemorySource = &SourceCopy;
+				ContentsCached->Data = ContentsCached->SharedData->GetData();
+				ContentsCached->DataLength = ContentsCached->SharedData->Num();
+			}
 		}
-		check(InMemorySource && !InMemorySource->IsEmpty());
-		ConvertAndStripComments(*InMemorySource, *ContentsCached);
 	}
 	check(ContentsCached);
-	*OutLength = ContentsCached->Num();
-	return ContentsCached->GetData();
+	*OutLength = ContentsCached->DataLength;
+	return ContentsCached->Data;
 }
 
 static void StbFreeFile(const ANSICHAR* Filename, const ANSICHAR* Contents, void* RawContext)
@@ -268,7 +195,7 @@ static const ANSICHAR* StbResolveInclude(const ANSICHAR* PathInSource, uint32 Pa
 
 	bool bExists =
 		Context.Environment.IncludeVirtualPathToContentsMap.Contains(PathModified) ||
-		Context.Environment.IncludeVirtualPathToExternalContentsMap.Contains(PathModified) ||
+		Context.Environment.IncludeVirtualPathToSharedContentsMap.Contains(PathModified) ||
 		// LoadShaderSourceFile will load the file if it exists, but then cache it internally, so the next call in StbLoadFile will be cheap
 		// (and hence this is not wasteful, just performs the loading earlier)
 		LoadShaderSourceFile(*PathModified, Context.ShaderInput.Target.GetPlatform(), nullptr, nullptr);
@@ -333,8 +260,19 @@ bool InnerPreprocessShaderStb(
 	auto InFilename = StringCast<ANSICHAR>(*Input.VirtualSourceFilePath);
 	int NumDiagnostics = 0;
 	pp_diagnostic* Diagnostics = nullptr;
-	
-	char* OutPreprocessedAnsi = preprocess_file(nullptr, InFilename.Get(), &Context, StbDefines, arrlen(StbDefines), &Diagnostics, &NumDiagnostics);
+
+	static const int32 ThreadLocalPreprocessBufferSize = CVarShaderCompilerThreadLocalPreprocessBuffer.GetValueOnAnyThread();
+	static thread_local char* ThreadLocalPreprocessBuffer = nullptr;
+
+	// Sanity check the buffer size so it won't OOM if a bad value is entered.
+	int32 ClampedPreprocessBufferSize = ThreadLocalPreprocessBufferSize ? FMath::Clamp(ThreadLocalPreprocessBufferSize, 64 * 1024, 4 * 1024 * 1024) : 0;
+	if (ClampedPreprocessBufferSize && !ThreadLocalPreprocessBuffer)
+	{
+		ThreadLocalPreprocessBuffer = new char[ClampedPreprocessBufferSize];
+	}
+
+	char* OutPreprocessedAnsi = preprocess_file(nullptr, InFilename.Get(), &Context, StbDefines, arrlen(StbDefines), &Diagnostics, &NumDiagnostics, ThreadLocalPreprocessBuffer, ClampedPreprocessBufferSize);
+
 	bool HasError = false;
 	if (Diagnostics != nullptr)
 	{
diff --git a/Engine/Source/Developer/ShaderPreprocessor/Private/stb_preprocess/preprocessor.c b/Engine/Source/Developer/ShaderPreprocessor/Private/stb_preprocess/preprocessor.c
--- a/Engine/Source/Developer/ShaderPreprocessor/Private/stb_preprocess/preprocessor.c
+++ b/Engine/Source/Developer/ShaderPreprocessor/Private/stb_preprocess/preprocessor.c
@@ -2953,6 +2953,8 @@ static void maybe_expand_macro(parse_state* cs, struct macro_definition* pending
 		int parse_rest = 0;
 		const char* s = p;
 
+		stbds_arrinline(arguments, char*, 16);
+
 		// plan:
 		//   1. find opening parenthesis, if any
 		//   2. parse arguments into arguments[] array
@@ -3001,9 +3003,13 @@ static void maybe_expand_macro(parse_state* cs, struct macro_definition* pending
 		//   2. parse arguments into arguments[] array
 		//
 
+		// Inline storage shared for all argument strings.  Uses stbds_arrinline_suballoc to reuse remaining portion of buffer.
+		char* argument_buffer;
+		stbds_arrinline(argument_buffer, char, 4096);
+
 		for (i = 0; i < md->num_parameters; ++i)
 		{
-			char* copy = 0;
+			char* copy = argument_buffer;
 			int arg_newlines = 0;
 
 			p = preprocessor_skip_whitespace(p, &cs->src_line_number);
@@ -3062,6 +3068,12 @@ static void maybe_expand_macro(parse_state* cs, struct macro_definition* pending
 				}
 				*e = 0;
 
+				// Sub-allocate the remaining space in argument_buffer for subsequent arguments.  Note that it's necessary to do the suballocation BEFORE
+				// the call to arrsetlen below, because downstream code assumes there is a null terminator beyond the official length of the array.  Doing
+				// the sub-allocation first preserves the null terminator -- without this, the sub-allocation will overwrite it.  We only suballocate if
+				// a minimum of 200 characters is available, as copy_argument reserves this.
+				argument_buffer = stbds_arrinline_suballoc(argument_buffer, char, 200);
+
 				arrsetlen(copy, e - p);
 				arrput(arguments, copy);
 				cs->src_line_number += arg_newlines;
@@ -3099,7 +3111,7 @@ static void maybe_expand_macro(parse_state* cs, struct macro_definition* pending
 			else
 			{
 				int arg_newlines = 0;
-				char *e, *copy = 0;
+				char *e, *copy = argument_buffer;
 				p = preprocessor_skip_whitespace(p, &cs->src_line_number);
 				// parse all the remaining arguments
 				s = p;
@@ -3137,6 +3149,8 @@ static void maybe_expand_macro(parse_state* cs, struct macro_definition* pending
 					e = copy;
 				*e = 0;
 
+				argument_buffer = stbds_arrinline_suballoc(argument_buffer, char, 200);
+
 				arrput(arguments, copy);
 				cs->src_line_number += arg_newlines;
 			}
@@ -3167,6 +3181,7 @@ static void maybe_expand_macro(parse_state* cs, struct macro_definition* pending
 
 		// populate temporary buffer
 		char* tmp = 0;
+		stbds_arrinline(tmp, char, 4096);
 		for (i = 0; i < arrlen(md->expansion) - 1; ++i)
 		{
 			int an;
@@ -3297,6 +3312,9 @@ static char* macro_expand_directive(parse_state* cs, const char* p, char* direct
 
 static int evaluate_if(parse_state* cs, const char* p, int* syntax_error)
 {
+	char* dest = 0;
+	stbds_arrinline(dest, char, 1024);
+
 	int result = 0;
 	parse_state ncs;
 	ncs = *cs;
@@ -3304,7 +3322,7 @@ static int evaluate_if(parse_state* cs, const char* p, int* syntax_error)
 	ncs.src_length = strlen(p);
 	ncs.src_offset = 0;
 	ncs.parent = cs;
-	ncs.dest = 0;
+	ncs.dest = dest;
 	ncs.fast_dest = 0;
 
 	preprocess_string(&ncs, IN_MACRO_if_condition, "#if", NULL);
@@ -3484,7 +3502,7 @@ static void process_directive(parse_state* cs, conditional_state* cons)
 	pp_context* c = cs->context;
 	const char *p, *endptr;
 	const char* alloc = NULL;	 // all exit paths SHOULD do 'if (alloc) arrfree(alloc)', but we allow it to leak on error-handling returns
-	char buffer[1000];
+	char buffer[2048];
 	int dummy;	// uninitialized, should never get incremented by parse_directive_after_hash
 
 	assert(cs->src[cs->src_offset] == '#');
@@ -4576,7 +4594,9 @@ char* preprocess_file(char* output_autobuffer,
 	struct macro_definition** predefined_macros,
 	int num_predefined_macros,
 	pp_diagnostic** pd,
-	int* num_pd)
+	int* num_pd,
+	char* output_inlinebuffer,
+	size_t output_inlinebuffersize)
 {
 	char* output = 0;
 	pp_context c = { 0 };
@@ -4584,6 +4604,11 @@ char* preprocess_file(char* output_autobuffer,
 	size_t length;
 	const char* main_file;
 
+	if (output_inlinebuffer && output_inlinebuffersize)
+	{
+		output = (char*)stbds_arrinlinef((size_t*)output_inlinebuffer, sizeof(char), output_inlinebuffersize - sizeof(stbds_array_header));
+	}
+
 	main_file = (*loadfile_callback)(filename, custom_context, &length);
 	*num_pd = 0;
 
@@ -4627,6 +4652,7 @@ char* preprocess_file(char* output_autobuffer,
 		// initial parse_state for preprocess_string_from_file to initialize from
 		parse_state ps = {0};
 		ps.context = &c;
+		ps.dest = output;
 		ps.fast_dest = output_autobuffer;
 		ps.state_limit = PP_STATE_active_count;
 
@@ -4683,6 +4709,15 @@ char* preprocess_file(char* output_autobuffer,
 	return output;
 }
 
+int preprocessor_file_size(char* text)
+{
+	return text ? arrlen(text) : 0;
+}
+int preprocessor_file_capacity(char* text)
+{
+	return text ? arrcap(text) : 0;
+}
+
 void preprocessor_file_free(char* text, pp_diagnostic* pd)
 {
 	int i, j;
diff --git a/Engine/Source/Developer/ShaderPreprocessor/Private/stb_preprocess/preprocessor.h b/Engine/Source/Developer/ShaderPreprocessor/Private/stb_preprocess/preprocessor.h
--- a/Engine/Source/Developer/ShaderPreprocessor/Private/stb_preprocess/preprocessor.h
+++ b/Engine/Source/Developer/ShaderPreprocessor/Private/stb_preprocess/preprocessor.h
@@ -91,7 +91,12 @@ STB_PP_DEF char* preprocess_file(
 	struct macro_definition** predefined_macros,
 	int num_predefined_macros,
 	pp_diagnostic** pd,
-	int* num_pd);
+	int* num_pd,
+	char* output_inlinebuffer,
+	size_t output_inlinebuffersize);
+
+STB_PP_DEF int preprocessor_file_size(char* text);
+STB_PP_DEF int preprocessor_file_capacity(char* text);
 
 // frees memory allocated by preprocess_file (preprocessed results and diagnostic messages)
 STB_PP_DEF void preprocessor_file_free(char* text, pp_diagnostic* pd);
diff --git a/Engine/Source/Developer/ShaderPreprocessor/Private/stb_preprocess/stb_ds.c b/Engine/Source/Developer/ShaderPreprocessor/Private/stb_preprocess/stb_ds.c
--- a/Engine/Source/Developer/ShaderPreprocessor/Private/stb_preprocess/stb_ds.c
+++ b/Engine/Source/Developer/ShaderPreprocessor/Private/stb_preprocess/stb_ds.c
@@ -30,6 +30,56 @@ size_t stbds_rehash_items;
 // int *prev_allocs[65536];
 // int num_prev;
 
+void* stbds_arrinlinef(size_t* buf, size_t elemsize, size_t elemcount)
+{
+	stbds_array_header* h = (stbds_array_header*)buf;
+
+	h->length = 0;
+	h->capacity = elemcount;
+	h->hash_table = 0;
+	h->temp = 0;
+	h->elemsize = elemsize;
+	h->inlinealloc = 1;
+
+	return h + 1;
+}
+
+void* stbds_arrinline_suballocf(void* a, size_t min_capacity)
+{
+	if (!a || !stbds_header(a)->inlinealloc)
+	{
+		return 0;
+	}
+
+	stbds_array_header* h = stbds_header(a);
+
+	// Get the end of the existing array and buffer from the capacity -- used_end is rounded upward to word size
+	char* used_end = (char*)a + ((h->elemsize * h->length + sizeof(size_t) - 1) & ~(sizeof(size_t) - 1));
+	char* capacity_end = (char*)a + h->elemsize * h->capacity;
+
+	// Set existing allocation capacity equal to length, so no more can be allocated from the inline buffer.
+	// Growing will still work, it will just convert the inline allocation to a dynamic allocation.
+	h->capacity = h->length;
+
+	// Determine whether we have the minimum required capacity for the suballocation, if not return null.
+	if (capacity_end - used_end < sizeof(stbds_array_header) + min_capacity * h->elemsize)
+	{
+		return 0;
+	}
+
+	// Set up a new buffer
+	stbds_array_header* h_suballoc = (stbds_array_header*)used_end;
+
+	h_suballoc->length = 0;
+	h_suballoc->capacity = ((capacity_end - used_end) - sizeof(stbds_array_header)) / h->elemsize;
+	h_suballoc->hash_table = 0;
+	h_suballoc->temp = 0;
+	h_suballoc->elemsize = h->elemsize;
+	h_suballoc->inlinealloc = 1;
+
+	return h_suballoc + 1;
+}
+
 void* stbds_arrgrowf(void* a, size_t elemsize, size_t addlen, size_t min_cap)
 {
 	stbds_array_header temp = {0};	// force debugging
@@ -53,7 +103,17 @@ void* stbds_arrgrowf(void* a, size_t elemsize, size_t addlen, size_t min_cap)
 	// if (num_prev < 65536) if (a) prev_allocs[num_prev++] = (int *) ((char *) a+1);
 	// if (num_prev == 2201)
 	//   num_prev = num_prev;
-	b = STB_COMMON_REALLOC((a) ? stbds_header(a) : 0, elemsize * min_cap + sizeof(stbds_array_header));
+	if (a && stbds_header(a)->inlinealloc)
+	{
+		b = STB_COMMON_MALLOC(elemsize * min_cap + sizeof(stbds_array_header));
+		STB_ASSUME(b != NULL);
+		memcpy(b, stbds_header(a), elemsize * stbds_header(a)->length + sizeof(stbds_array_header));
+		((stbds_array_header*)b)->inlinealloc = 0;
+	}
+	else
+	{
+		b = STB_COMMON_REALLOC((a) ? stbds_header(a) : 0, elemsize * min_cap + sizeof(stbds_array_header));
+	}
 	STB_ASSUME(b != NULL);
 
 	// if (num_prev < 65536) prev_allocs[num_prev++] = (int *) (char *) b;
@@ -63,6 +123,8 @@ void* stbds_arrgrowf(void* a, size_t elemsize, size_t addlen, size_t min_cap)
 		stbds_header(b)->length = 0;
 		stbds_header(b)->hash_table = 0;
 		stbds_header(b)->temp = 0;
+		stbds_header(b)->elemsize = elemsize;
+		stbds_header(b)->inlinealloc = 0;
 	}
 	else
 	{
diff --git a/Engine/Source/Developer/ShaderPreprocessor/Private/stb_preprocess/stb_ds.h b/Engine/Source/Developer/ShaderPreprocessor/Private/stb_preprocess/stb_ds.h
--- a/Engine/Source/Developer/ShaderPreprocessor/Private/stb_preprocess/stb_ds.h
+++ b/Engine/Source/Developer/ShaderPreprocessor/Private/stb_preprocess/stb_ds.h
@@ -471,6 +471,8 @@ extern "C"
 	// Everything below here is implementation details
 	//
 
+	extern void* stbds_arrinlinef(size_t* buf, size_t elemsize, size_t elemcount);
+	extern void* stbds_arrinline_suballocf(void* a, size_t min_capacity);
 	extern void* stbds_arrgrowf(void* a, size_t elemsize, size_t addlen, size_t min_cap);
 	extern void stbds_hmfree_func(void* p, size_t elemsize);
 	extern void* stbds_hmget_key(void* a, size_t elemsize, void* key, size_t keysize, int mode);
@@ -514,6 +516,15 @@ extern "C"
 #define stbds_temp(t) stbds_header(t)->temp
 #define stbds_temp_key(t) (*(char**)stbds_header(t)->hash_table)
 
+// Macro to reserve inline (usually stack memory) storage for an array
+#define stbds_arrinline(a, a_type, n) \
+	size_t __buf ## a[(sizeof(stbds_array_header) + n * sizeof(*a) + sizeof(size_t) - 1) / sizeof(size_t)];	\
+	a = (a_type*)stbds_arrinlinef(__buf ## a, sizeof(*a), n)
+
+// Allow a single inline storage buffer to be used for multiple sub-allocations.  Trims current allocation to what's used and returns a new allocation with the remainder if enough capacity exists.
+#define stbds_arrinline_suballoc(a, a_type, min_capacity) \
+	a = (a_type*)stbds_arrinline_suballocf(a, min_capacity)
+
 #define stbds_arrsetcap(a, n) (stbds_arrgrow(a, 0, n))
 #define stbds_arrsetlen(a, n) ((stbds_arrcap(a) < (size_t)(n) ? stbds_arrsetcap((a), (size_t)(n)), 0 : 0), (a) ? stbds_header(a)->length = (size_t)(n) : 0)
 #define stbds_arrinitlen(a, n) ((stbds_arrsetcap((a), (size_t)(n)), stbds_header(a)->length = (size_t)(n)))
@@ -529,7 +540,7 @@ extern "C"
 #define stbds_arraddnoff(a, n) (stbds_arrmaybegrow(a, n), stbds_header(a)->length += (n), stbds_header(a)->length - (n))
 #define stbds_arraddnindex(a, n) stbds_arraddnoff(a, n)
 #define stbds_arrlast(a) ((a)[stbds_header(a)->length - 1])
-#define stbds_arrfree(a) ((void)((a) ? STB_COMMON_FREE(stbds_header(a)) : (void)0), (a) = NULL)
+#define stbds_arrfree(a) ((void)((a) && !stbds_header(a)->inlinealloc ? STB_COMMON_FREE(stbds_header(a)) : (void)0), (a) = NULL)
 #define stbds_arrdel(a, i) stbds_arrdeln(a, i, 1)
 #define stbds_arrdeln(a, i, n) (memmove(&(a)[i], &(a)[(i) + (n)], sizeof *(a) * (stbds_header(a)->length - (n) - (i))), stbds_header(a)->length -= (n))
 #define stbds_arrdelswap(a, i) ((a)[i] = stbds_arrlast(a), stbds_header(a)->length -= 1)
@@ -619,6 +630,8 @@ typedef struct
 {
 	size_t length;
 	size_t capacity;
+	int elemsize;
+	int inlinealloc;		// Allocated inline, don't free!
 	void* hash_table;
 	ptrdiff_t temp;
 } stbds_array_header;
diff --git a/Engine/Source/Programs/ShaderCompileWorker/Private/ShaderCompileWorker.cpp b/Engine/Source/Programs/ShaderCompileWorker/Private/ShaderCompileWorker.cpp
--- a/Engine/Source/Programs/ShaderCompileWorker/Private/ShaderCompileWorker.cpp
+++ b/Engine/Source/Programs/ShaderCompileWorker/Private/ShaderCompileWorker.cpp
@@ -455,7 +455,7 @@ private:
 		};
 
 		// Shared inputs
-		TMap<FString, FThreadSafeSharedStringPtr> ExternalIncludes;
+		TMap<FString, FThreadSafeSharedAnsiStringPtr> ExternalIncludes;
 		{
 			int32 NumExternalIncludes = 0;
 			InputFile << NumExternalIncludes;
@@ -465,7 +465,7 @@ private:
 			{
 				FString NewIncludeName;
 				InputFile << NewIncludeName;
-				FString* NewIncludeContents = new FString();
+				TArray<ANSICHAR>* NewIncludeContents = new TArray<ANSICHAR>;
 				InputFile << (*NewIncludeContents);
 				ExternalIncludes.Add(NewIncludeName, MakeShareable(NewIncludeContents));
 			}
diff --git a/Engine/Source/Runtime/Engine/Private/ShaderCompiler/ShaderCompiler.cpp b/Engine/Source/Runtime/Engine/Private/ShaderCompiler/ShaderCompiler.cpp
--- a/Engine/Source/Runtime/Engine/Private/ShaderCompiler/ShaderCompiler.cpp
+++ b/Engine/Source/Runtime/Engine/Private/ShaderCompiler/ShaderCompiler.cpp
@@ -2527,12 +2527,12 @@ bool DoWriteTasksInner(const TArray<FShaderCommonCompileJobPtr>& QueuedJobs, FAr
 
 	// Gather External Includes and serialize separately, these are largely shared between jobs
 	{
-		TMap<FString, FString> ExternalIncludes;
+		TMap<FString, TArray<ANSICHAR>> ExternalIncludes;
 		ExternalIncludes.Reserve(32);
 
 		for (int32 JobIndex = 0; JobIndex < QueuedSingleJobs.Num(); JobIndex++)
 		{
-			QueuedSingleJobs[JobIndex]->Input.GatherSharedInputs(ExternalIncludes, SharedEnvironments, RequestShaderParameterStructures);
+			QueuedSingleJobs[JobIndex]->Input.GatherSharedInputsAnsi(ExternalIncludes, SharedEnvironments, RequestShaderParameterStructures);
 		}
 
 		for (int32 JobIndex = 0; JobIndex < QueuedPipelineJobs.Num(); JobIndex++)
@@ -2542,14 +2542,14 @@ bool DoWriteTasksInner(const TArray<FShaderCommonCompileJobPtr>& QueuedJobs, FAr
 
 			for (int32 Index = 0; Index < NumStageJobs; Index++)
 			{
-				PipelineJob->StageJobs[Index]->Input.GatherSharedInputs(ExternalIncludes, SharedEnvironments, RequestShaderParameterStructures);
+				PipelineJob->StageJobs[Index]->Input.GatherSharedInputsAnsi(ExternalIncludes, SharedEnvironments, RequestShaderParameterStructures);
 			}
 		}
 
 		int32 NumExternalIncludes = ExternalIncludes.Num();
 		TransferFile << NumExternalIncludes;
 
-		for (TMap<FString, FString>::TIterator It(ExternalIncludes); It; ++It)
+		for (TMap<FString, TArray<ANSICHAR>>::TIterator It(ExternalIncludes); It; ++It)
 		{
 			TransferFile << It.Key();
 			TransferFile << It.Value();
@@ -7237,7 +7237,7 @@ void ValidateShaderFilePath(const FString& VirtualShaderFilePath, const FString&
 FCriticalSection GCachedGeneratedInstancedStereoCodeLock;
 
 /** Storage for instanced stereo code so it is not generated every time we compile a shader. */
-TMap<EShaderPlatform, FThreadSafeSharedStringPtr> GCachedGeneratedInstancedStereoCode;
+TMap<EShaderPlatform, FThreadSafeSharedAnsiStringPtr> GCachedGeneratedInstancedStereoCode;
 
 void GlobalBeginCompileShader(
 	const FString& DebugGroupName,
@@ -7335,7 +7335,7 @@ void GlobalBeginCompileShader(
 			ValidateShaderFilePath(Entry.Key, Input.VirtualSourceFilePath);
 		}
 
-		for (const auto& Entry : Input.Environment.IncludeVirtualPathToExternalContentsMap)
+		for (const auto& Entry : Input.Environment.IncludeVirtualPathToSharedContentsMap)
 		{
 			ValidateShaderFilePath(Entry.Key, Input.VirtualSourceFilePath);
 		}
@@ -7555,10 +7555,10 @@ void GlobalBeginCompileShader(
 		SET_SHADER_DEFINE_AND_COMPILE_ARGUMENT(Input.Environment, MOBILE_MULTI_VIEW, false);
 	}
 
-	// Reserve space in maps to prevent reallocation and rehashing in AddUniformBufferIncludesToEnvironment
-	const int32 UniformBufferReserveNum = Input.Environment.UniformBufferMap.Num() + ShaderType->GetReferencedUniformBufferNames().Num() + (VFType ? VFType->GetReferencedUniformBufferNames().Num() : 0);
+	// Reserve space in maps to prevent reallocation and rehashing in AddUniformBufferIncludesToEnvironment -- plus one at the end is for GeneratedInstancedStereo.ush
+	const int32 UniformBufferReserveNum = Input.Environment.UniformBufferMap.Num() + ShaderType->GetReferencedUniformBufferNames().Num() + (VFType ? VFType->GetReferencedUniformBufferNames().Num() : 0) + 1;
 	Input.Environment.UniformBufferMap.Reserve(UniformBufferReserveNum);
-	Input.Environment.IncludeVirtualPathToExternalContentsMap.Reserve(UniformBufferReserveNum);
+	Input.Environment.IncludeVirtualPathToSharedContentsMap.Reserve(UniformBufferReserveNum);
 
 	ShaderType->AddUniformBufferIncludesToEnvironment(Input.Environment, ShaderPlatform);
 
@@ -7572,16 +7572,20 @@ void GlobalBeginCompileShader(
 		// this function may be called on multiple threads, so protect the storage
 		FScopeLock GeneratedInstancedCodeLock(&GCachedGeneratedInstancedStereoCodeLock);
 
-		FThreadSafeSharedStringPtr* Existing = GCachedGeneratedInstancedStereoCode.Find(ShaderPlatform);
-		FThreadSafeSharedStringPtr CachedCodePtr = Existing ? *Existing : nullptr;
+		FThreadSafeSharedAnsiStringPtr* Existing = GCachedGeneratedInstancedStereoCode.Find(ShaderPlatform);
+		FThreadSafeSharedAnsiStringPtr CachedCodePtr = Existing ? *Existing : nullptr;
 		if (!CachedCodePtr.IsValid())
 		{
-			CachedCodePtr = MakeShareable(new FString());
-			GenerateInstancedStereoCode(*CachedCodePtr.Get(), ShaderPlatform);
+			FString CachedCode;
+			GenerateInstancedStereoCode(CachedCode, ShaderPlatform);
+
+			CachedCodePtr = MakeShareable(new TArray<ANSICHAR>());
+			ShaderConvertAndStripComments(CachedCode, *CachedCodePtr);
+
 			GCachedGeneratedInstancedStereoCode.Add(ShaderPlatform, CachedCodePtr);
 		}
 
-		Input.Environment.IncludeVirtualPathToExternalContentsMap.Add(TEXT("/Engine/Generated/GeneratedInstancedStereo.ush"), CachedCodePtr);
+		Input.Environment.IncludeVirtualPathToSharedContentsMap.Add(TEXT("/Engine/Generated/GeneratedInstancedStereo.ush"), CachedCodePtr);
 	}
 
 	{
diff --git a/Engine/Source/Runtime/RenderCore/Private/ShaderCore.cpp b/Engine/Source/Runtime/RenderCore/Private/ShaderCore.cpp
--- a/Engine/Source/Runtime/RenderCore/Private/ShaderCore.cpp
+++ b/Engine/Source/Runtime/RenderCore/Private/ShaderCore.cpp
@@ -157,13 +157,19 @@ PRAGMA_ENABLE_DEPRECATION_WARNINGS
 // Apply lock striping as we're mostly reader lock bound.
 constexpr int32 GSHADERFILECACHE_BUCKETS = 31; /* prime number for best distribution using modulo */
 
+struct FShaderFileCacheEntry
+{
+	FString Source;
+	FShaderSharedAnsiStringPtr StrippedSource;		// Source with comments stripped out, and converted to ANSICHAR
+};
+
 struct FShaderFileCache
 {
 	/** Protects Map from simultaneous access by multiple threads. */
 	FRWLock Lock;
 
 	/** The shader file cache, used to minimize shader file reads */
-	TMap<FString, FString> Map;
+	TMap<FString, FShaderFileCacheEntry> Map;
 } GShaderFileCache[GSHADERFILECACHE_BUCKETS];
 
 class FShaderHashCache
@@ -851,7 +857,7 @@ void FShaderCompilerEnvironment::Merge(const FShaderCompilerEnvironment& Other)
 		}
 	}
 
-	check(Other.IncludeVirtualPathToExternalContentsMap.Num() == 0);
+	check(Other.IncludeVirtualPathToSharedContentsMap.Num() == 0);
 
 	CompilerFlags.Append(Other.CompilerFlags);
 	ResourceTableMap.Append(Other.ResourceTableMap);
@@ -1809,7 +1815,101 @@ void FixupShaderFilePath(FString& VirtualFilePath, EShaderPlatform ShaderPlatfor
 	ReplaceVirtualFilePathForShaderAutogen(VirtualFilePath, ShaderPlatform, ShaderPlatformName);
 }
 
-bool LoadShaderSourceFile(const TCHAR* InVirtualFilePath, EShaderPlatform ShaderPlatform, FString* OutFileContents, TArray<FShaderCompilerError>* OutCompileErrors, const FName* ShaderPlatformName) // TODO: const FString&
+inline bool IsEndOfLine(TCHAR C)
+{
+	return C == TEXT('\r') || C == TEXT('\n');
+}
+
+inline bool CommentStripNeedsHandling(TCHAR C)
+{
+	return IsEndOfLine(C) || C == TEXT('/') || C == 0;
+}
+
+inline int NewlineCharCount(TCHAR First, TCHAR Second)
+{
+	return ((First + Second) == TEXT('\r') + TEXT('\n')) ? 2 : 1;
+}
+
+// Given an FString containing the contents of a shader source file, populates the given array with contents of
+// that source file with all comments stripped. This is needed since the STB preprocessor itself does not strip 
+// comments.
+void ShaderConvertAndStripComments(const FString& ShaderSource, TArray<ANSICHAR>& OutStripped)
+{
+	// STB preprocessor does not strip comments, so we do so here before returning the loaded source
+	// Doing so is barely more costly than the memcopy we require anyways so has negligible overhead.
+	// Reserve worst case (i.e. assuming there are no comments at all) to avoid reallocation
+	int32 BufferSize = ShaderSource.Len() + 1; // +1 to append null terminator
+	OutStripped.SetNumUninitialized(BufferSize);
+
+	ANSICHAR* CurrentOut = OutStripped.GetData();
+
+	const TCHAR* const End = ShaderSource.GetCharArray().GetData() + ShaderSource.Len();
+
+	// We rely on null termination to avoid the need to check Current < End in some cases
+	check(*End == TEXT('\0'));
+	for (const TCHAR* Current = ShaderSource.GetCharArray().GetData(); Current < End;)
+	{
+		// sanity check that we're not overrunning the buffer
+		check(CurrentOut < (OutStripped.GetData() + BufferSize));
+		// CommentStripNeedsHandling returns true when *Current == '\0';
+		while (!CommentStripNeedsHandling(*Current))
+		{
+			// straight cast to ansichar; since this is a character in hlsl source that's not in a comment
+			// we assume that it must be valid to do so. if this assumption is not valid the shader source was
+			// broken/corrupt anyways.
+			*CurrentOut++ = (ANSICHAR)(*Current++);
+		}
+
+		if (IsEndOfLine(*Current))
+		{
+			*CurrentOut++ = '\n';
+			Current += NewlineCharCount(Current[0], Current[1]);
+		}
+		else if (Current[0] == '/')
+		{
+			if (Current[1] == '/')
+			{
+				while (!IsEndOfLine(*Current) && Current < End)
+				{
+					++Current;
+				}
+			}
+			else if (Current[1] == '*')
+			{
+				Current += 2;
+				while (Current < End)
+				{
+					if (Current[0] == '*' && Current[1] == '/')
+					{
+						Current += 2;
+						break;
+					}
+					else if (IsEndOfLine(*Current))
+					{
+						*CurrentOut++ = '\n';
+						Current += NewlineCharCount(Current[0], Current[1]);
+					}
+					else
+					{
+						++Current;
+					}
+				}
+			}
+			else
+			{
+				*CurrentOut++ = (ANSICHAR)(*Current++);
+			}
+		}
+	}
+	// Null terminate after comment-stripped copy
+	check(CurrentOut < (OutStripped.GetData() + BufferSize));
+	*CurrentOut++ = 0;
+
+	// Set correct length after stripping but don't bother shrinking/reallocating, minor memory overhead to save time
+	OutStripped.SetNum(CurrentOut - OutStripped.GetData(), /* bAllowShrinking */false);
+}
+
+bool LoadShaderSourceFile(const TCHAR* InVirtualFilePath, EShaderPlatform ShaderPlatform, FString* OutFileContents, TArray<FShaderCompilerError>* OutCompileErrors, const FName* ShaderPlatformName, FShaderSharedAnsiStringPtr* OutStrippedContents) // TODO: const FString&
 {
 #if WITH_EDITORONLY_DATA
 	// it's not expected that cooked platforms get here, but if they do, this is the final out
@@ -1828,7 +1928,7 @@ bool LoadShaderSourceFile(const TCHAR* InVirtualFilePath, EShaderPlatform Shader
 		FString VirtualFilePath(InVirtualFilePath);
 		FixupShaderFilePath(VirtualFilePath, ShaderPlatform, ShaderPlatformName);
 
-		FString* CachedFile = nullptr;
+		FShaderFileCacheEntry* CachedFile = nullptr;
 
 		// First try a shared lock and only acquire exclusive access if element is not found in cache
 		uint32 CurrentHash = GetTypeHash(VirtualFilePath);
@@ -1842,7 +1942,11 @@ bool LoadShaderSourceFile(const TCHAR* InVirtualFilePath, EShaderPlatform Shader
 		{
 			if (OutFileContents)
 			{
-				*OutFileContents = *CachedFile;
+				*OutFileContents = CachedFile->Source;
+			}
+			if (OutStrippedContents)
+			{
+				*OutStrippedContents = CachedFile->StrippedSource;
 			}
 			bResult = true;
 		}
@@ -1858,7 +1962,11 @@ bool LoadShaderSourceFile(const TCHAR* InVirtualFilePath, EShaderPlatform Shader
 			{
 				if (OutFileContents)
 				{
-					*OutFileContents = *CachedFile;
+					*OutFileContents = CachedFile->Source;
+				}
+				if (OutStrippedContents)
+				{
+					*OutStrippedContents = CachedFile->StrippedSource;
 				}
 				bResult = true;
 			}
@@ -1867,15 +1975,23 @@ bool LoadShaderSourceFile(const TCHAR* InVirtualFilePath, EShaderPlatform Shader
 				FString ShaderFilePath = GetShaderSourceFilePath(VirtualFilePath, OutCompileErrors);
 
 				// verify SHA hash of shader files on load. missing entries trigger an error
-				FString FileContents;
-				if (!ShaderFilePath.IsEmpty() && FFileHelper::LoadFileToString(FileContents, *ShaderFilePath, FFileHelper::EHashOptions::EnableVerify|FFileHelper::EHashOptions::ErrorMissingHash) )
+				FShaderFileCacheEntry FileContents;
+				if (!ShaderFilePath.IsEmpty() && FFileHelper::LoadFileToString(FileContents.Source, *ShaderFilePath, FFileHelper::EHashOptions::EnableVerify|FFileHelper::EHashOptions::ErrorMissingHash) )
 				{
+					TArray<ANSICHAR>* StrippedSource = new TArray<ANSICHAR>;
+					ShaderConvertAndStripComments(FileContents.Source, *StrippedSource);
+					FileContents.StrippedSource = MakeShareable(StrippedSource);
+
 					//update the shader file cache
 					ShaderFileCache.Map.AddByHash(CurrentHash, VirtualFilePath, FileContents);
 
 					if (OutFileContents)
 					{
-						*OutFileContents = FileContents;
+						*OutFileContents = FileContents.Source;
+					}
+					if (OutStrippedContents)
+					{
+						*OutStrippedContents = FileContents.StrippedSource;
 					}
 					bResult = true;
 				}
@@ -3124,6 +3240,20 @@ FString FShaderCompileJobKey::ToString() const
 		PermutationId);
 }
 
+struct FShaderVirtualFileContents
+{
+	const FString* Wide;
+	const TArray<ANSICHAR>* Ansi;
+
+	FShaderVirtualFileContents(const FString* InWide)
+		: Wide(InWide), Ansi(nullptr)
+	{}
+
+	FShaderVirtualFileContents(const TArray<ANSICHAR>* InAnsi)
+		: Wide(nullptr), Ansi(InAnsi)
+	{}
+};
+
 FShaderCommonCompileJob::FInputHash FShaderCompileJob::GetInputHash()
 {
 	if (bInputHashSet)
@@ -3190,7 +3320,7 @@ FShaderCommonCompileJob::FInputHash FShaderCompileJob::GetInputHash()
 			// unroll the included files for the parallel processing.
 			// These are temporary arrays that only exist for the ParallelFor
 			TArray<const TCHAR*> IncludeVirtualPaths;
-			TArray<const FString*> Contents;
+			TArray<FShaderVirtualFileContents> Contents;
 			TArray<bool> OnlyHashIncludes;
 			TArray<FBlake3Hash> Hashes;
 
@@ -3204,7 +3334,7 @@ FShaderCommonCompileJob::FInputHash FShaderCompileJob::GetInputHash()
 				Hashes.AddDefaulted();
 			}
 
-			for (TMap<FString, FThreadSafeSharedStringPtr>::TConstIterator It(Input.Environment.IncludeVirtualPathToExternalContentsMap); It; ++It)
+			for (TMap<FString, FThreadSafeSharedAnsiStringPtr>::TConstIterator It(Input.Environment.IncludeVirtualPathToSharedContentsMap); It; ++It)
 			{
 				const FString& VirtualPath = It.Key();
 				IncludeVirtualPaths.Add(*VirtualPath);
@@ -3227,7 +3357,7 @@ FShaderCommonCompileJob::FInputHash FShaderCompileJob::GetInputHash()
 					Hashes.AddDefaulted();
 				}
 
-				for (TMap<FString, FThreadSafeSharedStringPtr>::TConstIterator It(Input.SharedEnvironment->IncludeVirtualPathToExternalContentsMap); It; ++It)
+				for (TMap<FString, FThreadSafeSharedAnsiStringPtr>::TConstIterator It(Input.SharedEnvironment->IncludeVirtualPathToSharedContentsMap); It; ++It)
 				{
 					const FString& VirtualPath = It.Key();
 					IncludeVirtualPaths.Add(*VirtualPath);
@@ -3246,7 +3376,19 @@ FShaderCommonCompileJob::FInputHash FShaderCompileJob::GetInputHash()
 			ParallelFor(Contents.Num(), [&IncludeVirtualPaths, &Contents, &OnlyHashIncludes, &Hashes, &Platform](int32 FileIndex)
 				{
 					FMemoryHasherBlake3 MemHasher;
-					HashShaderFileWithIncludes(MemHasher, IncludeVirtualPaths[FileIndex], *Contents[FileIndex], Platform, OnlyHashIncludes[FileIndex]);
+					if (Contents[FileIndex].Wide)
+					{
+						HashShaderFileWithIncludes(MemHasher, IncludeVirtualPaths[FileIndex], *Contents[FileIndex].Wide, Platform, OnlyHashIncludes[FileIndex]);
+					}
+					else
+					{
+						// ANSI files are shared uniform buffer struct declarations (or generated stereo code), and never have includes, so we just need to hash the
+						// single file contents.  Make sure that assumption hasn't been violated (this test costs less than 0.1% of GetInputHash, so might as well).
+						check(FCStringAnsi::Strstr(Contents[FileIndex].Ansi->GetData(), "#include") == nullptr);
+
+						MemHasher.Serialize(reinterpret_cast<void*>(const_cast<TCHAR*>(IncludeVirtualPaths[FileIndex])), FCString::Strlen(IncludeVirtualPaths[FileIndex]));
+						MemHasher << const_cast<TArray<ANSICHAR>&>(*Contents[FileIndex].Ansi);
+					}
 					Hashes[FileIndex] = MemHasher.Finalize();
 				},
 				EParallelForFlags::Unbalanced
diff --git a/Engine/Source/Runtime/RenderCore/Private/ShaderParameterMetadata.cpp b/Engine/Source/Runtime/RenderCore/Private/ShaderParameterMetadata.cpp
--- a/Engine/Source/Runtime/RenderCore/Private/ShaderParameterMetadata.cpp
+++ b/Engine/Source/Runtime/RenderCore/Private/ShaderParameterMetadata.cpp
@@ -8,6 +8,7 @@
 #include "RenderCore.h"
 #include "RHIUniformBufferLayoutInitializer.h"
 #include "ShaderCore.h"
+#include "ShaderCompilerCore.h"
 #include "ShaderParameters.h"
 #include "DataDrivenShaderPlatformInfo.h"
 #include "ShaderParameterMacros.h"
@@ -918,6 +919,11 @@ void FShaderParametersMetadata::InitializeUniformBufferDeclaration()
 
 		UniformBufferDeclaration = MakeShareable(NewDeclaration);
 
+		// Cache preprocessor friendly copy of uniform buffer declaration
+		TArray<ANSICHAR>* NewDeclarationAnsi = new TArray<ANSICHAR>;
+		ShaderConvertAndStripComments(*NewDeclaration, *NewDeclarationAnsi);
+		UniformBufferDeclarationAnsi = MakeShareable(NewDeclarationAnsi);
+
 		// Generate ResourceTableCache and MemberNameBuffer
 		FResourceTableEntryString Prefix;
 		Prefix.Append(ShaderVariableName);
diff --git a/Engine/Source/Runtime/RenderCore/Private/ShaderParameters.cpp b/Engine/Source/Runtime/RenderCore/Private/ShaderParameters.cpp
--- a/Engine/Source/Runtime/RenderCore/Private/ShaderParameters.cpp
+++ b/Engine/Source/Runtime/RenderCore/Private/ShaderParameters.cpp
@@ -454,14 +454,14 @@ void UE::ShaderParameters::AddUniformBufferIncludesToEnvironment(FShaderCompiler
 		{
 			if (const FShaderParametersMetadata* Metadata = FindShaderParametersMetadataWithVariableName(UniformBufferNameHash, UniformBufferNameView))
 			{
-				const FThreadSafeSharedStringPtr UniformBufferDeclaration = Metadata->GetUniformBufferDeclarationPtr();
+				const FThreadSafeSharedAnsiStringPtr UniformBufferDeclaration = Metadata->GetUniformBufferDeclarationAnsiPtr();
 
 				check(UniformBufferDeclaration.Get() != NULL);
 				check(!UniformBufferDeclaration.Get()->IsEmpty());
 
 				UniformBufferIncludes += Metadata->GetUniformBufferInclude();
 
-				OutEnvironment.IncludeVirtualPathToExternalContentsMap.AddByHash(Metadata->GetUniformBufferPathHash(), Metadata->GetUniformBufferPath(), UniformBufferDeclaration);
+				OutEnvironment.IncludeVirtualPathToSharedContentsMap.AddByHash(Metadata->GetUniformBufferPathHash(), Metadata->GetUniformBufferPath(), UniformBufferDeclaration);
 
 				Metadata->AddResourceTableEntries(OutEnvironment.ResourceTableMap, OutEnvironment.UniformBufferMap);
 			}
diff --git a/Engine/Source/Runtime/RenderCore/Private/ShaderPreprocessTypes.cpp b/Engine/Source/Runtime/RenderCore/Private/ShaderPreprocessTypes.cpp
--- a/Engine/Source/Runtime/RenderCore/Private/ShaderPreprocessTypes.cpp
+++ b/Engine/Source/Runtime/RenderCore/Private/ShaderPreprocessTypes.cpp
@@ -152,6 +152,9 @@ void FShaderPreprocessOutput::StripCode()
 {
 	// Reserve worst case slack (i.e. assuming there is nothing to strip) to avoid reallocation
 	FString PreprocessedSourceStripped(LineDirectiveSentinel, PreprocessedSource.Len() + 1);
+	PreprocessedSourceStripped.GetCharArray().AddUninitialized(PreprocessedSourceStripped.GetCharArray().GetSlack());
+	TCHAR* OutStrippedData = PreprocessedSourceStripped.GetCharArray().GetData();
+	TCHAR* OutStripped = OutStrippedData + LineDirectiveSentinel.Len();
 
 	const TCHAR* Begin = PreprocessedSource.GetCharArray().GetData(), *Current = Begin;
 	const TCHAR* End = Current + PreprocessedSource.Len();
@@ -161,21 +164,21 @@ void FShaderPreprocessOutput::StripCode()
 	{
 		while (!StripNeedsHandling(*Current))
 		{
-			PreprocessedSourceStripped.AppendChar(*Current++);
+			*OutStripped++ = *Current++;
 		}
 
 		if (IsEndOfLine(*Current))
 		{
 			// only emit \n if it wasn't preceded immediately by another linebreak 
 			// (i.e. skip empty lines)
-			if (PreprocessedSourceStripped.Len() && !IsEndOfLine(PreprocessedSourceStripped[PreprocessedSourceStripped.Len() - 1]))
+			if (!IsEndOfLine(*(OutStripped - 1)))
 			{
 				// Record the offset from the start of the block given by the last line directive for each line
 				// output in the stripped code. 
 				Remapper.AddStrippedLine(CurrentStrippedLineNum, CurrentBlockUnstrippedLineOffset);
 
 				// normalize line endings
-				PreprocessedSourceStripped.AppendChar(TEXT('\n'));
+				*OutStripped++ = TEXT('\n');
 				CurrentStrippedLineNum++;
 			}
 			CurrentBlockUnstrippedLineOffset++;
@@ -201,7 +204,7 @@ void FShaderPreprocessOutput::StripCode()
 			}
 			else
 			{
-				PreprocessedSourceStripped.AppendChar(*Current++);
+				*OutStripped++ = *Current++;
 			}
 		}
 		else if (Current[0] == TEXT('#'))
@@ -268,12 +271,14 @@ void FShaderPreprocessOutput::StripCode()
 			}
 			else
 			{
-				PreprocessedSourceStripped.AppendChar(*Current++);
+				*OutStripped++ = *Current++;
 			}
 		}
 	}
 	// Null terminate after stripped copy
-	PreprocessedSourceStripped.AppendChar(0);
+	*OutStripped++ = 0;
+	check(OutStripped <= OutStrippedData + PreprocessedSourceStripped.GetCharArray().Num());
+	PreprocessedSourceStripped.GetCharArray().SetNum((int32)(OutStripped - OutStrippedData));
 
 	OriginalPreprocessedSource = MoveTemp(PreprocessedSource);
 	PreprocessedSource = MoveTemp(PreprocessedSourceStripped);
diff --git a/Engine/Source/Runtime/RenderCore/Public/ShaderCompilerCore.h b/Engine/Source/Runtime/RenderCore/Public/ShaderCompilerCore.h
--- a/Engine/Source/Runtime/RenderCore/Public/ShaderCompilerCore.h
+++ b/Engine/Source/Runtime/RenderCore/Public/ShaderCompilerCore.h
@@ -27,9 +27,10 @@ class IShaderFormat;
 class FShaderCommonCompileJob;
 class FShaderCompileJob;
 class FShaderPipelineCompileJob;
+typedef TSharedPtr<TArray<ANSICHAR>, ESPMode::ThreadSafe> FShaderSharedAnsiStringPtr;
 
 // this is for the protocol, not the data, bump if FShaderCompilerInput or ProcessInputFromArchive changes.
-inline const int32 ShaderCompileWorkerInputVersion = 21;
+inline const int32 ShaderCompileWorkerInputVersion = 22;
 // this is for the protocol, not the data, bump if FShaderCompilerOutput or WriteToOutputArchive changes.
 inline const int32 ShaderCompileWorkerOutputVersion = 18;
 // this is for the protocol, not the data.
@@ -356,15 +357,47 @@ struct FShaderCompilerInput
 		TArray<TRefCountPtr<FSharedShaderCompilerEnvironment>>& SharedEnvironments,
 		TArray<const FShaderParametersMetadata*>& ParametersStructures)
 	{
-		check(!SharedEnvironment || SharedEnvironment->IncludeVirtualPathToExternalContentsMap.Num() == 0);
+		check(!SharedEnvironment || SharedEnvironment->IncludeVirtualPathToSharedContentsMap.Num() == 0);
 
 		// If the input is already preprocessed we don't need to serialize includes when writing worker input files
 		if (!bCachePreprocessed)
 		{
-			for (const auto& It : Environment.IncludeVirtualPathToExternalContentsMap)
+			for (const auto& It : Environment.IncludeVirtualPathToSharedContentsMap)
 			{
 				FString* FoundEntry = ExternalIncludes.Find(It.Key);
 
+				if (!FoundEntry)
+				{
+					ExternalIncludes.Add(It.Key, FString(*It.Value));
+				}
+			}
+		}
+
+		if (SharedEnvironment)
+		{
+			SharedEnvironments.AddUnique(SharedEnvironment);
+		}
+
+		if (RootParametersStructure)
+		{
+			ParametersStructures.AddUnique(RootParametersStructure);
+		}
+	}
+
+	void GatherSharedInputsAnsi(
+		TMap<FString, TArray<ANSICHAR>>& ExternalIncludes,
+		TArray<TRefCountPtr<FSharedShaderCompilerEnvironment>>& SharedEnvironments,
+		TArray<const FShaderParametersMetadata*>& ParametersStructures)
+	{
+		check(!SharedEnvironment || SharedEnvironment->IncludeVirtualPathToSharedContentsMap.Num() == 0);
+
+		// If the input is already preprocessed we don't need to serialize includes when writing worker input files
+		if (!bCachePreprocessed)
+		{
+			for (const auto& It : Environment.IncludeVirtualPathToSharedContentsMap)
+			{
+				TArray<ANSICHAR>* FoundEntry = ExternalIncludes.Find(It.Key);
+
 				if (!FoundEntry)
 				{
 					ExternalIncludes.Add(It.Key, *It.Value);
@@ -390,9 +423,9 @@ struct FShaderCompilerInput
 		if (!bCachePreprocessed)
 		{
 			TArray<FString> ReferencedExternalIncludes;
-			ReferencedExternalIncludes.Empty(Environment.IncludeVirtualPathToExternalContentsMap.Num());
+			ReferencedExternalIncludes.Empty(Environment.IncludeVirtualPathToSharedContentsMap.Num());
 
-			for (const auto& It : Environment.IncludeVirtualPathToExternalContentsMap)
+			for (const auto& It : Environment.IncludeVirtualPathToSharedContentsMap)
 			{
 				ReferencedExternalIncludes.Add(It.Key);
 			}
@@ -414,7 +447,7 @@ struct FShaderCompilerInput
 
 	void DeserializeSharedInputs(
 		FArchive& Ar,
-		const TMap<FString, FThreadSafeSharedStringPtr>& ExternalIncludes,
+		const TMap<FString, FThreadSafeSharedAnsiStringPtr>& ExternalIncludes,
 		const TArray<FShaderCompilerEnvironment>& SharedEnvironments,
 		const TArray<TUniquePtr<FShaderParametersMetadata>>& ShaderParameterStructures)
 	{
@@ -425,11 +458,11 @@ struct FShaderCompilerInput
 			TArray<FString> ReferencedExternalIncludes;
 			Ar << ReferencedExternalIncludes;
 
-			Environment.IncludeVirtualPathToExternalContentsMap.Reserve(ReferencedExternalIncludes.Num());
+			Environment.IncludeVirtualPathToSharedContentsMap.Reserve(ReferencedExternalIncludes.Num());
 
 			for (int32 i = 0; i < ReferencedExternalIncludes.Num(); i++)
 			{
-				Environment.IncludeVirtualPathToExternalContentsMap.Add(ReferencedExternalIncludes[i], ExternalIncludes.FindChecked(ReferencedExternalIncludes[i]));
+				Environment.IncludeVirtualPathToSharedContentsMap.Add(ReferencedExternalIncludes[i], ExternalIncludes.FindChecked(ReferencedExternalIncludes[i]));
 			}
 		}
 
@@ -731,13 +764,18 @@ extern RENDERCORE_API bool CheckVirtualShaderFilePath(FStringView VirtualPath, T
  */
 extern RENDERCORE_API void FixupShaderFilePath(FString& VirtualFilePath, EShaderPlatform ShaderPlatform, const FName* ShaderPlatformName);
 
+/**
+ * Utility function to strip comments and convert source to ANSI, useful for preprocessing
+ */
+extern RENDERCORE_API void ShaderConvertAndStripComments(const FString& ShaderSource, TArray<ANSICHAR>& OutStripped);
+
 /**
  * Loads the shader file with the given name.
  * @param VirtualFilePath - The virtual path of shader file to load.
  * @param OutFileContents - If true is returned, will contain the contents of the shader file. Can be null.
  * @return True if the file was successfully loaded.
  */
-extern RENDERCORE_API bool LoadShaderSourceFile(const TCHAR* VirtualFilePath, EShaderPlatform ShaderPlatform, FString* OutFileContents, TArray<FShaderCompilerError>* OutCompileErrors, const FName* ShaderPlatformName = nullptr);
+extern RENDERCORE_API bool LoadShaderSourceFile(const TCHAR* VirtualFilePath, EShaderPlatform ShaderPlatform, FString* OutFileContents, TArray<FShaderCompilerError>* OutCompileErrors, const FName* ShaderPlatformName = nullptr, FShaderSharedAnsiStringPtr* OutStrippedContents = nullptr);
 
 enum class EShaderCompilerWorkerType : uint8
 {
diff --git a/Engine/Source/Runtime/RenderCore/Public/ShaderCore.h b/Engine/Source/Runtime/RenderCore/Public/ShaderCore.h
--- a/Engine/Source/Runtime/RenderCore/Public/ShaderCore.h
+++ b/Engine/Source/Runtime/RenderCore/Public/ShaderCore.h
@@ -458,7 +458,10 @@ struct FShaderCompilerEnvironment
 	// Map of the virtual file path -> content.
 	// The virtual file paths are the ones that USF files query through the #include "<The Virtual Path of the file>"
 	TMap<FString,FString> IncludeVirtualPathToContentsMap;
-	
+
+	TMap<FString, FThreadSafeSharedAnsiStringPtr> IncludeVirtualPathToSharedContentsMap;
+
+	UE_DEPRECATED(5.4, "IncludeVirtualPathToExternalContentsMap has been replaced with IncludeVirtualPathToSharedContentsMap (type change from FString to ANSI string).")
 	TMap<FString, FThreadSafeSharedStringPtr> IncludeVirtualPathToExternalContentsMap;
 
 	FShaderCompilerFlags CompilerFlags;
@@ -624,7 +627,7 @@ struct FShaderCompilerEnvironment
 		// Note: this serialize is used to pass between UE and the shader compile worker, recompile both when modifying
 		Ar << Environment.IncludeVirtualPathToContentsMap;
 
-		// Note: skipping Environment.IncludeVirtualPathToExternalContentsMap, which is handled by FShaderCompileUtilities::DoWriteTasks in order to maintain sharing
+		// Note: skipping Environment.IncludeVirtualPathToSharedContentsMap, which is handled by FShaderCompileUtilities::DoWriteTasks in order to maintain sharing
 
 		Environment.SerializeEverythingButFiles(Ar);
 		return Ar;
@@ -1209,9 +1212,10 @@ RENDERCORE_API void UpdateReferencedUniformBufferNames(
 	TArrayView<const FVertexFactoryType*> OutdatedFactoryTypes,
 	TArrayView<const FShaderPipelineType*> OutdatedShaderPipelineTypes);
 
+/** Deprecated structure (only used by deprecated "CacheUniformBufferIncludes" and "SerializeUniformBufferInfo") */
 struct FCachedUniformBufferDeclaration
 {
-	// Using SharedPtr so we can hand off lifetime ownership to FShaderCompilerEnvironment::IncludeVirtualPathToExternalContentsMap when invalidating this cache
+	// Using SharedPtr so we can hand off lifetime ownership to FShaderCompilerEnvironment::IncludeVirtualPathToSharedContentsMap when invalidating this cache
 	FThreadSafeSharedStringPtr Declaration;
 };
 
diff --git a/Engine/Source/Runtime/RenderCore/Public/ShaderParameterMetadata.h b/Engine/Source/Runtime/RenderCore/Public/ShaderParameterMetadata.h
--- a/Engine/Source/Runtime/RenderCore/Public/ShaderParameterMetadata.h
+++ b/Engine/Source/Runtime/RenderCore/Public/ShaderParameterMetadata.h
@@ -26,6 +26,7 @@
 #include "Templates/SharedPointer.h"
 
 using FThreadSafeSharedStringPtr = TSharedPtr<FString, ESPMode::ThreadSafe>;
+using FThreadSafeSharedAnsiStringPtr = TSharedPtr<TArray<ANSICHAR>, ESPMode::ThreadSafe>;
 using FThreadSafeNameBufferPtr = TSharedPtr<TArray<TCHAR>, ESPMode::ThreadSafe>;
 struct FShaderResourceTableMap;
 
@@ -353,6 +354,7 @@ public:
 #if WITH_EDITOR
 	inline bool IsUniformBufferDeclarationInitialized() const { return UniformBufferDeclaration.IsValid(); }
 	FThreadSafeSharedStringPtr GetUniformBufferDeclarationPtr() const { return UniformBufferDeclaration; }
+	FThreadSafeSharedAnsiStringPtr GetUniformBufferDeclarationAnsiPtr() const { return UniformBufferDeclarationAnsi; }
 	const FString& GetUniformBufferDeclaration() const { return *UniformBufferDeclaration; }
 	FORCEINLINE const FString& GetUniformBufferPath() const { return UniformBufferPath; }
 	FORCEINLINE const FString& GetUniformBufferInclude() const { return UniformBufferInclude; }
@@ -442,6 +444,7 @@ private:
 #if WITH_EDITOR
 	/** Uniform buffer declaration, created once */
 	FThreadSafeSharedStringPtr UniformBufferDeclaration;
+	FThreadSafeSharedAnsiStringPtr UniformBufferDeclarationAnsi;
 
 	/** Cache of uniform buffer resource table, and storage for member names used by the table, created once */
 	TArray<FUniformResourceEntry> ResourceTableCache;
