Shader Compiler:  Fix for unity stress test compile error.  Give an inline utility function an arbitrarily different name from what it is in another C++ file that may be unity included together.

#rnx

diff --git a/Engine/Source/Runtime/RenderCore/Private/ShaderPreprocessTypes.cpp b/Engine/Source/Runtime/RenderCore/Private/ShaderPreprocessTypes.cpp
--- a/Engine/Source/Runtime/RenderCore/Private/ShaderPreprocessTypes.cpp
+++ b/Engine/Source/Runtime/RenderCore/Private/ShaderPreprocessTypes.cpp
@@ -131,14 +131,14 @@ FShaderDiagnosticRemapper::FRemapData FShaderDiagnosticRemapper::GetRemapData(in
 	return FRemapData{ FoundBlock.OriginalPath, FoundBlock.OriginalLineNum + StrippedLineOffsets[StrippedLineNum - 1] };
 }
 
-inline bool IsEndOfLine(TCHAR C)
+inline bool IsEndOfTheLine(TCHAR C)
 {
 	return C == TEXT('\r') || C == TEXT('\n');
 }
 
 inline bool StripNeedsHandling(TCHAR C)
 {
-	return IsEndOfLine(C) || C == TEXT('/') || C == 0 || C == TEXT('#');
+	return IsEndOfTheLine(C) || C == TEXT('/') || C == 0 || C == TEXT('#');
 }
 
 inline void SkipNewLine(const TCHAR*& Current, const TCHAR* End)
@@ -167,11 +167,11 @@ void FShaderPreprocessOutput::StripCode()
 			*OutStripped++ = *Current++;
 		}
 
-		if (IsEndOfLine(*Current))
+		if (IsEndOfTheLine(*Current))
 		{
 			// only emit \n if it wasn't preceded immediately by another linebreak 
 			// (i.e. skip empty lines)
-			if (!IsEndOfLine(*(OutStripped - 1)))
+			if (!IsEndOfTheLine(*(OutStripped - 1)))
 			{
 				// Record the offset from the start of the block given by the last line directive for each line
 				// output in the stripped code. 
@@ -188,7 +188,7 @@ void FShaderPreprocessOutput::StripCode()
 		{
 			if (Current[1] == TEXT('/'))
 			{
-				while (!IsEndOfLine(*Current) && Current < End)
+				while (!IsEndOfTheLine(*Current) && Current < End)
 				{
 					++Current;
 				}
@@ -230,7 +230,7 @@ void FShaderPreprocessOutput::StripCode()
 				}
 
 				// scan past open quote
-				while (*Current != '\"' && Current < End && !IsEndOfLine(*Current))
+				while (*Current != '\"' && Current < End && !IsEndOfTheLine(*Current))
 				{
 					++Current;
 				}
@@ -254,7 +254,7 @@ void FShaderPreprocessOutput::StripCode()
 					// scan to end-of-line and skip past the newline; this would be handled by the newline case above as well,
 					// but we don't want the newline at the end of the line directive to count in our calculated offsets for
 					// emitted stripped code
-					while (!IsEndOfLine(*Current) && Current < End)
+					while (!IsEndOfTheLine(*Current) && Current < End)
 					{
 						++Current;
 					}
