Add SupportsResourceType to transient allocation because certain RHIs might not support all resources types

diff --git a/Engine/Source/Runtime/D3D12RHI/Private/D3D12TransientResourceAllocator.h b/Engine/Source/Runtime/D3D12RHI/Private/D3D12TransientResourceAllocator.h
--- a/Engine/Source/Runtime/D3D12RHI/Private/D3D12TransientResourceAllocator.h
+++ b/Engine/Source/Runtime/D3D12RHI/Private/D3D12TransientResourceAllocator.h
@@ -46,6 +46,15 @@ public:
 	FD3D12TransientResourceHeapAllocator(FD3D12TransientHeapCache& InHeapCache);
 
 	//! IRHITransientResourceAllocator Overrides
+	bool SupportsResourceType(ERHITransientResourceType InType) const override 
+	{ 
+		switch (InType)
+		{
+		case ERHITransientResourceType::Buffer: return true;
+		case ERHITransientResourceType::Texture: return true;
+		default: checkNoEntry(); return false;
+		}
+	}
 	FRHITransientTexture* CreateTexture(const FRHITextureCreateInfo& InCreateInfo, const TCHAR* InDebugName, uint32 InPassIndex) override;
 	FRHITransientBuffer* CreateBuffer(const FRHIBufferCreateInfo& InCreateInfo, const TCHAR* InDebugName, uint32 InPassIndex) override;
 
diff --git a/Engine/Source/Runtime/RHI/Public/RHITransientResourceAllocator.h b/Engine/Source/Runtime/RHI/Public/RHITransientResourceAllocator.h
--- a/Engine/Source/Runtime/RHI/Public/RHITransientResourceAllocator.h
+++ b/Engine/Source/Runtime/RHI/Public/RHITransientResourceAllocator.h
@@ -354,6 +354,9 @@ class IRHITransientResourceAllocator
 public:
 	virtual ~IRHITransientResourceAllocator() = default;
 
+	// Supports transient allocations of given resource type
+	virtual bool SupportsResourceType(ERHITransientResourceType InType) const = 0;
+
 	// Allocates a new transient resource with memory backed by the transient allocator.
 	virtual FRHITransientTexture* CreateTexture(const FRHITextureCreateInfo& InCreateInfo, const TCHAR* InDebugName, uint32 InPassIndex) = 0;
 	virtual FRHITransientBuffer* CreateBuffer(const FRHIBufferCreateInfo& InCreateInfo, const TCHAR* InDebugName, uint32 InPassIndex) = 0;
diff --git a/Engine/Source/Runtime/RHI/Public/RHIValidationTransientResourceAllocator.h b/Engine/Source/Runtime/RHI/Public/RHIValidationTransientResourceAllocator.h
--- a/Engine/Source/Runtime/RHI/Public/RHIValidationTransientResourceAllocator.h
+++ b/Engine/Source/Runtime/RHI/Public/RHIValidationTransientResourceAllocator.h
@@ -16,6 +16,7 @@ public:
 	virtual ~FValidationTransientResourceAllocator();
 
 	// Implementation of FRHITransientResourceAllocator interface
+	virtual bool SupportsResourceType(ERHITransientResourceType InType) const override final { return RHIAllocator->SupportsResourceType(InType); }
 	virtual FRHITransientTexture* CreateTexture(const FRHITextureCreateInfo& InCreateInfo, const TCHAR* InDebugName, uint32 InPassIndex) override final;
 	virtual FRHITransientBuffer* CreateBuffer(const FRHIBufferCreateInfo& InCreateInfo, const TCHAR* InDebugName, uint32 InPassIndex) override final;
 	virtual void DeallocateMemory(FRHITransientTexture* InTexture, uint32 InPassIndex) override final;
diff --git a/Engine/Source/Runtime/RenderCore/Private/RenderGraphBuilder.cpp b/Engine/Source/Runtime/RenderCore/Private/RenderGraphBuilder.cpp
--- a/Engine/Source/Runtime/RenderCore/Private/RenderGraphBuilder.cpp
+++ b/Engine/Source/Runtime/RenderCore/Private/RenderGraphBuilder.cpp
@@ -581,6 +581,9 @@ FRDGBuilder::FRDGBuilder(FRHICommandListImmediate& InRHICmdList, FRDGEventName I
 {
 	AddProloguePass();
 
+	bSupportsTransientTextures = TransientResourceAllocator ? TransientResourceAllocator->SupportsResourceType(ERHITransientResourceType::Texture) : false;
+	bSupportsTransientBuffers = TransientResourceAllocator ? TransientResourceAllocator->SupportsResourceType(ERHITransientResourceType::Buffer) : false;
+
 #if RDG_EVENTS != RDG_EVENTS_NONE
 	// This is polled once as a workaround for a race condition since the underlying global is not always changed on the render thread.
 	GRDGEmitDrawEvents_RenderThread = GetEmitDrawEvents();
@@ -982,7 +985,7 @@ FRDGTexture* FRDGBuilder::RegisterExternalTexture(
 	if (FRDGTexture* FoundTexture = FindExternalTexture(ExternalTextureRHI))
 	{
 		return FoundTexture;
-	}
+	} 
 
 	const FRDGTextureDesc Desc = Translate(ExternalPooledTexture->GetDesc());
 	FRDGTexture* Texture = Textures.Allocate(Allocator, Name, Desc, Flags);
@@ -3589,7 +3592,7 @@ void FRDGBuilder::BeginResourceRHI(FRDGPassHandle PassHandle, FRDGTextureRef Tex
 	}
 #endif
 
-	if (TransientResourceAllocator && IsTransient(Texture))
+	if (TransientResourceAllocator && bSupportsTransientTextures && IsTransient(Texture))
 	{
 		if (FRHITransientTexture* TransientTexture = TransientResourceAllocator->CreateTexture(Texture->Desc, Texture->Name, PassHandle.GetIndex()))
 		{
@@ -3681,7 +3684,7 @@ void FRDGBuilder::BeginResourceRHI(FRDGPassHandle PassHandle, FRDGBufferRef Buff
 	Buffer->FinalizeDesc();
 
 	// If transient then create the resource on the transient allocator. External or extracted resource can't be transient because of lifetime tracking issues.
-	if (TransientResourceAllocator && IsTransient(Buffer))
+	if (TransientResourceAllocator && bSupportsTransientBuffers && IsTransient(Buffer))
 	{
 		if (FRHITransientBuffer* TransientBuffer = TransientResourceAllocator->CreateBuffer(Translate(Buffer->Desc), Buffer->Name, PassHandle.GetIndex()))
 		{
diff --git a/Engine/Source/Runtime/RenderCore/Public/RenderGraphBuilder.h b/Engine/Source/Runtime/RenderCore/Public/RenderGraphBuilder.h
--- a/Engine/Source/Runtime/RenderCore/Public/RenderGraphBuilder.h
+++ b/Engine/Source/Runtime/RenderCore/Public/RenderGraphBuilder.h
@@ -735,7 +735,9 @@ private:
 	IF_RDG_CMDLIST_STATS(TStatId CommandListStatScope);
 	IF_RDG_CMDLIST_STATS(TStatId CommandListStatState);
 
-	IRHITransientResourceAllocator* TransientResourceAllocator = nullptr;
+	IRHITransientResourceAllocator* TransientResourceAllocator = nullptr; 
+	bool bSupportsTransientTextures = false;
+	bool bSupportsTransientBuffers = false;
 
 	FRHICommandListScopedExtendResourceLifetime ExtendResourceLifetimeScope;
 
diff --git a/Engine/Source/Runtime/VulkanRHI/Private/VulkanTransientResourceAllocator.h b/Engine/Source/Runtime/VulkanRHI/Private/VulkanTransientResourceAllocator.h
--- a/Engine/Source/Runtime/VulkanRHI/Private/VulkanTransientResourceAllocator.h
+++ b/Engine/Source/Runtime/VulkanRHI/Private/VulkanTransientResourceAllocator.h
@@ -47,6 +47,15 @@ public:
 	FVulkanTransientResourceAllocator(FVulkanTransientHeapCache& InHeapCache);
 
 	//! IRHITransientResourceAllocator Overrides
+	bool SupportsResourceType(ERHITransientResourceType InType) const override
+	{
+		switch (InType)
+		{
+		case ERHITransientResourceType::Buffer: return true;
+		case ERHITransientResourceType::Texture: return true;
+		default: checkNoEntry(); return false;
+		}
+	}
 	FRHITransientTexture* CreateTexture(const FRHITextureCreateInfo& InCreateInfo, const TCHAR* InDebugName, uint32 InPassIndex) override;
 	FRHITransientBuffer* CreateBuffer(const FRHIBufferCreateInfo& InCreateInfo, const TCHAR* InDebugName, uint32 InPassIndex) override;
 };
\ No newline at end of file
