Fixes r.VisualizeTexture.AllowBlinking

#tests win64

diff --git a/Engine/Source/Runtime/RenderCore/Private/VisualizeTexture.cpp b/Engine/Source/Runtime/RenderCore/Private/VisualizeTexture.cpp
--- a/Engine/Source/Runtime/RenderCore/Private/VisualizeTexture.cpp
+++ b/Engine/Source/Runtime/RenderCore/Private/VisualizeTexture.cpp
@@ -675,7 +675,7 @@ FRDGTextureRef FVisualizeTexture::AddVisualizeTexturePass(
 
 				// w * almost_1 to avoid frac(1) => 0
 				PassParameters->VisualizeParam[0] = FVector4f(VisualizeConfig.RGBMul, VisualizeConfig.SingleChannelMul, Add, FracScale * 0.9999f);
-				PassParameters->VisualizeParam[1] = FVector4f(CVarAllowBlinking.GetValueOnRenderThread() ? BlinkState : 1.0f, (VisualizeConfig.ShaderOp == EShaderOp::Saturate) ? 1.0f : 0.0f, VisualizeConfig.ArrayIndex, VisualizeConfig.MipIndex);
+				PassParameters->VisualizeParam[1] = FVector4f(CVarAllowBlinking.GetValueOnRenderThread() ? BlinkState : 0.0f, (VisualizeConfig.ShaderOp == EShaderOp::Saturate) ? 1.0f : 0.0f, VisualizeConfig.ArrayIndex, VisualizeConfig.MipIndex);
 				PassParameters->VisualizeParam[2] = FVector4f((float)InputValueMapping, 0.0f, VisualizeConfig.SingleChannel);
 			}
 
