Fixed assert in RDG task queue flush due to waiting on task while in a task pipe.
 - Converted lock free list to locked array.

diff --git a/Engine/Source/Runtime/RenderCore/Private/RenderGraphBuilder.cpp b/Engine/Source/Runtime/RenderCore/Private/RenderGraphBuilder.cpp
--- a/Engine/Source/Runtime/RenderCore/Private/RenderGraphBuilder.cpp
+++ b/Engine/Source/Runtime/RenderCore/Private/RenderGraphBuilder.cpp
@@ -1648,18 +1648,20 @@ void FRDGBuilder::Compile()
 
 ///////////////////////////////////////////////////////////////////////////////////////////////////
 
+void FRDGBuilder::FPassQueue::Push(FRDGPass* Pass)
+{
+	UE::TScopeLock Lock(Mutex);
+	Queue.Emplace(Pass);
+}
+
 template <typename LambdaType>
 void FRDGBuilder::FPassQueue::Flush(UE::Tasks::FPipe& Pipe, const TCHAR* Name, LambdaType&& Lambda)
 {
 	if (LastTask.IsCompleted())
 	{
-		LastTask = Pipe.Launch(Name, [this, Lambda, Name]
+		LastTask = Pipe.Launch(Name, [this, Lambda = MoveTemp(Lambda), Name] () mutable
 		{
-			SCOPED_NAMED_EVENT_TCHAR(Name, FColor::Magenta);
-			while (FRDGPass* Pass = Queue.Pop())
-			{
-				Lambda(Pass);
-			}
+			Flush(Name, MoveTemp(Lambda));
 		});
 	}
 }
@@ -1668,8 +1670,12 @@ template <typename LambdaType>
 void FRDGBuilder::FPassQueue::Flush(const TCHAR* Name, LambdaType&& Lambda)
 {
 	SCOPED_NAMED_EVENT_TCHAR(Name, FColor::Magenta);
-	LastTask.Wait();
-	while (FRDGPass* Pass = Queue.Pop())
+
+	Mutex.Lock();
+	TArray<FRDGPass*, FRDGArrayAllocator> PoppedQueue = MoveTemp(Queue);
+	Mutex.Unlock();
+
+	for (FRDGPass* Pass : PoppedQueue)
 	{
 		Lambda(Pass);
 	}
diff --git a/Engine/Source/Runtime/RenderCore/Public/RenderGraphBuilder.h b/Engine/Source/Runtime/RenderCore/Public/RenderGraphBuilder.h
--- a/Engine/Source/Runtime/RenderCore/Public/RenderGraphBuilder.h
+++ b/Engine/Source/Runtime/RenderCore/Public/RenderGraphBuilder.h
@@ -548,10 +548,7 @@ private:
 	class FPassQueue
 	{
 	public:
-		void Push(FRDGPass* Pass)
-		{
-			Queue.Push(Pass);
-		}
+		void Push(FRDGPass* Pass);
 
 		template <typename LambdaType>
 		void Flush(UE::Tasks::FPipe& Pipe, const TCHAR* Name, LambdaType&& Lambda);
@@ -560,7 +557,8 @@ private:
 		void Flush(const TCHAR* Name, LambdaType&& Lambda);
 
 	private:
-		TLockFreePointerListFIFO<FRDGPass, PLATFORM_CACHE_LINE_SIZE> Queue;
+		UE::FMutex Mutex;
+		TArray<FRDGPass*, FRDGArrayAllocator> Queue;
 		UE::Tasks::FTask LastTask;
 	};
 
