Adding basic comment skipping support to ShaderParameterParser.

diff --git a/Engine/Source/Runtime/RenderCore/Private/ShaderParameterParser.cpp b/Engine/Source/Runtime/RenderCore/Private/ShaderParameterParser.cpp
--- a/Engine/Source/Runtime/RenderCore/Private/ShaderParameterParser.cpp
+++ b/Engine/Source/Runtime/RenderCore/Private/ShaderParameterParser.cpp
@@ -337,6 +337,7 @@ bool FShaderParameterParser::ParseParameters(
 
 		EState State = EState::Scanning;
 		bool bGoToNextLine = false;
+		bool bGoToCommentClose = false;
 
 		auto ResetState = [&]()
 		{
@@ -502,6 +503,15 @@ bool FShaderParameterParser::ParseParameters(
 				}
 				continue;
 			}
+			else if (bGoToCommentClose)
+			{
+				if (Char == '*' && UpComing[1] == '/')
+				{
+					Cursor++;
+					bGoToCommentClose = false;
+				}
+				continue;
+			}
 			else if (Char == '#')
 			{
 				if (RemainingSize > 6 && FCString::Strncmp(UpComing, TEXT("#line "), 6) == 0)
@@ -607,6 +617,20 @@ bool FShaderParameterParser::ParseParameters(
 				{
 					// Looks like redundant semicolon, just ignore and keep scanning.
 				}
+				else if (Char == '/')
+				{
+					if (UpComing[1] == '/')
+					{
+						bGoToNextLine = true;
+						continue;
+					}
+
+					if (UpComing[1] == '*')
+					{
+						bGoToCommentClose = true;
+						continue;
+					}
+				}
 				else
 				{
 					// No idea what this is, just go to next semi colon.
