[shaders] deprecating old PreprocessShader and IShaderFormat::CompileShader APIs now that all shader formats have been migrated

diff --git a/Engine/Source/Developer/ShaderPreprocessor/Public/ShaderPreprocessor.h b/Engine/Source/Developer/ShaderPreprocessor/Public/ShaderPreprocessor.h
--- a/Engine/Source/Developer/ShaderPreprocessor/Public/ShaderPreprocessor.h
+++ b/Engine/Source/Developer/ShaderPreprocessor/Public/ShaderPreprocessor.h
@@ -34,6 +34,7 @@ enum class EDumpShaderDefines : uint8
  * @param bShaderDumpDefinesAsCommentedCode - Whether to add shader definitions as comments.
  * @returns true if the shader is preprocessed without error.
  */
+UE_DEPRECATED(5.4, "Please use overload of PreprocessShader accepting a FShaderPreprocessOutput struct.")
 extern SHADERPREPROCESSOR_API bool PreprocessShader(
 	FString& OutPreprocessedShader,
 	FShaderCompilerOutput& ShaderOutput,
diff --git a/Engine/Source/Developer/TargetPlatform/Public/Interfaces/IShaderFormat.h b/Engine/Source/Developer/TargetPlatform/Public/Interfaces/IShaderFormat.h
--- a/Engine/Source/Developer/TargetPlatform/Public/Interfaces/IShaderFormat.h
+++ b/Engine/Source/Developer/TargetPlatform/Public/Interfaces/IShaderFormat.h
@@ -28,6 +28,7 @@ public:
 	 * @param Output The output from shader compiler.
 	 * @param WorkingDirectory The working directory.
 	 */
+	UE_DEPRECATED(5.4, "Monolithic CompileShader is deprecated. Please implement separate compilation and preprocess via PreprocessShader/CompilePreprocessedShader")
 	virtual void CompileShader(FName Format, const FShaderCompilerInput& Input, FShaderCompilerOutput& Output, const FString& WorkingDirectory) const {};
 
 	/**
diff --git a/Engine/Source/Runtime/RenderCore/Private/ShaderCore.cpp b/Engine/Source/Runtime/RenderCore/Private/ShaderCore.cpp
--- a/Engine/Source/Runtime/RenderCore/Private/ShaderCore.cpp
+++ b/Engine/Source/Runtime/RenderCore/Private/ShaderCore.cpp
@@ -1305,14 +1305,7 @@ public:
 				check(Job.SecondaryPreprocessOutput.IsValid());
 				Compiler->CompilePreprocessedShader(Job.Input, Job.PreprocessOutput, *Job.SecondaryPreprocessOutput, Job.Output, *Job.SecondaryOutput, WorkingDirectory);
 			}
-			else if (Compiler->SupportsIndependentPreprocessing())
-			{
-				Compiler->CompilePreprocessedShader(Job.Input, Job.PreprocessOutput, Job.Output, WorkingDirectory);
-			}
-			else
-			{
-				Compiler->CompileShader(Job.Input.ShaderFormat, Job.Input, Job.Output, WorkingDirectory);
-			}
+			Compiler->CompilePreprocessedShader(Job.Input, Job.PreprocessOutput, Job.Output, WorkingDirectory);
 		}
 #if PLATFORM_WINDOWS
 		__except(HandleShaderCompileException(GetExceptionInformation(), OutExceptionMsg, OutExceptionCallstack))
