Removing symbols deprecated in 5.2

diff --git a/Engine/Plugins/FX/Niagara/Source/NiagaraShader/Private/NiagaraShader.cpp b/Engine/Plugins/FX/Niagara/Source/NiagaraShader/Private/NiagaraShader.cpp
--- a/Engine/Plugins/FX/Niagara/Source/NiagaraShader/Private/NiagaraShader.cpp
+++ b/Engine/Plugins/FX/Niagara/Source/NiagaraShader/Private/NiagaraShader.cpp
@@ -390,13 +390,6 @@ void FNiagaraShaderType::BeginCompileShader(
 	NewJobs.Add(FShaderCommonCompileJobPtr(NewJob));
 }
 
-void FNiagaraShaderType::CacheUniformBufferIncludes(TMap<const TCHAR*, FCachedUniformBufferDeclaration, FDefaultSetAllocator, TStringPointerMapKeyFuncs_DEPRECATED<const TCHAR*, FCachedUniformBufferDeclaration>>& Cache, EShaderPlatform Platform) const
-{
-PRAGMA_DISABLE_DEPRECATION_WARNINGS
-	::CacheUniformBufferIncludes(Cache, Platform);
-PRAGMA_ENABLE_DEPRECATION_WARNINGS
-}
-
 void FNiagaraShaderType::AddUniformBufferIncludesToEnvironment(FShaderCompilerEnvironment& OutEnvironment, EShaderPlatform Platform) const
 {
 	UE::ShaderParameters::AddUniformBufferIncludesToEnvironment(OutEnvironment, ReferencedUniformBufferNames);
diff --git a/Engine/Plugins/FX/Niagara/Source/NiagaraShader/Public/NiagaraShaderType.h b/Engine/Plugins/FX/Niagara/Source/NiagaraShader/Public/NiagaraShaderType.h
--- a/Engine/Plugins/FX/Niagara/Source/NiagaraShader/Public/NiagaraShaderType.h
+++ b/Engine/Plugins/FX/Niagara/Source/NiagaraShader/Public/NiagaraShaderType.h
@@ -8,7 +8,6 @@
 
 #include "Shader.h"
 
-struct FCachedUniformBufferDeclaration;
 struct FNiagaraShaderScriptParametersMetadata;
 struct FShaderTarget;
 struct FSharedShaderCompilerEnvironment;
@@ -157,16 +156,6 @@ public:
 #if WITH_EDITOR
 	/** Adds include statements for uniform buffers that this shader type references, and builds a prefix for the shader file with the include statements. */
 	void AddUniformBufferIncludesToEnvironment(FShaderCompilerEnvironment& OutEnvironment, EShaderPlatform Platform) const;
-
-	UE_DEPRECATED(5.2, "AddReferencedUniformBufferIncludes has moved to AddUniformBufferIncludesToEnvironment and no longer takes a prefix argument.")
-	inline void AddReferencedUniformBufferIncludes(FShaderCompilerEnvironment& OutEnvironment, FString& OutSourceFilePrefix, EShaderPlatform Platform) const
-	{
-		AddUniformBufferIncludesToEnvironment(OutEnvironment, Platform);
-	}
-
-	UE_DEPRECATED(5.2, "CacheUniformBufferIncludes should no longer be used.")
-	inline void CacheUniformBufferIncludes(TMap<const TCHAR*, FCachedUniformBufferDeclaration, FDefaultSetAllocator, TStringPointerMapKeyFuncs_DEPRECATED<const TCHAR*, FCachedUniformBufferDeclaration>>& Cache, EShaderPlatform Platform) const;
-
 protected:
 
 	/**
diff --git a/Engine/Source/Runtime/RenderCore/Private/ShaderCore.cpp b/Engine/Source/Runtime/RenderCore/Private/ShaderCore.cpp
--- a/Engine/Source/Runtime/RenderCore/Private/ShaderCore.cpp
+++ b/Engine/Source/Runtime/RenderCore/Private/ShaderCore.cpp
@@ -2894,18 +2894,6 @@ void AppendKeyStringShaderDependencies(
 	}
 }
 
-void SerializeUniformBufferInfo(FShaderSaveArchive& Ar, const TSortedMap<const TCHAR*, FCachedUniformBufferDeclaration, FDefaultAllocator, FUniformBufferNameSortOrder>& UniformBufferEntries)
-{
-	TArray<const TCHAR*> UniformBufferNames;
-	for (const TPair<const TCHAR*, FCachedUniformBufferDeclaration>& Entry : UniformBufferEntries)
-	{
-		UniformBufferNames.Emplace(Entry.Key);
-	}
-	Algo::Sort(UniformBufferNames, FUniformBufferNameSortOrder());
-
-	SerializeUniformBufferInfo_Internal(Ar, UniformBufferNames);
-}
-
 #endif // WITH_EDITOR
 
 FString MakeInjectedShaderCodeBlock(const TCHAR* BlockName, const FString& CodeToInject)
diff --git a/Engine/Source/Runtime/RenderCore/Private/ShaderParameters.cpp b/Engine/Source/Runtime/RenderCore/Private/ShaderParameters.cpp
--- a/Engine/Source/Runtime/RenderCore/Private/ShaderParameters.cpp
+++ b/Engine/Source/Runtime/RenderCore/Private/ShaderParameters.cpp
@@ -419,27 +419,6 @@ static FShaderParametersMetadata* FindShaderParametersMetadataWithVariableName(u
 	return nullptr;
 }
 
-/* deprecated */
-void CacheUniformBufferIncludes(TMap<const TCHAR*, FCachedUniformBufferDeclaration, FDefaultSetAllocator, TStringPointerMapKeyFuncs_DEPRECATED<const TCHAR*, FCachedUniformBufferDeclaration>>& Cache, EShaderPlatform Platform)
-{
-	for (TMap<const TCHAR*, FCachedUniformBufferDeclaration, FDefaultSetAllocator, TStringPointerMapKeyFuncs_DEPRECATED<const TCHAR*, FCachedUniformBufferDeclaration>>::TIterator It(Cache); It; ++It)
-	{
-		const TCHAR* UniformBufferName = It.Key();
-		FCachedUniformBufferDeclaration& BufferDeclaration = It.Value();
-		check(BufferDeclaration.Declaration.Get() == NULL);
-
-		FStringView UniformBufferNameView(UniformBufferName);
-		uint32 UniformBufferNameHash = GetTypeHash(UniformBufferNameView);
-
-		if (const FShaderParametersMetadata* Metadata = FindShaderParametersMetadataWithVariableName(UniformBufferNameHash, UniformBufferNameView))
-		{
-			FString* NewDeclaration = new FString(UE::ShaderParameters::CreateUniformBufferShaderDeclaration(UniformBufferName, *Metadata));
-			check(!NewDeclaration->IsEmpty());
-			BufferDeclaration.Declaration = MakeShareable(NewDeclaration);
-		}
-	}
-}
-
 void UE::ShaderParameters::AddUniformBufferIncludesToEnvironment(FShaderCompilerEnvironment& OutEnvironment, const TSet<const TCHAR*, TStringPointerSetKeyFuncs_DEPRECATED<const TCHAR*>>& InUniformBufferNames)
 {
 	TRACE_CPUPROFILER_EVENT_SCOPE(UE::ShaderParameters::AddUniformBufferIncludesToEnvironment);
@@ -472,10 +451,6 @@ void UE::ShaderParameters::AddUniformBufferIncludesToEnvironment(FShaderCompiler
 	GeneratedUniformBuffersInclude.Append(UniformBufferIncludes);
 }
 
-void FShaderType::FlushShaderFileCache(const TMap<FString, TArray<const TCHAR*> >& ShaderFileToUniformBufferVariables)
-{
-}
-
 void FShaderType::AddUniformBufferIncludesToEnvironment(FShaderCompilerEnvironment& OutEnvironment, EShaderPlatform Platform) const
 {
 	TRACE_CPUPROFILER_EVENT_SCOPE(FShaderType::AddReferencedUniformBufferIncludes);
diff --git a/Engine/Source/Runtime/RenderCore/Public/Shader.h b/Engine/Source/Runtime/RenderCore/Public/Shader.h
--- a/Engine/Source/Runtime/RenderCore/Public/Shader.h
+++ b/Engine/Source/Runtime/RenderCore/Public/Shader.h
@@ -1440,17 +1440,8 @@ public:
 	/** Adds include statements for uniform buffers that this shader type references. */
 	RENDERCORE_API void AddUniformBufferIncludesToEnvironment(FShaderCompilerEnvironment& OutEnvironment, EShaderPlatform Platform) const;
 
-	UE_DEPRECATED(5.2, "AddReferencedUniformBufferIncludes has moved to AddUniformBufferIncludesToEnvironment and no longer takes a prefix argument.")
-	inline void AddReferencedUniformBufferIncludes(FShaderCompilerEnvironment& OutEnvironment, FString& OutSourceFilePrefix, EShaderPlatform Platform) const
-	{
-		AddUniformBufferIncludesToEnvironment(OutEnvironment, Platform);
-	}
-
 	RENDERCORE_API void UpdateReferencedUniformBufferNames(const TMap<FString, TArray<const TCHAR*>>& ShaderFileToUniformBufferVariables);
 
-	UE_DEPRECATED(5.2, "FlushShaderFileCache is deprecated. UpdateReferencedUniformBufferNames should be used to flush any uniform buffer changes")
-	RENDERCORE_API void FlushShaderFileCache(const TMap<FString, TArray<const TCHAR*> >& ShaderFileToUniformBufferVariables);
-
 	RENDERCORE_API void GetShaderStableKeyParts(struct FStableShaderKeyAndValue& SaveKeyVal);
 #endif // WITH_EDITOR
 
diff --git a/Engine/Source/Runtime/RenderCore/Public/ShaderCompilerCore.h b/Engine/Source/Runtime/RenderCore/Public/ShaderCompilerCore.h
--- a/Engine/Source/Runtime/RenderCore/Public/ShaderCompilerCore.h
+++ b/Engine/Source/Runtime/RenderCore/Public/ShaderCompilerCore.h
@@ -746,13 +746,6 @@ extern RENDERCORE_API void CompileShader(const TArray<const IShaderFormat*>& Sha
 extern RENDERCORE_API void CompileShader(const TArray<const IShaderFormat*>& ShaderFormats, FShaderCompileJob& Job, const FString& WorkingDirectory, int32* CompileCount = nullptr);
 extern RENDERCORE_API void CompileShaderPipeline(const TArray<const IShaderFormat*>& ShaderFormats, FShaderPipelineCompileJob* PipelineJob, const FString& WorkingDirectory, int32* CompileCount = nullptr);
 
-UE_DEPRECATED(5.2, "Functionality has moved to FShaderCompilerInput::ShouldUseStableConstantBuffer")
-inline bool ShouldUseStableConstantBuffer(const FShaderCompilerInput& Input)
-{
-	return Input.ShouldUseStableConstantBuffer();
-}
-
-
 /**
  * Validates the format of a virtual shader file path.
  * Meant to be use as such: check(CheckVirtualShaderFilePath(VirtualFilePath));
diff --git a/Engine/Source/Runtime/RenderCore/Public/ShaderCore.h b/Engine/Source/Runtime/RenderCore/Public/ShaderCore.h
--- a/Engine/Source/Runtime/RenderCore/Public/ShaderCore.h
+++ b/Engine/Source/Runtime/RenderCore/Public/ShaderCore.h
@@ -1192,12 +1192,6 @@ extern RENDERCORE_API const class FSHAHash& GetShaderFilesHash(const TArray<FStr
  */
 extern RENDERCORE_API void FlushShaderFileCache();
 
-UE_DEPRECATED(5.2, "FlushShaderFileCache no longer needs a ShaderPlatformName argument")
-inline void FlushShaderFileCache(const FName* ShaderPlatformName)
-{
-	FlushShaderFileCache();
-}
-
 extern RENDERCORE_API void VerifyShaderSourceFiles(EShaderPlatform ShaderPlatform);
 
 #if WITH_EDITOR
@@ -1212,13 +1206,6 @@ RENDERCORE_API void UpdateReferencedUniformBufferNames(
 	TArrayView<const FVertexFactoryType*> OutdatedFactoryTypes,
 	TArrayView<const FShaderPipelineType*> OutdatedShaderPipelineTypes);
 
-/** Deprecated structure (only used by deprecated "CacheUniformBufferIncludes" and "SerializeUniformBufferInfo") */
-struct FCachedUniformBufferDeclaration
-{
-	// Using SharedPtr so we can hand off lifetime ownership to FShaderCompilerEnvironment::IncludeVirtualPathToSharedContentsMap when invalidating this cache
-	FThreadSafeSharedStringPtr Declaration;
-};
-
 /** Parses the given source file and its includes for references of uniform buffers. */
 extern void GenerateReferencedUniformBufferNames(
 	const TCHAR* SourceFilename,
@@ -1234,14 +1221,6 @@ struct FUniformBufferNameSortOrder
 	}
 };
 
-using FSortedMapUniformBufferDeclaration = TSortedMap<const TCHAR*, FCachedUniformBufferDeclaration, FDefaultAllocator, FUniformBufferNameSortOrder>;
-
-/** Records information about all the uniform buffer layouts referenced by UniformBufferEntries. This function is now deprecated as there now a unified way of preparing
- * shadermap keys that includes this and more: AppendKeyStringShaderDependencies.
- */
-UE_DEPRECATED(5.2, "SerializeUniformBufferInfo is depreceated. For creating shadermap keys please use AppendKeyStringShaderDependencies")
-extern RENDERCORE_API void SerializeUniformBufferInfo(class FShaderSaveArchive& Ar, const FSortedMapUniformBufferDeclaration& UniformBufferEntries);
-
 /**
  * Return the hash of the given type layout for a partical platform type layout. This function employs caching to avoid re-hashing the same parameters several times.
  */
diff --git a/Engine/Source/Runtime/RenderCore/Public/ShaderParameterParser.h b/Engine/Source/Runtime/RenderCore/Public/ShaderParameterParser.h
--- a/Engine/Source/Runtime/RenderCore/Public/ShaderParameterParser.h
+++ b/Engine/Source/Runtime/RenderCore/Public/ShaderParameterParser.h
@@ -133,19 +133,6 @@ public:
 		EBindlessParameterMode BindlessParameterMode = EBindlessParameterMode::Default
 	);
 
-	UE_DEPRECATED(5.2, "ParseAndModify doesn't need ConstantBufferType anymore")
-	bool ParseAndModify(
-		const FShaderCompilerInput& CompilerInput,
-		FShaderCompilerOutput& CompilerOutput,
-		FString& PreprocessedShaderSource,
-		const TCHAR* InConstantBufferType
-	)
-	{
-		PRAGMA_DISABLE_DEPRECATION_WARNINGS
-		return ParseAndModify(CompilerInput, CompilerOutput, PreprocessedShaderSource);
-		PRAGMA_ENABLE_DEPRECATION_WARNINGS
-	}
-
 	/** Gets parsing information from a parameter binding name. */
 	const FParsedShaderParameter& FindParameterInfos(const FString& ParameterName) const
 	{
diff --git a/Engine/Source/Runtime/RenderCore/Public/ShaderParameters.h b/Engine/Source/Runtime/RenderCore/Public/ShaderParameters.h
--- a/Engine/Source/Runtime/RenderCore/Public/ShaderParameters.h
+++ b/Engine/Source/Runtime/RenderCore/Public/ShaderParameters.h
@@ -23,7 +23,6 @@ class FRHITexture;
 class FRHIUnorderedAccessView;
 class FShaderParameterMap;
 class FShaderParametersMetadata;
-struct FCachedUniformBufferDeclaration;
 struct FRWBuffer;
 struct FRWBufferStructured;
 struct FShaderCompilerEnvironment;
@@ -42,15 +41,6 @@ namespace UE::ShaderParameters
 
 	RENDERCORE_API void AddUniformBufferIncludesToEnvironment(FShaderCompilerEnvironment& OutEnvironment, const TSet<const TCHAR*, TStringPointerSetKeyFuncs_DEPRECATED<const TCHAR*>>& InUniformBufferNames);
 }
-
-UE_DEPRECATED(5.2, "CreateUniformBufferShaderDeclaration has moved to UE::ShaderParameters::CreateUniformBufferShaderDeclaration, does not take a EShaderPlatform argument and now returns the Declaration.")
-inline void CreateUniformBufferShaderDeclaration(const TCHAR* Name, const FShaderParametersMetadata& UniformBufferStruct, EShaderPlatform Platform, FString& OutDeclaration)
-{
-	OutDeclaration = UE::ShaderParameters::CreateUniformBufferShaderDeclaration(Name, UniformBufferStruct);
-}
-
-UE_DEPRECATED(5.2, "CacheUniformBufferIncludes should no longer be used.")
-RENDERCORE_API void CacheUniformBufferIncludes(TMap<const TCHAR*, FCachedUniformBufferDeclaration, FDefaultSetAllocator, TStringPointerMapKeyFuncs_DEPRECATED<const TCHAR*, FCachedUniformBufferDeclaration>>& Cache, EShaderPlatform Platform);
 #endif
 
 enum EShaderParameterFlags
