[shaders] fix error introduced by deprecation - shaders requiring a secondary compile were still calling the single compile path as well, resulting in errors when populating an already-populated output struct. tested local cook on affected platform to confirm issue is resolved (preflights are currently broken due to an unrelated problem)

diff --git a/Engine/Source/Runtime/RenderCore/Private/ShaderCore.cpp b/Engine/Source/Runtime/RenderCore/Private/ShaderCore.cpp
--- a/Engine/Source/Runtime/RenderCore/Private/ShaderCore.cpp
+++ b/Engine/Source/Runtime/RenderCore/Private/ShaderCore.cpp
@@ -1305,7 +1305,10 @@ public:
 				check(Job.SecondaryPreprocessOutput.IsValid());
 				Compiler->CompilePreprocessedShader(Job.Input, Job.PreprocessOutput, *Job.SecondaryPreprocessOutput, Job.Output, *Job.SecondaryOutput, WorkingDirectory);
 			}
-			Compiler->CompilePreprocessedShader(Job.Input, Job.PreprocessOutput, Job.Output, WorkingDirectory);
+			else
+			{
+				Compiler->CompilePreprocessedShader(Job.Input, Job.PreprocessOutput, Job.Output, WorkingDirectory);
+			}
 		}
 #if PLATFORM_WINDOWS
 		__except(HandleShaderCompileException(GetExceptionInformation(), OutExceptionMsg, OutExceptionCallstack))
