Implements r.DumpGPU.CameraCut

diff --git a/Engine/Source/Runtime/Launch/Private/LaunchEngineLoop.cpp b/Engine/Source/Runtime/Launch/Private/LaunchEngineLoop.cpp
--- a/Engine/Source/Runtime/Launch/Private/LaunchEngineLoop.cpp
+++ b/Engine/Source/Runtime/Launch/Private/LaunchEngineLoop.cpp
@@ -189,7 +189,7 @@
 	#include "RenderUtils.h"
 	#include "DynamicResolutionState.h"
 	#include "EngineModule.h"
-	#include "RenderGraphBuilder.h"
+	#include "DumpGPU.h"
 
 #if !UE_SERVER
 	#include "AppMediaTimeSource.h"
@@ -5825,11 +5825,6 @@ void FEngineLoop::Tick()
             }
 		}
 		
-		// init for RDG resource dump
-		#if WITH_ENGINE && WITH_DUMPGPU
-			FRDGBuilder::InitResourceDump();
-		#endif
-
 		// main game engine tick (world, game objects, etc.)
 		GEngine->Tick(FApp::GetDeltaTime(), bIdleMode);
 
@@ -6105,9 +6100,9 @@ void FEngineLoop::Tick()
 
 		FCoreDelegates::OnEndFrame.Broadcast();
 
-		// end of RDG resource dump
+		// Tick DumpGPU to start/stop frame dumps
 		#if WITH_ENGINE && WITH_DUMPGPU
-			FRDGBuilder::EndResourceDump();
+			UE::RenderCore::DumpGPU::TickEndFrame();
 		#endif
 
 		#if !UE_SERVER && WITH_ENGINE
diff --git a/Engine/Source/Runtime/RenderCore/Private/DumpGPU.cpp b/Engine/Source/Runtime/RenderCore/Private/DumpGPU.cpp
--- a/Engine/Source/Runtime/RenderCore/Private/DumpGPU.cpp
+++ b/Engine/Source/Runtime/RenderCore/Private/DumpGPU.cpp
@@ -100,6 +100,11 @@ static TAutoConsoleVariable<int32> GDumpGPUFrameCount(
 	TEXT("Number of consecutive frames to dump (default=1)."),
 	ECVF_Default);
 
+static TAutoConsoleVariable<int32> GDumpGPUCameraCut(
+	TEXT("r.DumpGPU.CameraCut"), 0,
+	TEXT("Whether to issue a camera cut on the first frame of the dump."),
+	ECVF_Default);
+
 static TAutoConsoleVariable<float> GDumpGPUFixedTickRate(
 	TEXT("r.DumpGPU.FixedTickRate"), 0.0f,
 	TEXT("Override the engine's tick rate to be fixed for every dumped frames (default=0)."),
@@ -232,7 +237,8 @@ public:
 	bool bUpload = false;
 	bool bStream = false;
 	float DeltaTime = 0.0f;
-	int32 RemainingFrameCount = 1;
+	int32 DumpedFrameId = 0;
+	int32 FrameCount = 1;
 	FName UploadResourceCompressionName;
 	FString DumpingDirectoryPath;
 	FDateTime Time;
@@ -1997,47 +2003,23 @@ public:
 		DumpStatusToFile(TEXT("crash"));
 		FGenericCrashContext::DumpLog(DumpingDirectoryPath / FRDGResourceDumpContext::kBaseDir);
 	}
+
+	void Start();
+	void Finish();
 };
 
-// 0 = not dumping, MAX_uint64 dump request for next frame, otherwise dump frame counter
-static uint64 DumpingFrameCounter_GameThread = 0;
-static FRDGResourceDumpContext* GRDGResourceDumpContext = nullptr;
+static FRDGResourceDumpContext* GRDGResourceDumpContext_GameThread = nullptr;
+static FRDGResourceDumpContext* GRDGResourceDumpContext_RenderThread = nullptr;
 
 static float GNextDumpingRemainingTime = -1.0f;
 static FRDGResourceDumpContext* GNextRDGResourceDumpContext = nullptr;
 
-bool IsDumpingRDGResources()
-{
-	return GRDGResourceDumpContext != nullptr;
-}
-
-void FRDGBuilder::InitResourceDump()
-{
-	check(IsInGameThread());
-
-	if(DumpingFrameCounter_GameThread == MAX_uint64)
-	{
-		DumpingFrameCounter_GameThread = GFrameCounter;
-	}
-
-	if (GRDGResourceDumpContext)
-	{
-		float DeltaTime = FApp::GetDeltaTime();
-
-		ENQUEUE_RENDER_COMMAND(FDumpGPUUpdateDeltaTime)(
-			[DeltaTime](FRHICommandListImmediate& ImmediateRHICmdList)
-		{
-			GRDGResourceDumpContext->DeltaTime = DeltaTime;
-			UE_LOG(LogDumpGPU, Display, TEXT("Remaining frames %d"), GRDGResourceDumpContext->RemainingFrameCount);
-		});
-	}
-}
 
 FString FRDGBuilder::BeginResourceDump(const TCHAR* Cmd)
 {
 	check(IsInGameThread());
 
-	if (DumpingFrameCounter_GameThread != 0 || GNextDumpingRemainingTime > 0.0f)
+	if (GRDGResourceDumpContext_GameThread || GNextRDGResourceDumpContext)
 	{
 		return FString();
 	}
@@ -2078,119 +2060,127 @@ FString FRDGBuilder::BeginResourceDump(const TCHAR* Cmd)
 	}
 
 	FRDGResourceDumpContext* NewResourceDumpContext;
-	if (GNextRDGResourceDumpContext == nullptr)
+	NewResourceDumpContext = new FRDGResourceDumpContext;
+
+	NewResourceDumpContext->Time = FDateTime::Now();
 	{
-		NewResourceDumpContext = new FRDGResourceDumpContext;
+		FString CVarDirectoryPath = GDumpGPUDirectoryCVar.GetValueOnGameThread();
+		FString EnvDirectoryPath = FPlatformMisc::GetEnvironmentVariable(TEXT("UE-DumpGPUPath"));
 
-		NewResourceDumpContext->Time = FDateTime::Now();
+		FString DirectoryPath;
+		if (!CVarDirectoryPath.IsEmpty())
 		{
-			FString CVarDirectoryPath = GDumpGPUDirectoryCVar.GetValueOnGameThread();
-			FString EnvDirectoryPath = FPlatformMisc::GetEnvironmentVariable(TEXT("UE-DumpGPUPath"));
-
-			FString DirectoryPath;
-			if (!CVarDirectoryPath.IsEmpty())
-			{
-				DirectoryPath = CVarDirectoryPath;
-			}
-			else if (!EnvDirectoryPath.IsEmpty())
-			{
-				DirectoryPath = EnvDirectoryPath;
-			}
-			else
-			{
-				DirectoryPath = FPaths::ProjectSavedDir() / TEXT("GPUDumps/");
-			}
-			NewResourceDumpContext->DumpingDirectoryPath = DirectoryPath / FApp::GetProjectName() + TEXT("-") + FPlatformProperties::PlatformName() + TEXT("-") + NewResourceDumpContext->Time.ToString() + TEXT("/");
+			DirectoryPath = CVarDirectoryPath;
 		}
-		NewResourceDumpContext->bStream = GDumpGPUStream.GetValueOnGameThread() != 0;
-		NewResourceDumpContext->bEnableDiskWrite = GDumpTestEnableDiskWrite.GetValueOnGameThread() != 0;
-		NewResourceDumpContext->DeltaTime = FApp::GetDeltaTime();
-		NewResourceDumpContext->RemainingFrameCount = FMath::Max(GDumpGPUFrameCount.GetValueOnGameThread(), 1);
-
-		if (Switches.Contains(TEXT("upload")))
+		else if (!EnvDirectoryPath.IsEmpty())
 		{
-			if (!IDumpGPUUploadServiceProvider::GProvider || !NewResourceDumpContext->bEnableDiskWrite)
-			{
-				UE_LOG(LogDumpGPU, Warning, TEXT("DumpGPU upload services are not set up."));
-			}
-			else if (GDumpGPUUploadCVar.GetValueOnGameThread() == 0)
-			{
-				UE_LOG(LogDumpGPU, Warning, TEXT("DumpGPU upload services are not available because r.DumpGPU.Upload=0."));
-			}
-			else
-			{
-				NewResourceDumpContext->bUpload = true;
-			}
+			DirectoryPath = EnvDirectoryPath;
+		}
+		else
+		{
+			DirectoryPath = FPaths::ProjectSavedDir() / TEXT("GPUDumps/");
 		}
+		NewResourceDumpContext->DumpingDirectoryPath = DirectoryPath / FApp::GetProjectName() + TEXT("-") + FPlatformProperties::PlatformName() + TEXT("-") + NewResourceDumpContext->Time.ToString() + TEXT("/");
+	}
+	NewResourceDumpContext->bStream = GDumpGPUStream.GetValueOnGameThread() != 0;
+	NewResourceDumpContext->bEnableDiskWrite = GDumpTestEnableDiskWrite.GetValueOnGameThread() != 0;
+	NewResourceDumpContext->DeltaTime = FApp::GetDeltaTime();
+	NewResourceDumpContext->FrameCount = FMath::Max(GDumpGPUFrameCount.GetValueOnGameThread(), 1);
 
-		if (NewResourceDumpContext->bUpload)
+	if (Switches.Contains(TEXT("upload")))
+	{
+		if (!IDumpGPUUploadServiceProvider::GProvider || !NewResourceDumpContext->bEnableDiskWrite)
 		{
-			if (GDumpGPUUploadCompressResources.GetValueOnGameThread() == 1)
-			{
-				NewResourceDumpContext->UploadResourceCompressionName = NAME_Zlib;
-			}
-			else if (GDumpGPUUploadCompressResources.GetValueOnGameThread() == 2)
-			{
-				NewResourceDumpContext->UploadResourceCompressionName = NAME_Gzip;
-			}
+			UE_LOG(LogDumpGPU, Warning, TEXT("DumpGPU upload services are not set up."));
+		}
+		else if (GDumpGPUUploadCVar.GetValueOnGameThread() == 0)
+		{
+			UE_LOG(LogDumpGPU, Warning, TEXT("DumpGPU upload services are not available because r.DumpGPU.Upload=0."));
+		}
+		else
+		{
+			NewResourceDumpContext->bUpload = true;
 		}
+	}
 
-		NewResourceDumpContext->bShowInExplore = NewResourceDumpContext->bEnableDiskWrite && GDumpExploreCVar.GetValueOnGameThread() != 0 && !NewResourceDumpContext->bUpload;
-	
-		if (GDumpGPUDelay.GetValueOnGameThread() > 0.0f)
+	if (NewResourceDumpContext->bUpload)
+	{
+		if (GDumpGPUUploadCompressResources.GetValueOnGameThread() == 1)
+		{
+			NewResourceDumpContext->UploadResourceCompressionName = NAME_Zlib;
+		}
+		else if (GDumpGPUUploadCompressResources.GetValueOnGameThread() == 2)
 		{
-			GNextDumpingRemainingTime = GDumpGPUDelay.GetValueOnGameThread();
-			GNextRDGResourceDumpContext = NewResourceDumpContext;
-			UE_LOG(LogDumpGPU, Display, TEXT("DumpGPU to %s armed for %f seconds."), *NewResourceDumpContext->DumpingDirectoryPath, GNextDumpingRemainingTime);
-			return NewResourceDumpContext->DumpingDirectoryPath;
+			NewResourceDumpContext->UploadResourceCompressionName = NAME_Gzip;
 		}
 	}
+
+	NewResourceDumpContext->bShowInExplore = NewResourceDumpContext->bEnableDiskWrite && GDumpExploreCVar.GetValueOnGameThread() != 0 && !NewResourceDumpContext->bUpload;
+
+	check(GNextRDGResourceDumpContext == nullptr);
+	GNextRDGResourceDumpContext = NewResourceDumpContext;
+	if (GDumpGPUDelay.GetValueOnGameThread() > 0.0f)
+	{
+		GNextDumpingRemainingTime = GDumpGPUDelay.GetValueOnGameThread();
+		UE_LOG(LogDumpGPU, Display, TEXT("DumpGPU to %s armed for %d frames starting in %f seconds."), *NewResourceDumpContext->DumpingDirectoryPath, NewResourceDumpContext->FrameCount, GNextDumpingRemainingTime);
+	}
 	else
 	{
-		NewResourceDumpContext = GNextRDGResourceDumpContext;
-		GNextRDGResourceDumpContext = nullptr;
+		UE_LOG(LogDumpGPU, Display, TEXT("Dump to %s armed for %d frames starting next frame."), *NewResourceDumpContext->DumpingDirectoryPath, NewResourceDumpContext->FrameCount);
+	}
+
+	if (NewResourceDumpContext->bEnableDiskWrite)
+	{
+		return NewResourceDumpContext->DumpingDirectoryPath;
 	}
+	return FString();
+}
+
+void FRDGResourceDumpContext::Start()
+{
+	check(IsInGameThread());
+	check(GNextRDGResourceDumpContext == this);
 
 	// Dumping resource may take a while and we don't want to get 'hang' crashes in the process.
 	// We resume from EndResourceDump after the dump has finished.
 	FThreadHeartBeat::Get().SuspendHeartBeat(true);
 	SuspendRenderThreadTimeout();
 
-	UE_LOG(LogDumpGPU, Display, TEXT("DumpGPU to %s starting this frame"), *NewResourceDumpContext->DumpingDirectoryPath);
-	NewResourceDumpContext->MemoryConstants = FPlatformMemory::GetConstants();
-	NewResourceDumpContext->MemoryStats = FPlatformMemory::GetStats();
+	UE_LOG(LogDumpGPU, Display, TEXT("DumpGPU to %s starting this frame"), *DumpingDirectoryPath);
+	MemoryConstants = FPlatformMemory::GetConstants();
+	MemoryStats = FPlatformMemory::GetStats();
 
 	IPlatformFile& PlatformFile = FPlatformFileManager::Get().GetPlatformFile();
-	if (NewResourceDumpContext->bEnableDiskWrite)
+	if (bEnableDiskWrite)
 	{
-		if (!PlatformFile.DirectoryExists(*NewResourceDumpContext->DumpingDirectoryPath))
+		if (!PlatformFile.DirectoryExists(*DumpingDirectoryPath))
 		{
-			PlatformFile.CreateDirectoryTree(*NewResourceDumpContext->DumpingDirectoryPath);
+			PlatformFile.CreateDirectoryTree(*DumpingDirectoryPath);
 		}
-		PlatformFile.CreateDirectoryTree(*(NewResourceDumpContext->DumpingDirectoryPath / FRDGResourceDumpContext::kBaseDir));
-		PlatformFile.CreateDirectoryTree(*(NewResourceDumpContext->DumpingDirectoryPath / FRDGResourceDumpContext::kResourcesDir));
+		PlatformFile.CreateDirectoryTree(*(DumpingDirectoryPath / FRDGResourceDumpContext::kBaseDir));
+		PlatformFile.CreateDirectoryTree(*(DumpingDirectoryPath / FRDGResourceDumpContext::kResourcesDir));
 
-		NewResourceDumpContext->DumpStringToFile(TEXT(""), FString(FRDGResourceDumpContext::kBaseDir) / TEXT("Passes.json"));
-		NewResourceDumpContext->DumpStringToFile(TEXT(""), FString(FRDGResourceDumpContext::kBaseDir) / TEXT("ResourceDescs.json"));
-		NewResourceDumpContext->DumpStringToFile(TEXT(""), FString(FRDGResourceDumpContext::kBaseDir) / TEXT("PassDrawCounts.json"));
+		DumpStringToFile(TEXT(""), FString(FRDGResourceDumpContext::kBaseDir) / TEXT("Passes.json"));
+		DumpStringToFile(TEXT(""), FString(FRDGResourceDumpContext::kBaseDir) / TEXT("ResourceDescs.json"));
+		DumpStringToFile(TEXT(""), FString(FRDGResourceDumpContext::kBaseDir) / TEXT("PassDrawCounts.json"));
 	}
 
 	// Dump status file and register NewResourceDumpContext to listen for OnShutdownAfterError to get a log of what happened with callstack.
 	{
-		NewResourceDumpContext->DumpStatusToFile(TEXT("dumping"));
-		FCoreDelegates::OnShutdownAfterError.AddRaw(NewResourceDumpContext, &FRDGResourceDumpContext::OnCrash);
+		DumpStatusToFile(TEXT("dumping"));
+		FCoreDelegates::OnShutdownAfterError.AddRaw(this, &FRDGResourceDumpContext::OnCrash);
 	}
 
 	// Dump service parameters so GPUDumpViewer.html remain compatible when not using upload provider.
-	if (NewResourceDumpContext->bEnableDiskWrite)
+	if (bEnableDiskWrite)
 	{
-		NewResourceDumpContext->GetDumpParameters().DumpServiceParametersFile();
+		GetDumpParameters().DumpServiceParametersFile();
 	}
 
 	if (GDumpGPUFixedTickRate.GetValueOnGameThread() > 0.0f && !FApp::UseFixedTimeStep())
 	{
-		NewResourceDumpContext->bOverrideFixedDeltaTime = true;
-		NewResourceDumpContext->PreviousFixedDeltaTime = FApp::GetFixedDeltaTime();
+		bOverrideFixedDeltaTime = true;
+		PreviousFixedDeltaTime = FApp::GetFixedDeltaTime();
 
 		FApp::SetFixedDeltaTime(1.0f / GDumpGPUFixedTickRate.GetValueOnGameThread());
 		FApp::SetUseFixedTimeStep(true);
@@ -2239,49 +2229,52 @@ FString FRDGBuilder::BeginResourceDump(const TCHAR* Cmd)
 		JsonObject->SetStringField(TEXT("GPUDriverUserVersion"), GPUDriverInfo.UserDriverVersion);
 		JsonObject->SetStringField(TEXT("GPUDriverInternalVersion"), GPUDriverInfo.GetUnifiedDriverVersion());
 		JsonObject->SetStringField(TEXT("GPUDriverDate"), GPUDriverInfo.DriverDate);
-		JsonObject->SetNumberField(TEXT("MemoryTotalPhysical"), NewResourceDumpContext->MemoryConstants.TotalPhysical);
-		JsonObject->SetNumberField(TEXT("MemoryPageSize"), NewResourceDumpContext->MemoryConstants.PageSize);
+		JsonObject->SetNumberField(TEXT("MemoryTotalPhysical"), MemoryConstants.TotalPhysical);
+		JsonObject->SetNumberField(TEXT("MemoryPageSize"), MemoryConstants.PageSize);
 		JsonObject->SetStringField(TEXT("RHI"), GDynamicRHI->GetName());
 		JsonObject->SetStringField(TEXT("RHIMaxFeatureLevel"), LexToString(GMaxRHIFeatureLevel));
-		JsonObject->SetStringField(TEXT("DumpTime"), NewResourceDumpContext->Time.ToString());
+		JsonObject->SetStringField(TEXT("DumpTime"), Time.ToString());
 		{
 			const FString LogSrcAbsolute = FPlatformOutputDevices::GetAbsoluteLogFilename();
 			FString LogFilename = FPaths::GetCleanFilename(LogSrcAbsolute);
 			JsonObject->SetStringField(TEXT("LogFilename"), LogFilename);
 		}
 
-		NewResourceDumpContext->DumpJsonToFile(JsonObject, FString(FRDGResourceDumpContext::kBaseDir) / TEXT("Infos.json"));
+		DumpJsonToFile(JsonObject, FString(FRDGResourceDumpContext::kBaseDir) / TEXT("Infos.json"));
 	}
 
 	// Dump the rendering cvars
-	if (NewResourceDumpContext->bEnableDiskWrite)
+	if (bEnableDiskWrite)
 	{
-		NewResourceDumpContext->KickOffAsyncTask([NewResourceDumpContext]() {
-			NewResourceDumpContext->DumpRenderingCVarsToCSV();
+		KickOffAsyncTask([this]() {
+			this->DumpRenderingCVarsToCSV();
 		});
 	}
 
 	// Copy the viewer
-	if (NewResourceDumpContext->bEnableDiskWrite)
+	if (bEnableDiskWrite)
 	{
-		NewResourceDumpContext->KickOffAsyncTask([NewResourceDumpContext, &PlatformFile]() {
+		KickOffAsyncTask([this, &PlatformFile]() {
 			const TCHAR* OpenGPUDumpViewerBatName = TEXT("OpenGPUDumpViewer.bat");
 			const TCHAR* OpenGPUDumpViewerShName = TEXT("OpenGPUDumpViewer.sh");
 
 			const TCHAR* ViewerHTML = TEXT("GPUDumpViewer.html");
 			FString DumpGPUViewerSourcePath = FPaths::EngineDir() + FString(TEXT("Extras")) / TEXT("GPUDumpViewer");
 
-			PlatformFile.CopyFile(*(NewResourceDumpContext->DumpingDirectoryPath / ViewerHTML), *(DumpGPUViewerSourcePath / ViewerHTML));
-			PlatformFile.CopyFile(*(NewResourceDumpContext->DumpingDirectoryPath / OpenGPUDumpViewerBatName), *(DumpGPUViewerSourcePath / OpenGPUDumpViewerBatName));
-			PlatformFile.CopyFile(*(NewResourceDumpContext->DumpingDirectoryPath / OpenGPUDumpViewerShName), *(DumpGPUViewerSourcePath / OpenGPUDumpViewerShName));
+			PlatformFile.CopyFile(*(this->DumpingDirectoryPath / ViewerHTML), *(DumpGPUViewerSourcePath / ViewerHTML));
+			PlatformFile.CopyFile(*(this->DumpingDirectoryPath / OpenGPUDumpViewerBatName), *(DumpGPUViewerSourcePath / OpenGPUDumpViewerBatName));
+			PlatformFile.CopyFile(*(this->DumpingDirectoryPath / OpenGPUDumpViewerShName), *(DumpGPUViewerSourcePath / OpenGPUDumpViewerShName));
 		});
 	}
 
+	GNextRDGResourceDumpContext = nullptr;
+	GRDGResourceDumpContext_GameThread = this;
+
 	ENQUEUE_RENDER_COMMAND(FStartGPUDump)(
-		[NewResourceDumpContext](FRHICommandListImmediate& ImmediateRHICmdList)
+		[this](FRHICommandListImmediate& ImmediateRHICmdList)
 	{
 		check(IsInRenderingThread());
-		GRDGResourceDumpContext = NewResourceDumpContext;
+		GRDGResourceDumpContext_RenderThread = this;
 
 		ImmediateRHICmdList.SubmitCommandsAndFlushGPU();
 
@@ -2290,65 +2283,23 @@ FString FRDGBuilder::BeginResourceDump(const TCHAR* Cmd)
 			GRHIValidateBufferSourceCopy = false;
 		#endif
 	});
-
-	// Mark ready for dump on next available frame
-	DumpingFrameCounter_GameThread = MAX_uint64;
-
-	if (NewResourceDumpContext->bEnableDiskWrite)
-	{
-		return NewResourceDumpContext->DumpingDirectoryPath;
-	}
-	return FString();
 }
 
-void FRDGBuilder::EndResourceDump()
+void FRDGResourceDumpContext::Finish()
 {
 	check(IsInGameThread());
-
-	 // make sure at least one frame has passed since we start a resource dump and we are not waiting on the dump to begin
-	if (DumpingFrameCounter_GameThread == 0 ||
-		DumpingFrameCounter_GameThread == MAX_uint64 ||
-		DumpingFrameCounter_GameThread >= GFrameCounter)
-	{
-		if (GNextDumpingRemainingTime > 0.0f)
-		{
-			check(GNextRDGResourceDumpContext);
-			GNextDumpingRemainingTime -= FApp::GetDeltaTime();
-			if (GNextDumpingRemainingTime <= 0.0)
-			{
-				GNextDumpingRemainingTime = -1.0f;
-				FRDGBuilder::BeginResourceDump(TEXT(""));
-				check(!GNextRDGResourceDumpContext);
-			}
-		}
-		return;
-	}
-
-	GRDGResourceDumpContext->RemainingFrameCount--;
-	if (GRDGResourceDumpContext->RemainingFrameCount > 0)
-	{
-		if (GRDGResourceDumpContext->bStream)
-		{
-			ENQUEUE_RENDER_COMMAND(FEndGPUDump)(
-				[](FRHICommandListImmediate& ImmediateRHICmdList)
-			{
-				GRDGResourceDumpContext->LandCompletedResources(ImmediateRHICmdList);
-			});
-		}
-
-		return;
-	}
-
+	check(GRDGResourceDumpContext_GameThread == this);
+	
 	// Wait all rendering commands are completed to finish with GRDGResourceDumpContext.
 	{
 		UE_LOG(LogDumpGPU, Display, TEXT("Stalling game thread until render thread finishes to dump resources"));
 
 		ENQUEUE_RENDER_COMMAND(FEndGPUDump)(
-			[](FRHICommandListImmediate& ImmediateRHICmdList)
+			[this](FRHICommandListImmediate& ImmediateRHICmdList)
 		{
-			if (GRDGResourceDumpContext->bStream)
+			if (bStream)
 			{
-				GRDGResourceDumpContext->WaitAndReleaseStagingResources(ImmediateRHICmdList);
+				WaitAndReleaseStagingResources(ImmediateRHICmdList);
 			}
 
 			ImmediateRHICmdList.SubmitCommandsAndFlushGPU();
@@ -2361,65 +2312,65 @@ void FRDGBuilder::EndResourceDump()
 	}
 
 	// Log information about the dump.
-	FString AbsDumpingDirectoryPath = IFileManager::Get().ConvertToAbsolutePathForExternalAppForRead(*GRDGResourceDumpContext->DumpingDirectoryPath);
+	FString AbsDumpingDirectoryPath = IFileManager::Get().ConvertToAbsolutePathForExternalAppForRead(*DumpingDirectoryPath);
 	{
 		FDateTime Now = FDateTime::Now();
-		double TotalDumpSeconds = (Now - GRDGResourceDumpContext->Time).GetTotalSeconds();
-		double RHIReadbackCommandsSeconds = GRDGResourceDumpContext->TimingBucket[int32(FRDGResourceDumpContext::ETimingBucket::RHIReadbackCommands)];
-		double GPUWaitSeconds = GRDGResourceDumpContext->TimingBucket[int32(FRDGResourceDumpContext::ETimingBucket::GPUWait)];
-		double CPUPostProcessingSeconds = GRDGResourceDumpContext->TimingBucket[int32(FRDGResourceDumpContext::ETimingBucket::CPUPostProcessing)];
-		double MetadataFileWriteSeconds = GRDGResourceDumpContext->TimingBucket[int32(FRDGResourceDumpContext::ETimingBucket::MetadataFileWrite)];
-		double ParametersFileWriteSeconds = GRDGResourceDumpContext->TimingBucket[int32(FRDGResourceDumpContext::ETimingBucket::ParametersFileWrite)];
-		double ResourceBinaryFileWriteSeconds = GRDGResourceDumpContext->TimingBucket[int32(FRDGResourceDumpContext::ETimingBucket::ResourceBinaryFileWrite)];
-		double RHIReleaseResourcesTimeSeconds = GRDGResourceDumpContext->TimingBucket[int32(FRDGResourceDumpContext::ETimingBucket::RHIReleaseResources)];
-
-		UE_LOG(LogDumpGPU, Display, TEXT("Dumped %d resources in %.3f s to %s"), GRDGResourceDumpContext->ResourcesDumpPasses, float(TotalDumpSeconds), *AbsDumpingDirectoryPath);
+		double TotalDumpSeconds = (Now - Time).GetTotalSeconds();
+		double RHIReadbackCommandsSeconds = TimingBucket[int32(FRDGResourceDumpContext::ETimingBucket::RHIReadbackCommands)];
+		double GPUWaitSeconds = TimingBucket[int32(FRDGResourceDumpContext::ETimingBucket::GPUWait)];
+		double CPUPostProcessingSeconds = TimingBucket[int32(FRDGResourceDumpContext::ETimingBucket::CPUPostProcessing)];
+		double MetadataFileWriteSeconds = TimingBucket[int32(FRDGResourceDumpContext::ETimingBucket::MetadataFileWrite)];
+		double ParametersFileWriteSeconds = TimingBucket[int32(FRDGResourceDumpContext::ETimingBucket::ParametersFileWrite)];
+		double ResourceBinaryFileWriteSeconds = TimingBucket[int32(FRDGResourceDumpContext::ETimingBucket::ResourceBinaryFileWrite)];
+		double RHIReleaseResourcesTimeSeconds = TimingBucket[int32(FRDGResourceDumpContext::ETimingBucket::RHIReleaseResources)];
+
+		UE_LOG(LogDumpGPU, Display, TEXT("Dumped %d resources in %.3f s to %s"), ResourcesDumpPasses, float(TotalDumpSeconds), *AbsDumpingDirectoryPath);
 		UE_LOG(LogDumpGPU, Display, TEXT("Dumped GPU readback commands: %.3f s"), float(RHIReadbackCommandsSeconds));
 		UE_LOG(LogDumpGPU, Display, TEXT("Dumped GPU wait: %.3f s"), float(GPUWaitSeconds));
 		UE_LOG(LogDumpGPU, Display, TEXT("Dumped CPU resource binary post processing: %.3f s"), float(CPUPostProcessingSeconds));
 		UE_LOG(LogDumpGPU, Display, TEXT("Dumped metadata: %.3f MB in %d files under %.3f s at %.3f MB/s"),
-			float(GRDGResourceDumpContext->MetadataFilesWriteBytes) / float(1024 * 1024),
-			int32(GRDGResourceDumpContext->MetadataFilesOpened),
+			float(MetadataFilesWriteBytes) / float(1024 * 1024),
+			int32(MetadataFilesOpened),
 			float(MetadataFileWriteSeconds),
-			float(GRDGResourceDumpContext->MetadataFilesWriteBytes) / (float(1024 * 1024) * float(MetadataFileWriteSeconds)));
+			float(MetadataFilesWriteBytes) / (float(1024 * 1024) * float(MetadataFileWriteSeconds)));
 		UE_LOG(LogDumpGPU, Display, TEXT("Dumped parameters: %.3f MB in %d files under %.3f s at %.3f MB/s"),
-			float(GRDGResourceDumpContext->ParametersFilesWriteBytes) / float(1024 * 1024),
-			int32(GRDGResourceDumpContext->ParametersFilesOpened),
+			float(ParametersFilesWriteBytes) / float(1024 * 1024),
+			int32(ParametersFilesOpened),
 			float(ParametersFileWriteSeconds),
-			float(GRDGResourceDumpContext->ParametersFilesWriteBytes) / (float(1024 * 1024) * float(ParametersFileWriteSeconds)));
+			float(ParametersFilesWriteBytes) / (float(1024 * 1024) * float(ParametersFileWriteSeconds)));
 		UE_LOG(LogDumpGPU, Display, TEXT("Dumped resource binary: %.3f MB in %d files under %.3f s at %.3f MB/s"),
-			float(GRDGResourceDumpContext->ResourceBinaryWriteBytes) / float(1024 * 1024),
-			int32(GRDGResourceDumpContext->ResourceBinaryFilesOpened),
+			float(ResourceBinaryWriteBytes) / float(1024 * 1024),
+			int32(ResourceBinaryFilesOpened),
 			float(ResourceBinaryFileWriteSeconds),
-			float(GRDGResourceDumpContext->ResourceBinaryWriteBytes) / (float(1024 * 1024) * float(ResourceBinaryFileWriteSeconds)));
+			float(ResourceBinaryWriteBytes) / (float(1024 * 1024) * float(ResourceBinaryFileWriteSeconds)));
 		UE_LOG(LogDumpGPU, Display, TEXT("Dumped GPU readback resource release: %.3f s"), float(RHIReleaseResourcesTimeSeconds));
 	}
 
 	// Dump the log into the dump directory.
-	if (GRDGResourceDumpContext->bEnableDiskWrite)
+	if (bEnableDiskWrite)
 	{
 		if (GLog)
 		{
 			GLog->FlushThreadedLogs();
 			GLog->Flush();
 		}
-		FGenericCrashContext::DumpLog(GRDGResourceDumpContext->DumpingDirectoryPath / FRDGResourceDumpContext::kBaseDir);
+		FGenericCrashContext::DumpLog(DumpingDirectoryPath / FRDGResourceDumpContext::kBaseDir);
 	}
 
 	// Update the dump status to OK and removes subscription to crashes
 	{
-		GRDGResourceDumpContext->DumpStatusToFile(TEXT("ok"));
-		FCoreDelegates::OnShutdownAfterError.RemoveAll(GRDGResourceDumpContext);
+		DumpStatusToFile(TEXT("ok"));
+		FCoreDelegates::OnShutdownAfterError.RemoveAll(this);
 	}
 
-	if (GRDGResourceDumpContext->bUpload && IDumpGPUUploadServiceProvider::GProvider)
+	if (bUpload && IDumpGPUUploadServiceProvider::GProvider)
 	{
-		IDumpGPUUploadServiceProvider::FDumpParameters DumpCompletedParameters = GRDGResourceDumpContext->GetDumpParameters();
+		IDumpGPUUploadServiceProvider::FDumpParameters DumpCompletedParameters = GetDumpParameters();
 
 		// Compress the resource binary in background before uploading.
-		if (!GRDGResourceDumpContext->UploadResourceCompressionName.IsNone())
+		if (!UploadResourceCompressionName.IsNone())
 		{
-			DumpCompletedParameters.CompressionName = GRDGResourceDumpContext->UploadResourceCompressionName;
+			DumpCompletedParameters.CompressionName = UploadResourceCompressionName;
 			DumpCompletedParameters.CompressionFiles = FWildcardString(TEXT("*.bin"));
 		}
 
@@ -2427,26 +2378,19 @@ void FRDGBuilder::EndResourceDump()
 	}
 
 	#if PLATFORM_DESKTOP
-	if (GRDGResourceDumpContext->bShowInExplore)
+	if (bShowInExplore)
 	{
 		FPlatformProcess::ExploreFolder(*AbsDumpingDirectoryPath);
 	}
 	#endif
 
 	// Restore the engine tick rate to what it was
-	if (GRDGResourceDumpContext->bOverrideFixedDeltaTime)
+	if (bOverrideFixedDeltaTime)
 	{
-		FApp::SetFixedDeltaTime(GRDGResourceDumpContext->PreviousFixedDeltaTime);
+		FApp::SetFixedDeltaTime(PreviousFixedDeltaTime);
 		FApp::SetUseFixedTimeStep(false);
 	}
 
-	delete GRDGResourceDumpContext;
-	GRDGResourceDumpContext = nullptr;
-	DumpingFrameCounter_GameThread = 0;
-
-	// It matches SuspendHeartBeat from BeginResourceDump.
-	FThreadHeartBeat::Get().ResumeHeartBeat(true);
-	ResumeRenderThreadTimeout();
 }
 
 static const TCHAR* GetPassEventNameWithGPUMask(const FRDGPass* Pass, FString& OutNameStorage)
@@ -2474,13 +2418,15 @@ void FRDGBuilder::DumpResourcePassOutputs(const FRDGPass* Pass)
 		return;
 	}
 
-	if (!GRDGResourceDumpContext)
+	FRDGResourceDumpContext* ResourceDumpContext = GRDGResourceDumpContext_RenderThread;
+
+	if (!ResourceDumpContext)
 	{
 		return;
 	}
 
 	check(IsInRenderingThread());
-	if (!GRDGResourceDumpContext->IsDumpingPass(Pass))
+	if (!ResourceDumpContext->IsDumpingPass(Pass))
 	{
 		return;
 	}
@@ -2498,7 +2444,7 @@ void FRDGBuilder::DumpResourcePassOutputs(const FRDGPass* Pass)
 			if (FRDGTextureRef Texture = Parameter.GetAsTexture())
 			{
 				FRDGTextureSRVDesc TextureSubResource = FRDGTextureSRVDesc::Create(Texture);
-				GRDGResourceDumpContext->AddDumpTexturePasses(*this, InputResourceNames, OutputResourceNames, Pass, TextureSubResource, /* bIsOutputResource = */ false);
+				ResourceDumpContext->AddDumpTexturePasses(*this, InputResourceNames, OutputResourceNames, Pass, TextureSubResource, /* bIsOutputResource = */ false);
 			}
 		}
 		break;
@@ -2508,7 +2454,7 @@ void FRDGBuilder::DumpResourcePassOutputs(const FRDGPass* Pass)
 			{
 				if (SRV->Desc.MetaData == ERHITextureMetaDataAccess::None)
 				{
-					GRDGResourceDumpContext->AddDumpTexturePasses(*this, InputResourceNames, OutputResourceNames, Pass, SRV->Desc, /* bIsOutputResource = */ false);
+					ResourceDumpContext->AddDumpTexturePasses(*this, InputResourceNames, OutputResourceNames, Pass, SRV->Desc, /* bIsOutputResource = */ false);
 				}
 				else
 				{
@@ -2531,7 +2477,7 @@ void FRDGBuilder::DumpResourcePassOutputs(const FRDGPass* Pass)
 						TextureSubResource.FirstArraySlice = UAV->Desc.FirstArraySlice;
 						TextureSubResource.NumArraySlices = UAV->Desc.NumArraySlices;
 					}
-					GRDGResourceDumpContext->AddDumpTexturePasses(*this, InputResourceNames, OutputResourceNames, Pass, TextureSubResource, /* bIsOutputResource = */ true);
+					ResourceDumpContext->AddDumpTexturePasses(*this, InputResourceNames, OutputResourceNames, Pass, TextureSubResource, /* bIsOutputResource = */ true);
 				}
 				else
 				{
@@ -2547,7 +2493,7 @@ void FRDGBuilder::DumpResourcePassOutputs(const FRDGPass* Pass)
 				bool bIsOutputResource = IsWritableAccess(TextureAccess.GetAccess());
 
 				FRDGTextureSRVDesc TextureSubResource = FRDGTextureSRVDesc::Create(TextureAccess);
-				GRDGResourceDumpContext->AddDumpTexturePasses(*this, InputResourceNames, OutputResourceNames, Pass, TextureSubResource, bIsOutputResource);
+				ResourceDumpContext->AddDumpTexturePasses(*this, InputResourceNames, OutputResourceNames, Pass, TextureSubResource, bIsOutputResource);
 			}
 		}
 		break;
@@ -2560,7 +2506,7 @@ void FRDGBuilder::DumpResourcePassOutputs(const FRDGPass* Pass)
 				bool bIsOutputResource = IsWritableAccess(TextureAccess.GetAccess());
 
 				FRDGTextureSRVDesc TextureSubResource = FRDGTextureSRVDesc::Create(TextureAccess);
-				GRDGResourceDumpContext->AddDumpTexturePasses(*this, InputResourceNames, OutputResourceNames, Pass, TextureSubResource, bIsOutputResource);
+				ResourceDumpContext->AddDumpTexturePasses(*this, InputResourceNames, OutputResourceNames, Pass, TextureSubResource, bIsOutputResource);
 			}
 		}
 		break;
@@ -2570,7 +2516,7 @@ void FRDGBuilder::DumpResourcePassOutputs(const FRDGPass* Pass)
 			if (FRDGBufferSRVRef SRV = Parameter.GetAsBufferSRV())
 			{
 				FRDGBufferRef Buffer = SRV->Desc.Buffer;
-				GRDGResourceDumpContext->AddDumpBufferPass(*this, InputResourceNames, OutputResourceNames, Pass, Buffer, /* bIsOutputResource = */ false);
+				ResourceDumpContext->AddDumpBufferPass(*this, InputResourceNames, OutputResourceNames, Pass, Buffer, /* bIsOutputResource = */ false);
 			}
 		}
 		break;
@@ -2579,7 +2525,7 @@ void FRDGBuilder::DumpResourcePassOutputs(const FRDGPass* Pass)
 			if (FRDGBufferUAVRef UAV = Parameter.GetAsBufferUAV())
 			{
 				FRDGBufferRef Buffer = UAV->Desc.Buffer;
-				GRDGResourceDumpContext->AddDumpBufferPass(*this, InputResourceNames, OutputResourceNames, Pass, Buffer, /* bIsOutputResource = */ true);
+				ResourceDumpContext->AddDumpBufferPass(*this, InputResourceNames, OutputResourceNames, Pass, Buffer, /* bIsOutputResource = */ true);
 			}
 		}
 		break;
@@ -2589,7 +2535,7 @@ void FRDGBuilder::DumpResourcePassOutputs(const FRDGPass* Pass)
 			{
 				bool bIsOutputResource = IsWritableAccess(BufferAccess.GetAccess());
 
-				GRDGResourceDumpContext->AddDumpBufferPass(*this, InputResourceNames, OutputResourceNames, Pass, BufferAccess, bIsOutputResource);
+				ResourceDumpContext->AddDumpBufferPass(*this, InputResourceNames, OutputResourceNames, Pass, BufferAccess, bIsOutputResource);
 			}
 		}
 		break;
@@ -2601,7 +2547,7 @@ void FRDGBuilder::DumpResourcePassOutputs(const FRDGPass* Pass)
 			{
 				bool bIsOutputResource = IsWritableAccess(BufferAccess.GetAccess());
 
-				GRDGResourceDumpContext->AddDumpBufferPass(*this, InputResourceNames, OutputResourceNames, Pass, BufferAccess, bIsOutputResource);
+				ResourceDumpContext->AddDumpBufferPass(*this, InputResourceNames, OutputResourceNames, Pass, BufferAccess, bIsOutputResource);
 			}
 		}
 		break;
@@ -2621,7 +2567,7 @@ void FRDGBuilder::DumpResourcePassOutputs(const FRDGPass* Pass)
 					TextureSubResource.NumArraySlices = 1;
 				}
 
-				GRDGResourceDumpContext->AddDumpTexturePasses(*this, InputResourceNames, OutputResourceNames, Pass, TextureSubResource, /* bIsOutputResource = */ true);
+				ResourceDumpContext->AddDumpTexturePasses(*this, InputResourceNames, OutputResourceNames, Pass, TextureSubResource, /* bIsOutputResource = */ true);
 			});
 
 			const FDepthStencilBinding& DepthStencil = RenderTargets.DepthStencil;
@@ -2639,7 +2585,7 @@ void FRDGBuilder::DumpResourcePassOutputs(const FRDGPass* Pass)
 						TextureSubResource.NumArraySlices = 1;
 					}
 
-					GRDGResourceDumpContext->AddDumpTexturePasses(
+					ResourceDumpContext->AddDumpTexturePasses(
 						*this, InputResourceNames, OutputResourceNames, Pass, TextureSubResource,
 						/* bIsOutputResource = */ DepthStencilAccess.IsDepthWrite());
 				}
@@ -2653,7 +2599,7 @@ void FRDGBuilder::DumpResourcePassOutputs(const FRDGPass* Pass)
 						TextureSubResource.NumArraySlices = 1;
 					}
 
-					GRDGResourceDumpContext->AddDumpTexturePasses(
+					ResourceDumpContext->AddDumpTexturePasses(
 						*this, InputResourceNames, OutputResourceNames, Pass, TextureSubResource,
 						/* bIsOutputResource = */ DepthStencilAccess.IsStencilWrite());
 				}
@@ -2678,7 +2624,7 @@ void FRDGBuilder::DumpResourcePassOutputs(const FRDGPass* Pass)
 		}
 		#endif
 		{
-			ParentEventScopeNames.Add(MakeShareable(new FJsonValueString(FString::Printf(TEXT("Frame %llu (Delta=%fs)"), GFrameCounterRenderThread, GRDGResourceDumpContext->DeltaTime))));
+			ParentEventScopeNames.Add(MakeShareable(new FJsonValueString(FString::Printf(TEXT("Frame %llu (Delta=%fs)"), GFrameCounterRenderThread, ResourceDumpContext->DeltaTime))));
 		}
 
 		FString EventNameStorage;
@@ -2689,12 +2635,12 @@ void FRDGBuilder::DumpResourcePassOutputs(const FRDGPass* Pass)
 		JsonObject->SetStringField(TEXT("Parameters"), FRDGResourceDumpContext::PtrToString(Pass->GetParameters().GetContents()));
 		JsonObject->SetStringField(TEXT("ParametersMetadata"), FRDGResourceDumpContext::PtrToString(Pass->GetParameters().GetMetadata()));
 		JsonObject->SetStringField(TEXT("Pointer"), FString::Printf(TEXT("%016x"), FRDGResourceDumpContext::PtrToUint(Pass)));
-		JsonObject->SetNumberField(TEXT("Id"), GRDGResourceDumpContext->PassesCount);
+		JsonObject->SetNumberField(TEXT("Id"), ResourceDumpContext->PassesCount);
 		JsonObject->SetArrayField(TEXT("ParentEventScopes"), ParentEventScopeNames);
 		JsonObject->SetArrayField(TEXT("InputResources"), InputResourceNames);
 		JsonObject->SetArrayField(TEXT("OutputResources"), OutputResourceNames);
 
-		GRDGResourceDumpContext->DumpJsonToFile(JsonObject, FString(FRDGResourceDumpContext::kBaseDir) / TEXT("Passes.json"), FILEWRITE_Append);
+		ResourceDumpContext->DumpJsonToFile(JsonObject, FString(FRDGResourceDumpContext::kBaseDir) / TEXT("Passes.json"), FILEWRITE_Append);
 	}
 
 	// Dump the pass' parameters
@@ -2705,7 +2651,7 @@ void FRDGBuilder::DumpResourcePassOutputs(const FRDGPass* Pass)
 			const FShaderParametersMetadata* Metadata = Pass->GetParameters().GetMetadata();
 			if (Metadata)
 			{
-				GRDGResourceDumpContext->Dump(Metadata);
+				ResourceDumpContext->Dump(Metadata);
 				PassParametersByteSize = Metadata->GetSize();
 			}
 		}
@@ -2716,27 +2662,29 @@ void FRDGBuilder::DumpResourcePassOutputs(const FRDGPass* Pass)
 		}
 
 		const uint8* PassParametersContent = Pass->GetParameters().GetContents();
-		if (PassParametersContent && !GRDGResourceDumpContext->IsDumped(PassParametersContent))
+		if (PassParametersContent && !ResourceDumpContext->IsDumped(PassParametersContent))
 		{
 			TArrayView<const uint8> ArrayView(PassParametersContent, PassParametersByteSize);
 			FString DumpFilePath = FRDGResourceDumpContext::kStructuresDir / FRDGResourceDumpContext::PtrToString(PassParametersContent) + TEXT(".bin");
-			GRDGResourceDumpContext->DumpBinaryToFile(ArrayView, DumpFilePath, FRDGResourceDumpContext::ETimingBucket::ParametersFileWrite);
-			GRDGResourceDumpContext->SetDumped(PassParametersContent);
+			ResourceDumpContext->DumpBinaryToFile(ArrayView, DumpFilePath, FRDGResourceDumpContext::ETimingBucket::ParametersFileWrite);
+			ResourceDumpContext->SetDumped(PassParametersContent);
 		}
 	}
 
-	GRDGResourceDumpContext->PassesCount++;
+	ResourceDumpContext->PassesCount++;
 }
 
 #if RDG_DUMP_RESOURCES_AT_EACH_DRAW
 
 void FRDGBuilder::BeginPassDump(const FRDGPass* Pass)
 {
-	if (!GRDGResourceDumpContext)
+	if (!GRDGResourceDumpContext_RenderThread)
 	{
 		return;
 	}
 
+	FRDGResourceDumpContext* ResourceDumpContext = GRDGResourceDumpContext_RenderThread;
+
 	if (!GDumpGPUDraws.GetValueOnRenderThread())
 	{
 		return;
@@ -2753,19 +2701,19 @@ void FRDGBuilder::BeginPassDump(const FRDGPass* Pass)
 		return;
 	}
 
-	check(GRDGResourceDumpContext->DrawDumpingPass == nullptr);
+	check(ResourceDumpContext->DrawDumpingPass == nullptr);
 
-	if (GRDGResourceDumpContext->IsDumpingPass(Pass))
+	if (ResourceDumpContext->IsDumpingPass(Pass))
 	{
-		GRDGResourceDumpContext->DrawDumpingPass = Pass;
-		GRDGResourceDumpContext->DrawDumpCount = 0;
+		ResourceDumpContext->DrawDumpingPass = Pass;
+		ResourceDumpContext->DrawDumpCount = 0;
 	}
 }
 
 // static
 void FRDGBuilder::DumpDraw(const FRDGEventName& DrawEventName)
 {
-	if (!GRDGResourceDumpContext)
+	if (!GRDGResourceDumpContext_RenderThread)
 	{
 		return;
 	}
@@ -2776,12 +2724,14 @@ void FRDGBuilder::DumpDraw(const FRDGEventName& DrawEventName)
 		return;
 	}
 
-	if (!GRDGResourceDumpContext->DrawDumpingPass)
+	FRDGResourceDumpContext* ResourceDumpContext = GRDGResourceDumpContext_RenderThread;
+
+	if (!ResourceDumpContext->DrawDumpingPass)
 	{
 		return;
 	}
 
-	const FRDGPass* Pass = GRDGResourceDumpContext->DrawDumpingPass;
+	const FRDGPass* Pass = ResourceDumpContext->DrawDumpingPass;
 
 	FRHICommandListImmediate& RHICmdList = FRHICommandListExecutor::GetImmediateCommandList();
 
@@ -2802,7 +2752,7 @@ void FRDGBuilder::DumpDraw(const FRDGEventName& DrawEventName)
 			{
 				FRDGTextureRef Texture = RenderTarget.GetTexture();
 				FRDGTextureSRVDesc TextureSubResource = FRDGTextureSRVDesc::CreateForMipLevel(Texture, RenderTarget.GetMipIndex());
-				GRDGResourceDumpContext->DumpDrawTextureSubResource(
+				ResourceDumpContext->DumpDrawTextureSubResource(
 					RHICmdList,
 					TextureSubResource,
 					ERHIAccess::RTV);
@@ -2817,7 +2767,7 @@ void FRDGBuilder::DumpDraw(const FRDGEventName& DrawEventName)
 				if (DepthStencilAccess.IsDepthWrite())
 				{
 					FRDGTextureSRVDesc TextureSubResource = FRDGTextureSRVDesc::CreateForMipLevel(Texture, 0);
-					GRDGResourceDumpContext->DumpDrawTextureSubResource(
+					ResourceDumpContext->DumpDrawTextureSubResource(
 						RHICmdList,
 						TextureSubResource,
 						ERHIAccess::RTV);
@@ -2826,7 +2776,7 @@ void FRDGBuilder::DumpDraw(const FRDGEventName& DrawEventName)
 				if (DepthStencilAccess.IsStencilWrite())
 				{
 					FRDGTextureSRVDesc TextureSubResource = FRDGTextureSRVDesc::CreateWithPixelFormat(Texture, PF_X24_G8);
-					GRDGResourceDumpContext->DumpDrawTextureSubResource(
+					ResourceDumpContext->DumpDrawTextureSubResource(
 						RHICmdList,
 						TextureSubResource,
 						ERHIAccess::RTV);
@@ -2848,21 +2798,21 @@ void FRDGBuilder::DumpDraw(const FRDGEventName& DrawEventName)
 		JsonObject->SetStringField(TEXT("DrawName"), DrawEventName.GetTCHAR());
 
 		FString DumpFilePath = FRDGResourceDumpContext::kPassesDir / FString::Printf(TEXT("Pass.%016x.Draws.json"), FRDGResourceDumpContext::PtrToUint(Pass));
-		GRDGResourceDumpContext->DumpJsonToFile(JsonObject, DumpFilePath, FILEWRITE_Append);
+		ResourceDumpContext->DumpJsonToFile(JsonObject, DumpFilePath, FILEWRITE_Append);
 	}
 
-	GRDGResourceDumpContext->DrawDumpCount++;
+	ResourceDumpContext->DrawDumpCount++;
 
-	if (GRDGResourceDumpContext->DrawDumpCount % 10 == 0)
+	if (ResourceDumpContext->DrawDumpCount % 10 == 0)
 	{
-		UE_LOG(LogDumpGPU, Display, TEXT("Dumped %d draws' resources"), GRDGResourceDumpContext->DrawDumpCount);
+		UE_LOG(LogDumpGPU, Display, TEXT("Dumped %d draws' resources"), ResourceDumpContext->DrawDumpCount);
 		return;
 	}
 }
 
 void FRDGBuilder::EndPassDump(const FRDGPass* Pass)
 {
-	if (!GRDGResourceDumpContext)
+	if (!GRDGResourceDumpContext_RenderThread)
 	{
 		return;
 	}
@@ -2872,41 +2822,43 @@ void FRDGBuilder::EndPassDump(const FRDGPass* Pass)
 		return;
 	}
 
-	if (!GRDGResourceDumpContext->DrawDumpingPass)
+	FRDGResourceDumpContext* ResourceDumpContext = GRDGResourceDumpContext_RenderThread;
+
+	if (!ResourceDumpContext->DrawDumpingPass)
 	{
 		return;
 	}
 
-	check(Pass == GRDGResourceDumpContext->DrawDumpingPass);
+	check(Pass == ResourceDumpContext->DrawDumpingPass);
 
 	// Output how many draw has been dump for this pass.
-	if (GRDGResourceDumpContext->DrawDumpCount > 0)
+	if (ResourceDumpContext->DrawDumpCount > 0)
 	{
 		FString EventNameStorage;
 
 		TSharedPtr<FJsonObject> JsonObject = MakeShareable(new FJsonObject);
 		JsonObject->SetStringField(TEXT("EventName"), GetPassEventNameWithGPUMask(Pass, EventNameStorage));
 		JsonObject->SetStringField(TEXT("Pointer"), FString::Printf(TEXT("%016x"), FRDGResourceDumpContext::PtrToUint(Pass)));
-		JsonObject->SetNumberField(TEXT("DrawCount"), GRDGResourceDumpContext->DrawDumpCount);
+		JsonObject->SetNumberField(TEXT("DrawCount"), ResourceDumpContext->DrawDumpCount);
 
-		GRDGResourceDumpContext->DumpJsonToFile(JsonObject, FString(FRDGResourceDumpContext::kBaseDir) / TEXT("PassDrawCounts.json"), FILEWRITE_Append);
+		ResourceDumpContext->DumpJsonToFile(JsonObject, FString(FRDGResourceDumpContext::kBaseDir) / TEXT("PassDrawCounts.json"), FILEWRITE_Append);
 
-		UE_LOG(LogDumpGPU, Display, TEXT("Completed dump of %d draws for pass: %s"), GRDGResourceDumpContext->DrawDumpCount, Pass->GetEventName().GetTCHAR());
+		UE_LOG(LogDumpGPU, Display, TEXT("Completed dump of %d draws for pass: %s"), ResourceDumpContext->DrawDumpCount, Pass->GetEventName().GetTCHAR());
 	}
 
-	GRDGResourceDumpContext->DrawDumpingPass = nullptr;
-	GRDGResourceDumpContext->DrawDumpCount = 0;
+	ResourceDumpContext->DrawDumpingPass = nullptr;
+	ResourceDumpContext->DrawDumpCount = 0;
 }
 
 // static
 bool FRDGBuilder::IsDumpingFrame()
 {
-	return GRDGResourceDumpContext != nullptr;
+	return GRDGResourceDumpContext_RenderThread != nullptr;
 }
 
 bool FRDGBuilder::IsDumpingDraws()
 {
-	if (!GRDGResourceDumpContext)
+	if (!FRDGBuilder::IsDumpingFrame())
 	{
 		return false;
 	}
@@ -2916,11 +2868,104 @@ bool FRDGBuilder::IsDumpingDraws()
 
 #endif // RDG_DUMP_RESOURCES_AT_EACH_DRAW
 
-#else //! RDG_DUMP_RESOURCES
 
-bool IsDumpingRDGResources()
+
+namespace UE::RenderCore::DumpGPU
+{
+
+void TickEndFrame()
 {
+	check(IsInGameThread());
+
+	if (GRDGResourceDumpContext_GameThread)
+	{
+		check(GNextRDGResourceDumpContext == nullptr);
+
+		GRDGResourceDumpContext_GameThread->DumpedFrameId++;
+		check(GRDGResourceDumpContext_GameThread->DumpedFrameId <= GRDGResourceDumpContext_GameThread->FrameCount);
+
+		if (GRDGResourceDumpContext_GameThread->DumpedFrameId < GRDGResourceDumpContext_GameThread->FrameCount)
+		{
+			int32 RemainingFrameCount = GRDGResourceDumpContext_GameThread->FrameCount - GRDGResourceDumpContext_GameThread->DumpedFrameId;
+			ENQUEUE_RENDER_COMMAND(FDumpGPULogRemaingFrames)(
+				[RemainingFrameCount](FRHICommandListImmediate& ImmediateRHICmdList)
+			{
+				UE_LOG(LogDumpGPU, Display, TEXT("Remaining frames %d"), RemainingFrameCount);
+
+				if (GRDGResourceDumpContext_RenderThread->bStream)
+				{
+					GRDGResourceDumpContext_RenderThread->LandCompletedResources(ImmediateRHICmdList);
+				}
+			});
+		}
+		else
+		{
+			GRDGResourceDumpContext_GameThread->Finish();
+			delete GRDGResourceDumpContext_GameThread;
+			GRDGResourceDumpContext_GameThread = nullptr;
+			GRDGResourceDumpContext_RenderThread = nullptr;
+
+			// It matches SuspendHeartBeat from BeginResourceDump.
+			FThreadHeartBeat::Get().ResumeHeartBeat(true);
+			ResumeRenderThreadTimeout();
+		}
+	}
+	else if (GNextRDGResourceDumpContext)
+	{
+		check(GRDGResourceDumpContext_GameThread == nullptr);
+
+		if (GNextDumpingRemainingTime > 0.0f)
+		{
+			check(GNextRDGResourceDumpContext);
+			GNextDumpingRemainingTime -= FApp::GetDeltaTime();
+			if (GNextDumpingRemainingTime <= 0.0)
+			{
+				GNextDumpingRemainingTime = -1.0f;
+				GNextRDGResourceDumpContext->Start();
+			}
+		}
+		else
+		{
+			GNextRDGResourceDumpContext->Start();
+		}
+	}
+}
+
+bool IsDumpingFrame()
+{
+	check(IsInGameThread() || IsInRenderingThread());
+
+	if (IsInGameThread())
+	{
+		return GRDGResourceDumpContext_GameThread != nullptr;
+	}
+	else if (IsInRenderingThread())
+	{
+		return GRDGResourceDumpContext_RenderThread != nullptr;
+	}
+	else
+	{
+		check(0);
+	}
 	return false;
 }
 
-#endif //! RDG_DUMP_RESOURCES
+bool ShouldCameraCut()
+{
+	check(IsInGameThread());
+	if (GRDGResourceDumpContext_GameThread == nullptr)
+	{
+		return false;
+	}
+
+	if (GDumpGPUCameraCut.GetValueOnGameThread() == 0)
+	{
+		return false;
+	}
+
+	return GRDGResourceDumpContext_GameThread->DumpedFrameId == 0;
+}
+
+} // namespace UE::RenderCore::DumpGPU
+
+#endif // RDG_DUMP_RESOURCES
diff --git a/Engine/Source/Runtime/RenderCore/Private/RenderGraphPrivate.cpp b/Engine/Source/Runtime/RenderCore/Private/RenderGraphPrivate.cpp
--- a/Engine/Source/Runtime/RenderCore/Private/RenderGraphPrivate.cpp
+++ b/Engine/Source/Runtime/RenderCore/Private/RenderGraphPrivate.cpp
@@ -4,6 +4,7 @@
 #include "DataDrivenShaderPlatformInfo.h"
 #include "Misc/CommandLine.h"
 #include "RHICommandList.h"
+#include "DumpGPU.h"
 
 #if RDG_ENABLE_DEBUG
 
@@ -638,7 +639,7 @@ bool IsParallelExecuteEnabled()
 		&& !IsVulkanMobileSM5Platform(GMaxRHIShaderPlatform)
 		&& GRHISupportsMultithreadedShaderCreation
 #if WITH_DUMPGPU
-		&& !IsDumpingRDGResources()
+		&& !UE::RenderCore::DumpGPU::IsDumpingFrame()
 #endif
 		// Only run parallel RDG if we have a rendering thread.
 		&& IsInActualRenderingThread()
@@ -657,7 +658,7 @@ bool IsParallelSetupEnabled()
 		&& !IsVulkanMobileSM5Platform(GMaxRHIShaderPlatform)
 		&& GRHISupportsMultithreadedShaderCreation
 #if WITH_DUMPGPU
-		&& !IsDumpingRDGResources()
+		&& !UE::RenderCore::DumpGPU::IsDumpingFrame()
 #endif
 		// Only run parallel RDG if we have a rendering thread.
 		&& IsInActualRenderingThread()
diff --git a/Engine/Source/Runtime/RenderCore/Private/RenderGraphPrivate.h b/Engine/Source/Runtime/RenderCore/Private/RenderGraphPrivate.h
--- a/Engine/Source/Runtime/RenderCore/Private/RenderGraphPrivate.h
+++ b/Engine/Source/Runtime/RenderCore/Private/RenderGraphPrivate.h
@@ -260,8 +260,6 @@ FORCEINLINE bool IsAsyncComputeSupported()
 	return GRDGAsyncCompute > 0 && !IsImmediateMode() && GSupportsEfficientAsyncCompute && GRHISupportsSeparateDepthStencilCopyAccess && !GTriggerGPUProfile;
 }
 
-bool IsDumpingRDGResources();
-
 extern bool IsParallelExecuteEnabled();
 extern bool IsParallelSetupEnabled();
 
diff --git a/Engine/Source/Runtime/RenderCore/Public/DumpGPU.h b/Engine/Source/Runtime/RenderCore/Public/DumpGPU.h
--- a/Engine/Source/Runtime/RenderCore/Public/DumpGPU.h
+++ b/Engine/Source/Runtime/RenderCore/Public/DumpGPU.h
@@ -2,10 +2,10 @@
 
 #pragma once
 
+#include "Misc/Build.h"
 #include "Misc/WildcardString.h"
 #include "UObject/NameTypes.h"
 
-
 class IDumpGPUUploadServiceProvider
 {
 public:
@@ -30,3 +30,18 @@ public:
 
 	static RENDERCORE_API IDumpGPUUploadServiceProvider* GProvider;
 };
+
+#if WITH_DUMPGPU
+
+namespace UE::RenderCore::DumpGPU
+{
+
+RENDERCORE_API void TickEndFrame();
+
+RENDERCORE_API bool IsDumpingFrame();
+
+RENDERCORE_API bool ShouldCameraCut();
+
+} // namespace  UE::RenderCore::DumpGPU
+
+#endif // WITH_DUMPGPU
diff --git a/Engine/Source/Runtime/RenderCore/Public/RenderGraphBuilder.h b/Engine/Source/Runtime/RenderCore/Public/RenderGraphBuilder.h
--- a/Engine/Source/Runtime/RenderCore/Public/RenderGraphBuilder.h
+++ b/Engine/Source/Runtime/RenderCore/Public/RenderGraphBuilder.h
@@ -388,8 +388,6 @@ public:
 
 #if RDG_DUMP_RESOURCES
 	static RENDERCORE_API FString BeginResourceDump(const TCHAR* Cmd);
-	static RENDERCORE_API void InitResourceDump();
-	static RENDERCORE_API void EndResourceDump();
 	static RENDERCORE_API bool IsDumpingFrame();
 #else
 	static bool IsDumpingFrame() { return false; }
diff --git a/Engine/Source/Runtime/Renderer/Private/SceneRendering.cpp b/Engine/Source/Runtime/Renderer/Private/SceneRendering.cpp
--- a/Engine/Source/Runtime/Renderer/Private/SceneRendering.cpp
+++ b/Engine/Source/Runtime/Renderer/Private/SceneRendering.cpp
@@ -15,6 +15,7 @@
 #include "Components/SceneCaptureComponent2D.h"
 #include "Components/SceneCaptureComponentCube.h"
 #include "DeferredShadingRenderer.h"
+#include "DumpGPU.h"
 #include "DynamicPrimitiveDrawing.h"
 #include "RenderTargetTemp.h"
 #include "RendererModule.h"
@@ -2568,6 +2569,13 @@ FSceneRenderer::FSceneRenderer(const FSceneViewFamily* InViewFamily, FHitProxyCo
 		}
 #endif
 
+#if WITH_DUMPGPU
+		if (UE::RenderCore::DumpGPU::ShouldCameraCut())
+		{
+			ViewInfo->bCameraCut = true;
+		}
+#endif
+
 		if (ShouldRenderLumenDiffuseGI(Scene, *ViewInfo) || 
 			ShouldRenderLumenReflections(*ViewInfo) || 
 			ShouldRenderVolumetricCloudWithBlueNoise_GameThread(Scene, *ViewInfo) || 
