Deduplicates RDG resources pointer across graph builders in DumpGPU to handle r.DumpGPU.FrameCount with many many frames

diff --git a/Engine/Extras/GPUDumpViewer/GPUDumpViewer.html b/Engine/Extras/GPUDumpViewer/GPUDumpViewer.html
--- a/Engine/Extras/GPUDumpViewer/GPUDumpViewer.html
+++ b/Engine/Extras/GPUDumpViewer/GPUDumpViewer.html
@@ -13,7 +13,7 @@ const k_style_scroll_width = 8;
 const k_style_padding_size = 10;
 const k_style_main_div_margin = 15;
 
-const k_null_json_ptr = "0000000000000000";
+const k_null_json_ptr = "00000000000000000000";
 
 
 // -------------------------------------------------------------------- GLOBALS
diff --git a/Engine/Source/Runtime/RenderCore/Private/DumpGPU.cpp b/Engine/Source/Runtime/RenderCore/Private/DumpGPU.cpp
--- a/Engine/Source/Runtime/RenderCore/Private/DumpGPU.cpp
+++ b/Engine/Source/Runtime/RenderCore/Private/DumpGPU.cpp
@@ -246,6 +246,7 @@ public:
 	FGenericPlatformMemoryStats MemoryStats;
 	int32 ResourcesDumpPasses = 0;
 	int32 ResourcesDumpExecutedPasses = 0;
+	uint16 GraphBuilderIndex = 0;
 	int32 PassesCount = 0;
 	TMap<const FRDGResource*, const FRDGPass*> LastResourceVersion;
 	TSet<const void*> IsDumpedToDisk;
@@ -477,22 +478,20 @@ public:
 	}
 
 	template<typename T>
-	static uint64 PtrToUint(const T* Ptr)
+	FString PtrToString(const T* Ptr)
 	{
-		return static_cast<uint64>(reinterpret_cast<size_t>(Ptr));
-	}
-
-	template<typename T>
-	static FString PtrToString(const T* Ptr)
-	{
-		return FString::Printf(TEXT("%016x"), PtrToUint(Ptr));
+		if (Ptr == nullptr)
+		{
+			return TEXT("00000000000000000000");
+		}
+		return FString::Printf(TEXT("%04u%016x"), uint32(GraphBuilderIndex), static_cast<uint64>(reinterpret_cast<size_t>(Ptr)));
 	}
 
-	static FString GetUniqueResourceName(const FRDGResource* Resource)
+	FString GetUniqueResourceName(const FRDGResource* Resource)
 	{
 		if (GDumpTestPrettifyResourceFileNames.GetValueOnRenderThread())
 		{
-			FString UniqueResourceName = FString::Printf(TEXT("%s.%016x"), Resource->Name, PtrToUint(Resource));
+			FString UniqueResourceName = FString::Printf(TEXT("%s.%s"), Resource->Name, *PtrToString(Resource));
 			UniqueResourceName.ReplaceInline(TEXT("/"), TEXT(""));
 			UniqueResourceName.ReplaceInline(TEXT("\\"), TEXT(""));
 			return UniqueResourceName;
@@ -500,7 +499,7 @@ public:
 		return PtrToString(Resource);
 	}
 
-	static FString GetUniqueSubResourceName(const FRDGTextureSRVDesc& SubResourceDesc)
+	FString GetUniqueSubResourceName(const FRDGTextureSRVDesc& SubResourceDesc)
 	{
 		check(SubResourceDesc.NumMipLevels == 1);
 		check(!SubResourceDesc.Texture->Desc.IsTextureArray() || SubResourceDesc.NumArraySlices == 1);
@@ -1319,9 +1318,9 @@ public:
 		}
 
 		FString DumpFilePath = kResourcesDir / FString::Printf(
-			TEXT("%s.v%016x.d%d.bin"),
+			TEXT("%s.v%s.d%d.bin"),
 			*UniqueResourceSubResourceName,
-			PtrToUint(DrawDumpingPass),
+			*PtrToString(DrawDumpingPass),
 			DrawDumpCount);
 
 		DumpTextureSubResource(
@@ -1409,7 +1408,7 @@ public:
 
 		// Dump the resource's binary to a .bin file.
 		{
-			FString DumpFilePath = kResourcesDir / FString::Printf(TEXT("%s.v%016x.bin"), *UniqueResourceSubResourceName, PtrToUint(bIsOutputResource ? Pass : nullptr));
+			FString DumpFilePath = kResourcesDir / FString::Printf(TEXT("%s.v%s.bin"), *UniqueResourceSubResourceName, *PtrToString(bIsOutputResource ? Pass : nullptr));
 
 			FDumpTexturePass* PassParameters = GraphBuilder.AllocParameters<FDumpTexturePass>();
 			if (SubresourceDumpDesc.bPreprocessForStaging)
@@ -1683,7 +1682,7 @@ public:
 		{
 			const int32 StagingResourceByteSize = bStream ? ByteSize : FMath::Min(ByteSize, GDumpMaxStagingSize.GetValueOnRenderThread() * 1024 * 1024);
 
-			FString DumpFilePath = kResourcesDir / FString::Printf(TEXT("%s.v%016x.bin"), *UniqueResourceName, PtrToUint(bIsOutputResource ? Pass : nullptr));
+			FString DumpFilePath = kResourcesDir / FString::Printf(TEXT("%s.v%s.bin"), *UniqueResourceName, *PtrToString(bIsOutputResource ? Pass : nullptr));
 
 			if (IsUnsafeToDumpResource(StagingResourceByteSize, 1.2f))
 			{
@@ -2411,6 +2410,20 @@ static const TCHAR* GetPassEventNameWithGPUMask(const FRDGPass* Pass, FString& O
 	}
 }
 
+void FRDGBuilder::DumpNewGraphBuilder()
+{
+	FRDGResourceDumpContext* ResourceDumpContext = GRDGResourceDumpContext_RenderThread;
+	if (!ResourceDumpContext)
+	{
+		return;
+	}
+
+	check(IsInRenderingThread());
+	ResourceDumpContext->GraphBuilderIndex++;
+	ResourceDumpContext->LastResourceVersion.Empty();
+	ResourceDumpContext->IsDumpedToDisk.Empty();
+}
+
 void FRDGBuilder::DumpResourcePassOutputs(const FRDGPass* Pass)
 {
 	if (!AuxiliaryPasses.IsDumpAllowed())
@@ -2632,9 +2645,9 @@ void FRDGBuilder::DumpResourcePassOutputs(const FRDGPass* Pass)
 		TSharedPtr<FJsonObject> JsonObject = MakeShareable(new FJsonObject);
 		JsonObject->SetStringField(TEXT("EventName"), GetPassEventNameWithGPUMask(Pass, EventNameStorage));
 		JsonObject->SetStringField(TEXT("ParametersName"), Pass->GetParameters().GetLayout().GetDebugName());
-		JsonObject->SetStringField(TEXT("Parameters"), FRDGResourceDumpContext::PtrToString(Pass->GetParameters().GetContents()));
-		JsonObject->SetStringField(TEXT("ParametersMetadata"), FRDGResourceDumpContext::PtrToString(Pass->GetParameters().GetMetadata()));
-		JsonObject->SetStringField(TEXT("Pointer"), FString::Printf(TEXT("%016x"), FRDGResourceDumpContext::PtrToUint(Pass)));
+		JsonObject->SetStringField(TEXT("Parameters"), ResourceDumpContext->PtrToString(Pass->GetParameters().GetContents()));
+		JsonObject->SetStringField(TEXT("ParametersMetadata"), ResourceDumpContext->PtrToString(Pass->GetParameters().GetMetadata()));
+		JsonObject->SetStringField(TEXT("Pointer"), ResourceDumpContext->PtrToString(Pass));
 		JsonObject->SetNumberField(TEXT("Id"), ResourceDumpContext->PassesCount);
 		JsonObject->SetArrayField(TEXT("ParentEventScopes"), ParentEventScopeNames);
 		JsonObject->SetArrayField(TEXT("InputResources"), InputResourceNames);
@@ -2665,7 +2678,7 @@ void FRDGBuilder::DumpResourcePassOutputs(const FRDGPass* Pass)
 		if (PassParametersContent && !ResourceDumpContext->IsDumped(PassParametersContent))
 		{
 			TArrayView<const uint8> ArrayView(PassParametersContent, PassParametersByteSize);
-			FString DumpFilePath = FRDGResourceDumpContext::kStructuresDir / FRDGResourceDumpContext::PtrToString(PassParametersContent) + TEXT(".bin");
+			FString DumpFilePath = FRDGResourceDumpContext::kStructuresDir / ResourceDumpContext->PtrToString(PassParametersContent) + TEXT(".bin");
 			ResourceDumpContext->DumpBinaryToFile(ArrayView, DumpFilePath, FRDGResourceDumpContext::ETimingBucket::ParametersFileWrite);
 			ResourceDumpContext->SetDumped(PassParametersContent);
 		}
@@ -2797,7 +2810,7 @@ void FRDGBuilder::DumpDraw(const FRDGEventName& DrawEventName)
 		TSharedPtr<FJsonObject> JsonObject = MakeShareable(new FJsonObject);
 		JsonObject->SetStringField(TEXT("DrawName"), DrawEventName.GetTCHAR());
 
-		FString DumpFilePath = FRDGResourceDumpContext::kPassesDir / FString::Printf(TEXT("Pass.%016x.Draws.json"), FRDGResourceDumpContext::PtrToUint(Pass));
+		FString DumpFilePath = FRDGResourceDumpContext::kPassesDir / FString::Printf(TEXT("Pass.%s.Draws.json"), *ResourceDumpContext->PtrToString(Pass));
 		ResourceDumpContext->DumpJsonToFile(JsonObject, DumpFilePath, FILEWRITE_Append);
 	}
 
@@ -2838,7 +2851,7 @@ void FRDGBuilder::EndPassDump(const FRDGPass* Pass)
 
 		TSharedPtr<FJsonObject> JsonObject = MakeShareable(new FJsonObject);
 		JsonObject->SetStringField(TEXT("EventName"), GetPassEventNameWithGPUMask(Pass, EventNameStorage));
-		JsonObject->SetStringField(TEXT("Pointer"), FString::Printf(TEXT("%016x"), FRDGResourceDumpContext::PtrToUint(Pass)));
+		JsonObject->SetStringField(TEXT("Pointer"), ResourceDumpContext->PtrToString(Pass));
 		JsonObject->SetNumberField(TEXT("DrawCount"), ResourceDumpContext->DrawDumpCount);
 
 		ResourceDumpContext->DumpJsonToFile(JsonObject, FString(FRDGResourceDumpContext::kBaseDir) / TEXT("PassDrawCounts.json"), FILEWRITE_Append);
diff --git a/Engine/Source/Runtime/RenderCore/Private/RenderGraphBuilder.cpp b/Engine/Source/Runtime/RenderCore/Private/RenderGraphBuilder.cpp
--- a/Engine/Source/Runtime/RenderCore/Private/RenderGraphBuilder.cpp
+++ b/Engine/Source/Runtime/RenderCore/Private/RenderGraphBuilder.cpp
@@ -595,6 +595,10 @@ FRDGBuilder::FRDGBuilder(FRHICommandListImmediate& InRHICmdList, FRDGEventName I
 		BreadcrumbState = FRDGBreadcrumbState::Create(Allocator);
 	}
 #endif
+
+#if RDG_DUMP_RESOURCES
+	DumpNewGraphBuilder();
+#endif
 }
 
 FRDGBuilder::~FRDGBuilder()
diff --git a/Engine/Source/Runtime/RenderCore/Public/RenderGraphBuilder.h b/Engine/Source/Runtime/RenderCore/Public/RenderGraphBuilder.h
--- a/Engine/Source/Runtime/RenderCore/Public/RenderGraphBuilder.h
+++ b/Engine/Source/Runtime/RenderCore/Public/RenderGraphBuilder.h
@@ -828,6 +828,7 @@ private:
 	RENDERCORE_API FRDGSubresourceState* AllocSubresource(const FRDGSubresourceState& Other);
 
 #if RDG_DUMP_RESOURCES
+	RENDERCORE_API void DumpNewGraphBuilder();
 	RENDERCORE_API void DumpResourcePassOutputs(const FRDGPass* Pass);
 
 #if RDG_DUMP_RESOURCES_AT_EACH_DRAW
