Don't produce errors when accessing shader compile error file paths that are non-virtual paths.

#rnx

diff --git a/Engine/Source/Runtime/RenderCore/Private/ShaderCore.cpp b/Engine/Source/Runtime/RenderCore/Private/ShaderCore.cpp
--- a/Engine/Source/Runtime/RenderCore/Private/ShaderCore.cpp
+++ b/Engine/Source/Runtime/RenderCore/Private/ShaderCore.cpp
@@ -3006,7 +3006,9 @@ bool FShaderCompilerError::ExtractSourceLocation()
 
 FString FShaderCompilerError::GetShaderSourceFilePath() const
 {
-	if (IFileManager::Get().FileExists(*ErrorVirtualFilePath))
+	// Always return error file path as-is if it doesn't denote a virtual path.
+	// We don't wont to report errors when accessing a compile error's message.
+	if (ErrorVirtualFilePath.IsEmpty() || ErrorVirtualFilePath[0] != TEXT('/'))
 	{
 		return ErrorVirtualFilePath;
 	}
