Shader Core compile error fix for certain platforms:  need to use _mm_loadu/storeu_si128 (SSE2) instead of _mm_loadu_epi16/_mm_storeu_epi8 (AVX specific) -- the former are identical when a write mask isn't in use.

diff --git a/Engine/Source/Runtime/RenderCore/Private/ShaderCore.cpp b/Engine/Source/Runtime/RenderCore/Private/ShaderCore.cpp
--- a/Engine/Source/Runtime/RenderCore/Private/ShaderCore.cpp
+++ b/Engine/Source/Runtime/RenderCore/Private/ShaderCore.cpp
@@ -1864,8 +1864,8 @@ void ShaderConvertAndStripComments(const FString& ShaderSource, TArray<ANSICHAR>
 	const TCHAR* EndSse = End - 16;
 	for (; Current < EndSse; )
 	{
-		__m128i First8 = _mm_loadu_epi16(Current);
-		__m128i Second8 = _mm_loadu_epi16(Current + 8);
+		__m128i First8 = _mm_loadu_si128((const __m128i*)Current);
+		__m128i Second8 = _mm_loadu_si128((const __m128i*)(Current + 8));
 		__m128i CurrentWord = _mm_packus_epi16(First8, Second8);
 
 		int32 CRMask = _mm_movemask_epi8(_mm_cmpeq_epi8(CurrentWord, CharCR));
@@ -1886,7 +1886,7 @@ void ShaderConvertAndStripComments(const FString& ShaderSource, TArray<ANSICHAR>
 		}
 
 		// Echo the current word
-		_mm_storeu_epi8(CurrentOut, CurrentWord);
+		_mm_storeu_si128((__m128i*)CurrentOut, CurrentWord);
 
 		// Check if there is a comment start, meaning a slash followed by slash or star, which we can detect by shifting right
 		// a mask containing both slash and star, and seeing if that overlaps with a slash.
@@ -1911,8 +1911,8 @@ void ShaderConvertAndStripComments(const FString& ShaderSource, TArray<ANSICHAR>
 
 			for (; Current < EndSse;)
 			{
-				First8 = _mm_loadu_epi16(Current);
-				Second8 = _mm_loadu_epi16(Current + 8);
+				First8 = _mm_loadu_si128((const __m128i*)Current);
+				Second8 = _mm_loadu_si128((const __m128i*)(Current + 8));
 				CurrentWord = _mm_packus_epi16(First8, Second8);
 
 				CRMask = _mm_movemask_epi8(_mm_cmpeq_epi8(CurrentWord, CharCR));
@@ -1950,8 +1950,8 @@ void ShaderConvertAndStripComments(const FString& ShaderSource, TArray<ANSICHAR>
 
 			for (; Current < EndSse;)
 			{
-				First8 = _mm_loadu_epi16(Current);
-				Second8 = _mm_loadu_epi16(Current + 8);
+				First8 = _mm_loadu_si128((const __m128i*)Current);
+				Second8 = _mm_loadu_si128((const __m128i*)(Current + 8));
 				CurrentWord = _mm_packus_epi16(First8, Second8);
 
 				// Fall back to single character parsing if we hit a CR
@@ -1979,7 +1979,7 @@ void ShaderConvertAndStripComments(const FString& ShaderSource, TArray<ANSICHAR>
 					LFMask &= (0xffff >> (16 - CommentEndOffset));
 					if (LFMask)
 					{
-						_mm_storeu_epi8(CurrentOut, CharLF);
+						_mm_storeu_si128((__m128i*)CurrentOut, CharLF);
 						CurrentOut += _mm_popcnt_u32(LFMask);
 					}
 					Current += CommentEndOffset + 2;
@@ -1992,7 +1992,7 @@ void ShaderConvertAndStripComments(const FString& ShaderSource, TArray<ANSICHAR>
 					LFMask &= 0x7fff;
 					if (LFMask)
 					{
-						_mm_storeu_epi8(CurrentOut, CharLF);
+						_mm_storeu_si128((__m128i*)CurrentOut, CharLF);
 						CurrentOut += _mm_popcnt_u32(LFMask);
 					}
 					Current += 15;
