Explicit Evicted state in FRayTracingGeometry

diff --git a/Engine/Source/Runtime/RenderCore/Private/RayTracingGeometry.cpp b/Engine/Source/Runtime/RenderCore/Private/RayTracingGeometry.cpp
--- a/Engine/Source/Runtime/RenderCore/Private/RayTracingGeometry.cpp
+++ b/Engine/Source/Runtime/RenderCore/Private/RayTracingGeometry.cpp
@@ -93,9 +93,10 @@ void FRayTracingGeometry::RequestBuildIfNeeded(ERTAccelerationStructureBuildPrio
 	RequestBuildIfNeeded(FRHICommandListImmediate::Get(), InBuildPriority);
 }
 
-void FRayTracingGeometry::InitRHIForDynamicRayTracing(FRHICommandList& RHICmdList)
+void FRayTracingGeometry::MakeResident(FRHICommandList& RHICmdList)
 {
-	check(GetRayTracingMode() == ERayTracingMode::Dynamic);
+	check(EnumHasAllFlags(GeometryState, EGeometryStateFlags::Evicted) && RayTracingGeometryRHI == nullptr);
+	EnumRemoveFlags(GeometryState, EGeometryStateFlags::Evicted);
 
 	// Streaming BLAS needs special handling to not get their "streaming" type wiped out as it will cause issues down the line.	
 	// We only have to do this if the geometry was marked to be streamed in.
@@ -136,13 +137,17 @@ void FRayTracingGeometry::InitRHIForDynamicRayTracing(FRHICommandList& RHICmdLis
 
 void FRayTracingGeometry::InitRHIForDynamicRayTracing()
 {
-	InitRHIForDynamicRayTracing(FRHICommandListImmediate::Get());
+	check(GetRayTracingMode() == ERayTracingMode::Dynamic);
+
+	MakeResident(FRHICommandListImmediate::Get());
 }
 
 void FRayTracingGeometry::Evict()
 {
+	check(!EnumHasAllFlags(GeometryState, EGeometryStateFlags::Evicted) && RayTracingGeometryRHI != nullptr);
 	RemoveBuildRequest();
 	RayTracingGeometryRHI.SafeRelease();
+	EnumAddFlags(GeometryState, EGeometryStateFlags::Evicted);
 }
 
 void FRayTracingGeometry::CreateRayTracingGeometry(FRHICommandList& RHICmdList, ERTAccelerationStructureBuildPriority InBuildPriority)
@@ -185,13 +190,17 @@ void FRayTracingGeometry::CreateRayTracingGeometry(FRHICommandList& RHICmdList,
 		{
 			RayTracingGeometryRHI = RHICmdList.CreateRayTracingGeometry(Initializer);
 		}
+		else
+		{
+			EnumAddFlags(GeometryState, EGeometryStateFlags::Evicted);
+		}
 
 		if (Initializer.OfflineData == nullptr)
 		{
 			// Request build if not skip
 			if (InBuildPriority != ERTAccelerationStructureBuildPriority::Skip)
 			{
-				if (IsRayTracingEnabled())
+				if (RayTracingGeometryRHI)
 				{
 					RayTracingBuildRequestIndex = GRayTracingGeometryManager.RequestBuildAccelerationStructure(RHICmdList, this, InBuildPriority);
 				}
@@ -225,7 +234,32 @@ void FRayTracingGeometry::CreateRayTracingGeometry(ERTAccelerationStructureBuild
 
 bool FRayTracingGeometry::IsValid() const
 {
-	return RayTracingGeometryRHI != nullptr && Initializer.TotalPrimitiveCount > 0 && EnumHasAnyFlags(GeometryState, EGeometryStateFlags::Valid);
+	// can't check IsInitialized() because current implementation of hair ray tracing support doesn't initialize resource
+	//check(IsInitialized());
+
+	const bool bIsValidAndNotEvicted = EnumHasAllFlags(GeometryState, EGeometryStateFlags::Valid) && !EnumHasAllFlags(GeometryState, EGeometryStateFlags::Evicted);
+
+	if (bIsValidAndNotEvicted)
+	{
+		check(RayTracingGeometryRHI != nullptr && Initializer.TotalPrimitiveCount > 0);
+	}
+
+	return bIsValidAndNotEvicted;
+}
+
+bool FRayTracingGeometry::IsEvicted() const
+{
+	// can't check IsInitialized() because current implementation of hair ray tracing support doesn't initialize resource
+	//check(IsInitialized());
+
+	const bool bIsEvicted = EnumHasAllFlags(GeometryState, EGeometryStateFlags::Evicted);
+
+	if (bIsEvicted)
+	{
+		check(RayTracingGeometryRHI == nullptr);
+	}
+
+	return bIsEvicted;
 }
 
 void FRayTracingGeometry::InitRHI(FRHICommandListBase& RHICmdList)
diff --git a/Engine/Source/Runtime/RenderCore/Private/RayTracingGeometryManager.cpp b/Engine/Source/Runtime/RenderCore/Private/RayTracingGeometryManager.cpp
--- a/Engine/Source/Runtime/RenderCore/Private/RayTracingGeometryManager.cpp
+++ b/Engine/Source/Runtime/RenderCore/Private/RayTracingGeometryManager.cpp
@@ -128,8 +128,7 @@ void FRayTracingGeometryManager::Tick(FRHICommandList& RHICmdList, bool bHasRayT
 			FScopeLock ScopeLock(&RequestCS);
 			for (FRayTracingGeometry* Geometry : RegisteredGeometries)
 			{
-				checkf(!EnumHasAllFlags(Geometry->GetGeometryState(), FRayTracingGeometry::EGeometryStateFlags::Valid) || Geometry->RayTracingGeometryRHI != nullptr, 
-					TEXT("Ray tracing geometry should have a valid RHI resource when ray tracing is enabled."));
+				checkf(!Geometry->IsEvicted(), TEXT("Ray tracing geometry should not be evicted when ray tracing is enabled."));
 			}
 		}
 		else
@@ -137,7 +136,7 @@ void FRayTracingGeometryManager::Tick(FRHICommandList& RHICmdList, bool bHasRayT
 			FScopeLock ScopeLock(&RequestCS);
 			for (FRayTracingGeometry* Geometry : RegisteredGeometries)
 			{
-				checkf(Geometry->RayTracingGeometryRHI == nullptr, TEXT("Ray tracing geometry should not have a RHI resource when ray tracing is disabled."));
+				checkf(Geometry->IsEvicted() || Geometry->RayTracingGeometryRHI == nullptr, TEXT("Ray tracing geometry should be evicted when ray tracing is disabled."));
 			}
 		}
 #endif
@@ -149,9 +148,9 @@ void FRayTracingGeometryManager::Tick(FRHICommandList& RHICmdList, bool bHasRayT
 		FScopeLock ScopeLock(&RequestCS);
 		for (FRayTracingGeometry* Geometry : RegisteredGeometries)
 		{
-			if (Geometry->RayTracingGeometryRHI == nullptr)
+			if (Geometry->IsEvicted())
 			{
-				Geometry->InitRHIForDynamicRayTracing(RHICmdList);
+				Geometry->MakeResident(RHICmdList);
 			}
 		}
 	}
@@ -160,7 +159,10 @@ void FRayTracingGeometryManager::Tick(FRHICommandList& RHICmdList, bool bHasRayT
 		FScopeLock ScopeLock(&RequestCS);
 		for (FRayTracingGeometry* Geometry : RegisteredGeometries)
 		{
-			Geometry->Evict();			
+			if (Geometry->RayTracingGeometryRHI != nullptr)
+			{
+				Geometry->Evict();
+			}
 		}
 	}
 }
diff --git a/Engine/Source/Runtime/RenderCore/Public/RayTracingGeometry.h b/Engine/Source/Runtime/RenderCore/Public/RayTracingGeometry.h
--- a/Engine/Source/Runtime/RenderCore/Public/RayTracingGeometry.h
+++ b/Engine/Source/Runtime/RenderCore/Public/RayTracingGeometry.h
@@ -58,7 +58,10 @@ public:
 
 		// Special flag that is used when ray tracing is dynamic to mark the streamed geometry to be recreated when ray tracing is switched on.
 		// Only set when mesh streaming is used.
-		StreamedIn = 1 << 2
+		StreamedIn = 1 << 2,
+
+		// If the geometry is initialized but was evicted
+		Evicted = 1 << 3
 	};
 	FRIEND_ENUM_CLASS_FLAGS(EGeometryStateFlags);
 
@@ -71,6 +74,7 @@ public:
 	}
 
 	RENDERCORE_API bool IsValid() const;
+	RENDERCORE_API bool IsEvicted() const;
 
 	void SetAsStreamedIn()
 	{
@@ -109,11 +113,8 @@ public:
 
 	UE_DEPRECATED(5.4, "RequestBuildIfNeeded now requires a command list.")
 	RENDERCORE_API void RequestBuildIfNeeded(ERTAccelerationStructureBuildPriority InBuildPriority);
-	
-	// That function is only supposed to be used when dynamic ray tracing is enabled
-	RENDERCORE_API void InitRHIForDynamicRayTracing(FRHICommandList& RHICmdList);
 
-	UE_DEPRECATED(5.4, "InitRHIForDynamicRayTracing now requires a command list.")
+	UE_DEPRECATED(5.4, "InitRHIForDynamicRayTracing now requires a command list and was renamed to MakeResident().")
 	RENDERCORE_API void InitRHIForDynamicRayTracing();
 
 	RENDERCORE_API void CreateRayTracingGeometry(FRHICommandList& RHICmdList, ERTAccelerationStructureBuildPriority InBuildPriority);
@@ -121,6 +122,7 @@ public:
 	UE_DEPRECATED(5.4, "CreateRayTracingGeometry now requires a command list.")
 	RENDERCORE_API void CreateRayTracingGeometry(ERTAccelerationStructureBuildPriority InBuildPriority);
 
+	RENDERCORE_API void MakeResident(FRHICommandList& RHICmdList);
 	RENDERCORE_API void Evict();
 	
 	bool HasPendingBuildRequest() const
