Works arround a bug where UMaterial::bHasPixelAnimations wasn't being drawn due to VelocityEncodeDepth() return false at runtime on platform that do support it.

This solution is chosen to avoid a shader invalidation

#tests win64

diff --git a/Engine/Source/Runtime/RenderCore/Private/RenderUtils.cpp b/Engine/Source/Runtime/RenderCore/Private/RenderUtils.cpp
--- a/Engine/Source/Runtime/RenderCore/Private/RenderUtils.cpp
+++ b/Engine/Source/Runtime/RenderCore/Private/RenderUtils.cpp
@@ -1548,8 +1548,11 @@ bool VelocityEncodeDepth(const FStaticShaderPlatform Platform)
 
 bool VelocityEncodeHasPixelAnimation(const FStaticShaderPlatform Platform)
 {
-	// Matches VELOCITY_ENCODE_HAS_PIXEL_ANIMATION
-	return VelocityEncodeDepth(Platform);
+	// Should matches VELOCITY_ENCODE_HAS_PIXEL_ANIMATION.
+	// But there is a bug (UE-198416) where VelocityEncodeDepth() always returns false on staged build since only depends on TargetPlatform->VelocityEncodeDepth();
+	// which is completly disconnected to FVelocityRendering::GetFormat(). This ends up in the case where VelocityEncodeDepth() == true at cook time, but false at runtime
+	// and yet end up with FVelocityRendering::GetFormat() return a format to encode the velocity.
+	return !IsMobilePlatform(Platform);
 }
 
 RENDERCORE_API bool AllowTranslucencyPerObjectShadows(const FStaticShaderPlatform Platform)
