Bindless Resources
- Splitting "Bindless Resource" shader parameter types into "Bindless SRV" and "Bindless UAV" to make sure reflection/binding validation is correct.

diff --git a/Engine/Shaders/Public/BindlessResources.ush b/Engine/Shaders/Public/BindlessResources.ush
--- a/Engine/Shaders/Public/BindlessResources.ush
+++ b/Engine/Shaders/Public/BindlessResources.ush
@@ -12,15 +12,20 @@
 	#define ENABLE_BINDLESS_SAMPLERS 0
 #endif
 
-#if defined(UB_RESOURCE_MEMBER_RESOURCE)
+#if defined(UB_RESOURCE_MEMBER_SRV) || defined(UB_RESOURCE_MEMBER_UAV)
 	// Using platform defined bindless macros
 
 #elif ENABLE_BINDLESS_RESOURCES && PLATFORM_SUPPORTS_BINDLESS
-	#define UB_RESOURCE_MEMBER_RESOURCE(Type, UBName, Name) \
-		Type GetBindlessResource##UBName##_##Name() { return GetResourceFromHeap(Type, UB_CB_PREFIXED_MEMBER_ACCESS(UBName, BindlessResource_, Name)); } \
+	#define UB_RESOURCE_MEMBER_SRV(Type, UBName, Name) \
+		Type GetBindlessResource##UBName##_##Name() { return GetResourceFromHeap(Type, UB_CB_PREFIXED_MEMBER_ACCESS(UBName, BindlessSRV_, Name)); } \
+		static const Type UBName##_##Name = GetBindlessResource##UBName##_##Name();
+
+	#define UB_RESOURCE_MEMBER_UAV(Type, UBName, Name) \
+		Type GetBindlessResource##UBName##_##Name() { return GetResourceFromHeap(Type, UB_CB_PREFIXED_MEMBER_ACCESS(UBName, BindlessUAV_, Name)); } \
 		static const Type UBName##_##Name = GetBindlessResource##UBName##_##Name();
 #else
-	#define UB_RESOURCE_MEMBER_RESOURCE(Type, UBName, Name) Type UBName##_##Name
+	#define UB_RESOURCE_MEMBER_SRV(Type, UBName, Name) Type UBName##_##Name
+	#define UB_RESOURCE_MEMBER_UAV(Type, UBName, Name) Type UBName##_##Name
 #endif
 
 #if defined(UB_RESOURCE_MEMBER_SAMPLER)
diff --git a/Engine/Shaders/Public/Platform/Vulkan/VulkanCommon.ush b/Engine/Shaders/Public/Platform/Vulkan/VulkanCommon.ush
--- a/Engine/Shaders/Public/Platform/Vulkan/VulkanCommon.ush
+++ b/Engine/Shaders/Public/Platform/Vulkan/VulkanCommon.ush
@@ -40,13 +40,22 @@ void ImageInterlockedMaxUInt64(RWTexture2D<UlongType> Dest, uint2 Coord, UlongTy
 	// While we use descriptor indexing, declare unbounded resource arrays
 	COMPILER_DEFINE VULKAN_HEAP_NAME(Prefix, Type) Prefix##Type
 	COMPILER_DEFINE VULKAN_SAMPLER_HEAP(Type) VULKAN_HEAP_NAME(VULKAN_BINDLESS_SAMPLER_ARRAY_PREFIX,Type)
-	COMPILER_DEFINE VULKAN_RESOURCE_HEAP(Type) VULKAN_HEAP_NAME(VULKAN_BINDLESS_RESOURCE_ARRAY_PREFIX,Type)
+	COMPILER_DEFINE VULKAN_SRV_HEAP(Type) VULKAN_HEAP_NAME(VULKAN_BINDLESS_SRV_ARRAY_PREFIX,Type)
+	COMPILER_DEFINE VULKAN_UAV_HEAP(Type) VULKAN_HEAP_NAME(VULKAN_BINDLESS_UAV_ARRAY_PREFIX,Type)
 
-	#define UB_RESOURCE_MEMBER_RESOURCE(Type, UBName, Name) \
+	#define UB_RESOURCE_MEMBER_SRV(Type, UBName, Name) \
 		typedef Type SafeType##UBName##_##Name; \
-		SafeType##UBName##_##Name VULKAN_RESOURCE_HEAP(SafeType##UBName##_##Name)[]; \
+		SafeType##UBName##_##Name VULKAN_SRV_HEAP(SafeType##UBName##_##Name)[]; \
 		Type GetBindlessResource##UBName##_##Name() { \
-			return VULKAN_RESOURCE_HEAP(SafeType##UBName##_##Name)[UB_CB_PREFIXED_MEMBER_ACCESS(UBName, BindlessResource_, Name)]; \
+			return VULKAN_SRV_HEAP(SafeType##UBName##_##Name)[UB_CB_PREFIXED_MEMBER_ACCESS(UBName, BindlessSRV_, Name)]; \
+		} \
+		static const Type UBName##_##Name = GetBindlessResource##UBName##_##Name();
+
+	#define UB_RESOURCE_MEMBER_UAV(Type, UBName, Name) \
+		typedef Type SafeType##UBName##_##Name; \
+		SafeType##UBName##_##Name VULKAN_UAV_HEAP(SafeType##UBName##_##Name)[]; \
+		Type GetBindlessResource##UBName##_##Name() { \
+			return VULKAN_UAV_HEAP(SafeType##UBName##_##Name)[UB_CB_PREFIXED_MEMBER_ACCESS(UBName, BindlessUAV_, Name)]; \
 		} \
 		static const Type UBName##_##Name = GetBindlessResource##UBName##_##Name();
 
diff --git a/Engine/Shaders/Public/ShaderVersion.ush b/Engine/Shaders/Public/ShaderVersion.ush
--- a/Engine/Shaders/Public/ShaderVersion.ush
+++ b/Engine/Shaders/Public/ShaderVersion.ush
@@ -2,4 +2,4 @@
 // Each time the console command is executed it generates a new GUID. As this file is included
 // in Platform.ush (which should be included in any shader) it allows to invalidate the shader DDC.
 // 
-#pragma message("UESHADERMETADATA_VERSION 932D8FB4-8E40-48BC-AB6B-E89472D1644C")
+#pragma message("UESHADERMETADATA_VERSION 9A9CDD21-4719-4331-8470-F3EA166C1395")
diff --git a/Engine/Source/Developer/ShaderCompilerCommon/Private/ShaderCompilerCommon.cpp b/Engine/Source/Developer/ShaderCompilerCommon/Private/ShaderCompilerCommon.cpp
--- a/Engine/Source/Developer/ShaderCompilerCommon/Private/ShaderCompilerCommon.cpp
+++ b/Engine/Source/Developer/ShaderCompilerCommon/Private/ShaderCompilerCommon.cpp
@@ -108,7 +108,7 @@ bool BuildResourceTableMapping(
 		if (TOptional<FParameterAllocation> Allocation = ParameterMap.FindParameterAllocation(Name))
 		{
 			const EShaderParameterType ParameterType = Allocation->Type;
-			const bool bBindlessParameter = (ParameterType == EShaderParameterType::BindlessResourceIndex || ParameterType == EShaderParameterType::BindlessSamplerIndex);
+			const bool bBindlessParameter = IsParameterBindless(ParameterType);
 
 			// Force bindless "indices" to zero since they're not needed in SetResourcesFromTables
 			const uint16 BaseIndex = bBindlessParameter ? 0 : Allocation->BaseIndex;
diff --git a/Engine/Source/Developer/ShaderCompilerCommon/Public/ShaderCompilerCommon.h b/Engine/Source/Developer/ShaderCompilerCommon/Public/ShaderCompilerCommon.h
--- a/Engine/Source/Developer/ShaderCompilerCommon/Public/ShaderCompilerCommon.h
+++ b/Engine/Source/Developer/ShaderCompilerCommon/Public/ShaderCompilerCommon.h
@@ -15,11 +15,6 @@ class FShaderParameterParser;
 namespace UE::ShaderCompilerCommon
 {
 	static constexpr const TCHAR* kUniformBufferConstantBufferPrefix = TEXT("UniformBufferConstants_");
-
-	UE_DEPRECATED(5.3, "No longer used; use FShaderParameterParser::kBindlessResourcePrefix instead")
-	static constexpr const TCHAR* kBindlessResourcePrefix = TEXT("BindlessResource_");
-	UE_DEPRECATED(5.3, "No longer used; use FShaderParameterParser::kBindlessSamplerPrefix instead")
-	static constexpr const TCHAR* kBindlessSamplerPrefix = TEXT("BindlessSampler_");
 }
 
 /**
diff --git a/Engine/Source/Developer/VulkanShaderFormat/Private/VulkanShaderCompiler.cpp b/Engine/Source/Developer/VulkanShaderFormat/Private/VulkanShaderCompiler.cpp
--- a/Engine/Source/Developer/VulkanShaderFormat/Private/VulkanShaderCompiler.cpp
+++ b/Engine/Source/Developer/VulkanShaderFormat/Private/VulkanShaderCompiler.cpp
@@ -1232,15 +1232,15 @@ static void GatherSpirvReflectionBindings(
 		};
 
 		// Remove sampler heaps from binding arrays
-		MoveBindlessHeaps(OutBindings.Samplers, VulkanBindless::kBindlessSamplerArrayPrefix, VulkanBindless::BindlessSamplerSet);
+		MoveBindlessHeaps(OutBindings.Samplers, FShaderParameterParser::kBindlessSamplerArrayPrefix, VulkanBindless::BindlessSamplerSet);
 
 		// Remove resource heaps from binding arrays
-		MoveBindlessHeaps(OutBindings.SBufferUAVs, VulkanBindless::kBindlessResourceArrayPrefix, VulkanBindless::BindlessStorageBufferSet);
-		MoveBindlessHeaps(OutBindings.TextureSRVs, VulkanBindless::kBindlessResourceArrayPrefix, VulkanBindless::BindlessSampledImageSet);
-		MoveBindlessHeaps(OutBindings.TextureUAVs, VulkanBindless::kBindlessResourceArrayPrefix, VulkanBindless::BindlessStorageImageSet);
-		MoveBindlessHeaps(OutBindings.TBufferSRVs, VulkanBindless::kBindlessResourceArrayPrefix, VulkanBindless::BindlessUniformTexelBufferSet);
-		MoveBindlessHeaps(OutBindings.TBufferUAVs, VulkanBindless::kBindlessResourceArrayPrefix, VulkanBindless::BindlessStorageTexelBufferSet);
-		MoveBindlessHeaps(OutBindings.AccelerationStructures, VulkanBindless::kBindlessResourceArrayPrefix, VulkanBindless::BindlessAccelerationStructureSet);
+		MoveBindlessHeaps(OutBindings.SBufferUAVs, FShaderParameterParser::kBindlessUAVArrayPrefix, VulkanBindless::BindlessStorageBufferSet);
+		MoveBindlessHeaps(OutBindings.TextureSRVs, FShaderParameterParser::kBindlessSRVArrayPrefix, VulkanBindless::BindlessSampledImageSet);
+		MoveBindlessHeaps(OutBindings.TextureUAVs, FShaderParameterParser::kBindlessUAVArrayPrefix, VulkanBindless::BindlessStorageImageSet);
+		MoveBindlessHeaps(OutBindings.TBufferSRVs, FShaderParameterParser::kBindlessSRVArrayPrefix, VulkanBindless::BindlessUniformTexelBufferSet);
+		MoveBindlessHeaps(OutBindings.TBufferUAVs, FShaderParameterParser::kBindlessUAVArrayPrefix, VulkanBindless::BindlessStorageTexelBufferSet);
+		MoveBindlessHeaps(OutBindings.AccelerationStructures, FShaderParameterParser::kBindlessSRVArrayPrefix, VulkanBindless::BindlessAccelerationStructureSet);
 
 		// Move uniform buffers to the correct set
 		{
@@ -1500,10 +1500,12 @@ static bool BuildShaderOutputFromSpirv(
 							const SpvReflectBlockVariable& Member = Binding->block.members[MemberIndex];
 
 							FString MemberName(ANSI_TO_TCHAR(Member.name));
+							FStringView AdjustedMemberName(MemberName);
+
+							const EShaderParameterType BindlessParameterType = FShaderParameterParser::ParseAndRemoveBindlessParameterPrefix(AdjustedMemberName);
 
 							// Add all members of global ub, and only bindless samplers/resources for root param
-							if (!InternalState.UseRootParametersStructure() ||
-								(MemberName.StartsWith(FShaderParameterParser::kBindlessSamplerPrefix) || MemberName.StartsWith(FShaderParameterParser::kBindlessResourcePrefix)))
+							if (!InternalState.UseRootParametersStructure() || BindlessParameterType != EShaderParameterType::LooseData)
 							{
 								HandleReflectedGlobalConstantBufferMember(
 									MemberName,
@@ -1513,8 +1515,7 @@ static bool BuildShaderOutputFromSpirv(
 									Output
 								);
 
-								FShaderParameterParser::RemoveBindlessParameterPrefix(MemberName);
-								EntryTypes.Add(MemberName, FVulkanShaderHeader::PackedGlobal);
+								EntryTypes.Add(FString(AdjustedMemberName), FVulkanShaderHeader::PackedGlobal);
 							}
 
 							PackedGlobalArraySize = FMath::Max<uint32>((Member.absolute_offset + Member.size), PackedGlobalArraySize);
@@ -2197,8 +2198,9 @@ bool PreprocessVulkanShader(const FShaderCompilerInput& Input, const FShaderComp
 		AdditionalDefines.SetDefine(TEXT("VULKAN_SUPPORTS_SUBGROUP_SIZE_CONTROL"), 1);
 	}
 
-	AdditionalDefines.SetDefine(TEXT("VULKAN_BINDLESS_SAMPLER_ARRAY_PREFIX"), VulkanBindless::kBindlessSamplerArrayPrefix);
-	AdditionalDefines.SetDefine(TEXT("VULKAN_BINDLESS_RESOURCE_ARRAY_PREFIX"), VulkanBindless::kBindlessResourceArrayPrefix);
+	AdditionalDefines.SetDefine(TEXT("VULKAN_BINDLESS_SRV_ARRAY_PREFIX"), FShaderParameterParser::kBindlessSRVArrayPrefix);
+	AdditionalDefines.SetDefine(TEXT("VULKAN_BINDLESS_UAV_ARRAY_PREFIX"), FShaderParameterParser::kBindlessUAVArrayPrefix);
+	AdditionalDefines.SetDefine(TEXT("VULKAN_BINDLESS_SAMPLER_ARRAY_PREFIX"), FShaderParameterParser::kBindlessSamplerArrayPrefix);
 	AdditionalDefines.SetDefine(TEXT("VULKAN_MAX_BINDLESS_UNIFORM_BUFFERS_PER_STAGE"), VulkanBindless::MaxUniformBuffersPerStage);
 
 	if (IsAndroidShaderFormat(Input.ShaderFormat))
diff --git a/Engine/Source/Runtime/RenderCore/Private/Shader.cpp b/Engine/Source/Runtime/RenderCore/Private/Shader.cpp
--- a/Engine/Source/Runtime/RenderCore/Private/Shader.cpp
+++ b/Engine/Source/Runtime/RenderCore/Private/Shader.cpp
@@ -738,11 +738,11 @@ void FShader::BuildParameterMapInfo(const TMap<FString, FParameterAllocation>& P
 		case EShaderParameterType::UniformBuffer:
 			UniformCount++;
 			break;
-		case EShaderParameterType::BindlessSamplerIndex:
+		case EShaderParameterType::BindlessSampler:
 		case EShaderParameterType::Sampler:
 			SamplerCount++;
 			break;
-		case EShaderParameterType::BindlessResourceIndex:
+		case EShaderParameterType::BindlessSRV:
 		case EShaderParameterType::SRV:
 			SRVCount++;
 			break;
@@ -761,9 +761,9 @@ void FShader::BuildParameterMapInfo(const TMap<FString, FParameterAllocation>& P
 			return &ParameterMapInfo.TextureSamplers;
 		case EShaderParameterType::SRV:
 			return &ParameterMapInfo.SRVs;
-		case EShaderParameterType::BindlessResourceIndex:
+		case EShaderParameterType::BindlessSRV:
 			return &ParameterMapInfo.SRVs;
-		case EShaderParameterType::BindlessSamplerIndex:
+		case EShaderParameterType::BindlessSampler:
 			return &ParameterMapInfo.TextureSamplers;
 		default:
 			return nullptr;
diff --git a/Engine/Source/Runtime/RenderCore/Private/ShaderParameterParser.cpp b/Engine/Source/Runtime/RenderCore/Private/ShaderParameterParser.cpp
--- a/Engine/Source/Runtime/RenderCore/Private/ShaderParameterParser.cpp
+++ b/Engine/Source/Runtime/RenderCore/Private/ShaderParameterParser.cpp
@@ -232,14 +232,19 @@ EShaderParameterType FShaderParameterParser::ParseAndRemoveBindlessParameterPref
 {
 	const FStringView OriginalName = InName;
 
-	if (InName = UE::String::RemoveFromStart(InName, FStringView(kBindlessResourcePrefix)); InName != OriginalName)
+	if (InName = UE::String::RemoveFromStart(InName, FStringView(kBindlessSRVPrefix)); InName != OriginalName)
 	{
-		return EShaderParameterType::BindlessResourceIndex;
+		return EShaderParameterType::BindlessSRV;
+	}
+
+	if (InName = UE::String::RemoveFromStart(InName, FStringView(kBindlessUAVPrefix)); InName != OriginalName)
+	{
+		return EShaderParameterType::BindlessUAV;
 	}
 
 	if (InName = UE::String::RemoveFromStart(InName, FStringView(kBindlessSamplerPrefix)); InName != OriginalName)
 	{
-		return EShaderParameterType::BindlessSamplerIndex;
+		return EShaderParameterType::BindlessSampler;
 	}
 
 	return EShaderParameterType::LooseData;
@@ -256,10 +261,22 @@ EShaderParameterType FShaderParameterParser::ParseAndRemoveBindlessParameterPref
 
 bool FShaderParameterParser::RemoveBindlessParameterPrefix(FString& InName)
 {
-	return InName.RemoveFromStart(kBindlessResourcePrefix)
+	return InName.RemoveFromStart(kBindlessSRVPrefix)
+		|| InName.RemoveFromStart(kBindlessUAVPrefix)
 		|| InName.RemoveFromStart(kBindlessSamplerPrefix);
 }
 
+FStringView FShaderParameterParser::GetBindlessParameterPrefix(EShaderParameterType InShaderParameterType)
+{
+	switch (InShaderParameterType)
+	{
+	case EShaderParameterType::BindlessSampler: return kBindlessSamplerPrefix;
+	case EShaderParameterType::BindlessSRV:     return kBindlessSRVPrefix;
+	case EShaderParameterType::BindlessUAV:     return kBindlessUAVPrefix;
+	}
+
+	return FStringView();
+}
 
 bool FShaderParameterParser::ParseParameters(
 	const FShaderParametersMetadata* RootParametersStructure,
@@ -385,7 +402,7 @@ bool FShaderParameterParser::ParseParameters(
 				FStringView Leftovers = ShaderSource.Mid(NameEndPos + 1, (Cursor - 1) - (NameEndPos + 1) + 1);
 
 				EShaderParameterType ParsedParameterType = ParseAndRemoveBindlessParameterPrefix(Name);
-				const bool bBindlessIndex = (ParsedParameterType == EShaderParameterType::BindlessResourceIndex || ParsedParameterType == EShaderParameterType::BindlessSamplerIndex);
+				const bool bBindlessIndex = (ParsedParameterType != EShaderParameterType::LooseData);
 
 				EBindlessConversionType BindlessConversionType = EBindlessConversionType::None;
 				EShaderParameterType ParsedConstantBufferType = ParsedParameterType;
@@ -394,15 +411,20 @@ bool FShaderParameterParser::ParseParameters(
 				{
 					ParsedParameterType = ParseParameterType(Type, ExtraSRVTypes, ExtraUAVTypes);
 
-					if (bBindlessResources && (ParsedParameterType == EShaderParameterType::SRV || ParsedParameterType == EShaderParameterType::UAV))
+					if (bBindlessResources && ParsedParameterType == EShaderParameterType::SRV)
+					{
+						BindlessConversionType = EBindlessConversionType::SRV;
+						ParsedConstantBufferType = EShaderParameterType::BindlessSRV;
+					}
+					else if (bBindlessResources && ParsedParameterType == EShaderParameterType::UAV)
 					{
-						BindlessConversionType = EBindlessConversionType::Resource;
-						ParsedConstantBufferType = EShaderParameterType::BindlessResourceIndex;
+						BindlessConversionType = EBindlessConversionType::UAV;
+						ParsedConstantBufferType = EShaderParameterType::BindlessUAV;
 					}
 					else if (bBindlessSamplers && ParsedParameterType == EShaderParameterType::Sampler)
 					{
 						BindlessConversionType = EBindlessConversionType::Sampler;
-						ParsedConstantBufferType = EShaderParameterType::BindlessSamplerIndex;
+						ParsedConstantBufferType = EShaderParameterType::BindlessSampler;
 					}
 
 					if (BindlessConversionType != EBindlessConversionType::None && Leftovers.Contains(TEXT("register")))
@@ -890,13 +912,35 @@ void FShaderParameterParser::RemoveMovingParametersFromSource(FString& Preproces
 	}
 }
 
+static FStringView GetBindlessParameterPrefix(EBindlessConversionType InConversionType)
+{
+	switch (InConversionType)
+	{
+	case EBindlessConversionType::SRV:     return FShaderParameterParser::kBindlessSRVPrefix;
+	case EBindlessConversionType::UAV:     return FShaderParameterParser::kBindlessUAVPrefix;
+	case EBindlessConversionType::Sampler: return FShaderParameterParser::kBindlessSamplerPrefix;
+	}
+	return FStringView();
+}
+
+static FStringView GetBindlessArrayHeapPrefix(EBindlessConversionType InConversionType)
+{
+	switch (InConversionType)
+	{
+	case EBindlessConversionType::SRV:     return FShaderParameterParser::kBindlessSRVArrayPrefix;
+	case EBindlessConversionType::UAV:     return FShaderParameterParser::kBindlessUAVArrayPrefix;
+	case EBindlessConversionType::Sampler: return FShaderParameterParser::kBindlessSamplerArrayPrefix;
+	}
+	return FStringView();
+}
+
 FString FShaderParameterParser::GenerateBindlessParameterDeclaration(const FParsedShaderParameter& ParsedParameter) const
 {
 	const bool bIsSampler = (ParsedParameter.BindlessConversionType == EBindlessConversionType::Sampler);
 	const FStringView Name = ParsedParameter.ParsedName;
 	const FStringView Type = ParsedParameter.ParsedType;
 	const TCHAR* StorageClass = ParsedParameter.bGloballyCoherent ? TEXT("globallycoherent ") : TEXT("");
-	const TCHAR* IndexPrefix = bIsSampler ? FShaderParameterParser::kBindlessSamplerPrefix : FShaderParameterParser::kBindlessResourcePrefix;
+	const FStringView IndexPrefix = ::GetBindlessParameterPrefix(ParsedParameter.BindlessConversionType);
 
 	TStringBuilder<512> Result;
 
@@ -911,7 +955,7 @@ FString FShaderParameterParser::GenerateBindlessParameterDeclaration(const FPars
 
 	FString TypedefText = TEXT("");
 	// Vulkan requires these typedefs for both resources and samplers, while other platforms only require resources.
-	if (ParsedParameter.BindlessConversionType == EBindlessConversionType::Resource || BindlessParameterMode == EBindlessParameterMode::Vulkan)
+	if (ParsedParameter.BindlessConversionType != EBindlessConversionType::Sampler || BindlessParameterMode == EBindlessParameterMode::Vulkan)
 	{
 		RewriteType = FString::Printf(TEXT("SafeType%.*s"), Name.Len(), Name.GetData());
 		TypedefText = FString::Printf(TEXT("typedef %.*s %s;"), Type.Len(), Type.GetData(), *RewriteType);
@@ -921,7 +965,7 @@ FString FShaderParameterParser::GenerateBindlessParameterDeclaration(const FPars
 
 	if (BindlessParameterMode == EBindlessParameterMode::Vulkan)
 	{
-		const TCHAR* HeapPrefix = bIsSampler ? VulkanBindless::kBindlessSamplerArrayPrefix : VulkanBindless::kBindlessResourceArrayPrefix;
+		const FStringView HeapPrefix = GetBindlessArrayHeapPrefix(ParsedParameter.BindlessConversionType);
 
 		// Declare a heap for the RewriteType
 		// e.g. `SafeType##Name ResourceDescriptorHeap_SafeType##Name[];`
@@ -938,7 +982,7 @@ FString FShaderParameterParser::GenerateBindlessParameterDeclaration(const FPars
 		// e.g. `Type GetBindlessResource##Name() { return GetResourceFromHeap(Type, BindlessResource_##Name); } static const Type Name = GetBindlessResource##Name()`
 		// or   `Type GetBindlessSampler##Name() { return GetSamplerFromHeap(Type, BindlessSampler_##Name); } static const Type Name = GetBindlessSampler##Name()`
 		Result << StorageClass << RewriteType << TEXT(" GetBindless") << Kind << Name << TEXT("()");
-		Result << TEXT("{ return Get") << Kind << TEXT("FromHeap(") << StorageClass << RewriteType << TEXT(", Bindless") << Kind << TEXT("_") << Name << TEXT("); } ");
+		Result << TEXT("{ return Get") << Kind << TEXT("FromHeap(") << StorageClass << RewriteType << TEXT(", ") << IndexPrefix << Name << TEXT("); } ");
 		Result << TEXT("static const ") << StorageClass << RewriteType << TEXT(" ") << Name << TEXT(" = GetBindless") << Kind << Name << TEXT("();");
 	}
 
@@ -1065,18 +1109,13 @@ bool FShaderParameterParser::MoveShaderParametersToRootConstantBuffer(
 				const TCHAR* ConstantSwizzle = GetConstantSwizzle(ByteOffset);
 
 #define SVARG(N) N.Len(), N.GetData()
-				if (ParsedParameter->ConstantBufferParameterType == EShaderParameterType::BindlessResourceIndex)
-				{
-					RootCBufferContent.Append(FString::Printf(
-						TEXT("uint BindlessResource_%.*s : packoffset(c%d%s);\n"),
-						SVARG(ParsedParameter->ParsedName),
-						ConstantRegister,
-						ConstantSwizzle));
-				}
-				else if (ParsedParameter->ConstantBufferParameterType == EShaderParameterType::BindlessSamplerIndex)
+				if (IsParameterBindless(ParsedParameter->ConstantBufferParameterType))
 				{
+					const FStringView Prefix = GetBindlessParameterPrefix(ParsedParameter->ConstantBufferParameterType);
+
 					RootCBufferContent.Append(FString::Printf(
-						TEXT("uint BindlessSampler_%.*s : packoffset(c%d%s);\n"),
+						TEXT("uint %.*s%.*s : packoffset(c%d%s);\n"),
+						SVARG(Prefix),
 						SVARG(ParsedParameter->ParsedName),
 						ConstantRegister,
 						ConstantSwizzle));
diff --git a/Engine/Source/Runtime/RenderCore/Private/ShaderParameterStruct.cpp b/Engine/Source/Runtime/RenderCore/Private/ShaderParameterStruct.cpp
--- a/Engine/Source/Runtime/RenderCore/Private/ShaderParameterStruct.cpp
+++ b/Engine/Source/Runtime/RenderCore/Private/ShaderParameterStruct.cpp
@@ -197,7 +197,7 @@ struct FShaderParameterStructBindingContext
 				}
 				else if (bIsRHIResource || bIsRDGResource)
 				{
-					if (ParameterAllocation->Type == EShaderParameterType::BindlessResourceIndex || ParameterAllocation->Type == EShaderParameterType::BindlessSamplerIndex)
+					if (IsParameterBindless(ParameterAllocation->Type))
 					{
 						FShaderParameterBindings::FBindlessResourceParameter Parameter;
 						Parameter.GlobalConstantOffset = ParameterAllocation->BaseIndex;
diff --git a/Engine/Source/Runtime/RenderCore/Private/ShaderParameters.cpp b/Engine/Source/Runtime/RenderCore/Private/ShaderParameters.cpp
--- a/Engine/Source/Runtime/RenderCore/Private/ShaderParameters.cpp
+++ b/Engine/Source/Runtime/RenderCore/Private/ShaderParameters.cpp
@@ -49,7 +49,7 @@ void FShaderResourceParameter::Bind(const FShaderParameterMap& ParameterMap, con
 {
 	if (TOptional<FParameterAllocation> Allocation = ParameterMap.FindParameterAllocation(ParameterName))
 	{
-		if (Allocation->Type == EShaderParameterType::BindlessResourceIndex || Allocation->Type == EShaderParameterType::BindlessSamplerIndex)
+		if (IsParameterBindless(Allocation->Type))
 		{
 			checkf(Allocation->BufferIndex == 0, TEXT("Unexpected buffer index (%d) for bindless index. Global bindless parameters are expected to be in the global constant buffer (buffer index 0)."), Allocation->BufferIndex);
 		}
@@ -240,8 +240,23 @@ static void CreateHLSLUniformBufferStructMembersDeclaration(
 			PreviousBaseTypeName = TEXT("UB_UINT");
 			HLSLBaseOffset = AbsoluteMemberOffset + HLSLMemberSize;
 
+			auto GetBindlessPrefix = [](EUniformBufferBaseType InBaseType)
+				{
+					if (InBaseType == UBMT_SAMPLER)
+					{
+						return FShaderParameterParser::kBindlessSamplerPrefix;
+					}
+
+					if (InBaseType == UBMT_UAV || InBaseType == UBMT_RDG_TEXTURE_UAV)
+					{
+						return FShaderParameterParser::kBindlessUAVPrefix;
+					}
+
+					return FShaderParameterParser::kBindlessSRVPrefix;
+				};
+
 			// Generate the member declaration.
-			const TCHAR* MemberPrefix = (Member.GetBaseType() == UBMT_SAMPLER ? FShaderParameterParser::kBindlessSamplerPrefix : FShaderParameterParser::kBindlessResourcePrefix);
+			const TCHAR* MemberPrefix = GetBindlessPrefix(Member.GetBaseType());
 
 			TStringBuilder<256> ParameterName;
 			if (!GlobalPrefix.IsEmpty())
@@ -349,14 +364,14 @@ static void CreateHLSLUniformBufferStructMembersDeclaration(
 				Decl.ResourceMembers.Appendf(TEXT("UB_RESOURCE_MEMBER_SAMPLER(%s, %s, %s);\n"), Member.GetShaderType(), *UniformBufferName, *ParameterName);
 				AddStructMember(Member, true);
 			}
-			else if (Member.GetBaseType() == UBMT_SRV)
+			else if (Member.GetBaseType() == UBMT_UAV || Member.GetBaseType() == UBMT_RDG_TEXTURE_UAV)
 			{
-				Decl.ResourceMembers.Appendf(TEXT("UB_RESOURCE_MEMBER_RESOURCE(%s, %s, %s);\n"), Member.GetShaderType(), *UniformBufferName, *ParameterName);
+				Decl.ResourceMembers.Appendf(TEXT("UB_RESOURCE_MEMBER_UAV(%s, %s, %s);\n"), Member.GetShaderType(), *UniformBufferName, *ParameterName);
 				AddStructMember(Member, true);
 			}
 			else
 			{
-				Decl.ResourceMembers.Appendf(TEXT("UB_RESOURCE_MEMBER_RESOURCE(%s, %s, %s);\n"), Member.GetShaderType(), *UniformBufferName, *ParameterName);
+				Decl.ResourceMembers.Appendf(TEXT("UB_RESOURCE_MEMBER_SRV(%s, %s, %s);\n"), Member.GetShaderType(), *UniformBufferName, *ParameterName);
 				AddStructMember(Member, true);
 			}
 		}
diff --git a/Engine/Source/Runtime/RenderCore/Public/ShaderCore.h b/Engine/Source/Runtime/RenderCore/Public/ShaderCore.h
--- a/Engine/Source/Runtime/RenderCore/Public/ShaderCore.h
+++ b/Engine/Source/Runtime/RenderCore/Public/ShaderCore.h
@@ -216,12 +216,21 @@ enum class EShaderParameterType : uint8
 	SRV,
 	UAV,
 
-	BindlessResourceIndex,
-	BindlessSamplerIndex,
+	BindlessSampler,
+	BindlessSRV,
+	BindlessUAV,
 
 	Num
 };
 
+inline bool IsParameterBindless(EShaderParameterType ParameterType)
+{
+	return ParameterType == EShaderParameterType::BindlessSampler
+		|| ParameterType == EShaderParameterType::BindlessSRV
+		|| ParameterType == EShaderParameterType::BindlessUAV
+		;
+}
+
 struct FParameterAllocation
 {
 	uint16 BufferIndex = 0;
diff --git a/Engine/Source/Runtime/RenderCore/Public/ShaderParameterParser.h b/Engine/Source/Runtime/RenderCore/Public/ShaderParameterParser.h
--- a/Engine/Source/Runtime/RenderCore/Public/ShaderParameterParser.h
+++ b/Engine/Source/Runtime/RenderCore/Public/ShaderParameterParser.h
@@ -17,8 +17,9 @@ enum class EShaderParameterType : uint8;
 enum class EBindlessConversionType : uint8
 {
 	None,
-	Resource,
-	Sampler
+	SRV,
+	UAV,
+	Sampler,
 };
 
 enum class EBindlessParameterMode : uint8
@@ -27,13 +28,6 @@ enum class EBindlessParameterMode : uint8
 	Vulkan,
 };
 
-namespace VulkanBindless
-{
-	// Prefix used to declare arrays of samplers/resources for bindless
-	static constexpr const TCHAR* kBindlessResourceArrayPrefix = TEXT("ResourceDescriptorHeap_");
-	static constexpr const TCHAR* kBindlessSamplerArrayPrefix = TEXT("SamplerDescriptorHeap_");
-}
-
 inline FStringView StripTemplateFromType(const FStringView& Input)
 {
 	FStringView UntemplatedType = FStringView(Input);
@@ -110,13 +104,20 @@ public:
 
 	FShaderParameterParser& operator=(FShaderParameterParser&&) = default;
 
-	static constexpr const TCHAR* kBindlessResourcePrefix = TEXT("BindlessResource_");
+	static constexpr const TCHAR* kBindlessSRVPrefix = TEXT("BindlessSRV_");
+	static constexpr const TCHAR* kBindlessUAVPrefix = TEXT("BindlessUAV_");
 	static constexpr const TCHAR* kBindlessSamplerPrefix = TEXT("BindlessSampler_");
 
+	// Prefix used to declare arrays of samplers/resources for bindless
+	static constexpr const TCHAR* kBindlessSRVArrayPrefix = TEXT("SRVDescriptorHeap_");
+	static constexpr const TCHAR* kBindlessUAVArrayPrefix = TEXT("UAVDescriptorHeap_");
+	static constexpr const TCHAR* kBindlessSamplerArrayPrefix = TEXT("SamplerDescriptorHeap_");
+
 	static RENDERCORE_API EShaderParameterType ParseParameterType(FStringView InType, TConstArrayView<const TCHAR*> InExtraSRVTypes, TConstArrayView<const TCHAR*> InExtraUAVTypes);
 	static RENDERCORE_API EShaderParameterType ParseAndRemoveBindlessParameterPrefix(FStringView& InName);
 	static RENDERCORE_API EShaderParameterType ParseAndRemoveBindlessParameterPrefix(FString& InName);
 	static RENDERCORE_API bool RemoveBindlessParameterPrefix(FString& InName);
+	static RENDERCORE_API FStringView GetBindlessParameterPrefix(EShaderParameterType InShaderParameterType);
 
 	/** Parses the preprocessed shader code and applies the necessary modifications to it. */
 	UE_DEPRECATED(5.3, "Use ParseAndModify overload accepting array of FShaderCompilerError instead of passing FShaderCompilerOutput")
diff --git a/Engine/Source/Runtime/RenderCore/Public/ShaderParameterUtils.h b/Engine/Source/Runtime/RenderCore/Public/ShaderParameterUtils.h
--- a/Engine/Source/Runtime/RenderCore/Public/ShaderParameterUtils.h
+++ b/Engine/Source/Runtime/RenderCore/Public/ShaderParameterUtils.h
@@ -79,7 +79,7 @@ inline void SetTextureParameter(FRHIBatchedShaderParameters& BatchedParameters,
 	if (Parameter.IsBound())
 	{
 #if PLATFORM_SUPPORTS_BINDLESS_RENDERING
-		if (Parameter.GetType() == EShaderParameterType::BindlessResourceIndex)
+		if (Parameter.GetType() == EShaderParameterType::BindlessSRV)
 		{
 			BatchedParameters.SetBindlessTexture(Parameter.GetBaseIndex(), TextureRHI);
 		}
@@ -96,7 +96,7 @@ inline void SetSamplerParameter(FRHIBatchedShaderParameters& BatchedParameters,
 	if (Parameter.IsBound())
 	{
 #if PLATFORM_SUPPORTS_BINDLESS_RENDERING
-		if (Parameter.GetType() == EShaderParameterType::BindlessSamplerIndex)
+		if (Parameter.GetType() == EShaderParameterType::BindlessSampler)
 		{
 			BatchedParameters.SetBindlessSampler(Parameter.GetBaseIndex(), SamplerStateRHI);
 		}
@@ -141,7 +141,7 @@ inline void SetSRVParameter(FRHIBatchedShaderParameters& BatchedParameters, cons
 	if (Parameter.IsBound())
 	{
 #if PLATFORM_SUPPORTS_BINDLESS_RENDERING
-		if (Parameter.GetType() == EShaderParameterType::BindlessResourceIndex)
+		if (Parameter.GetType() == EShaderParameterType::BindlessSRV)
 		{
 			BatchedParameters.SetBindlessResourceView(Parameter.GetBaseIndex(), SRV);
 		}
@@ -158,7 +158,7 @@ inline void SetUAVParameter(FRHIBatchedShaderParameters& BatchedParameters, cons
 	if (Parameter.IsBound())
 	{
 #if PLATFORM_SUPPORTS_BINDLESS_RENDERING
-		if (Parameter.GetType() == EShaderParameterType::BindlessResourceIndex)
+		if (Parameter.GetType() == EShaderParameterType::BindlessUAV)
 		{
 			BatchedParameters.SetBindlessUAV(Parameter.GetBaseIndex(), UAV);
 		}
@@ -175,7 +175,7 @@ inline void UnsetSRVParameter(FRHIBatchedShaderUnbinds& BatchedUnbinds, const FS
 	if (Parameter.IsBound())
 	{
 #if PLATFORM_SUPPORTS_BINDLESS_RENDERING
-		if (Parameter.GetType() == EShaderParameterType::BindlessResourceIndex)
+		if (Parameter.GetType() == EShaderParameterType::BindlessSRV)
 		{
 			// We don't need to clear Bindless views
 		}
@@ -192,7 +192,7 @@ inline void UnsetUAVParameter(FRHIBatchedShaderUnbinds& BatchedUnbinds, const FS
 	if (Parameter.IsBound())
 	{
 #if PLATFORM_SUPPORTS_BINDLESS_RENDERING
-		if (Parameter.GetType() == EShaderParameterType::BindlessResourceIndex)
+		if (Parameter.GetType() == EShaderParameterType::BindlessSRV)
 		{
 			// We don't need to clear Bindless views
 		}
diff --git a/Engine/Source/Runtime/Renderer/Private/MeshPassProcessor.cpp b/Engine/Source/Runtime/Renderer/Private/MeshPassProcessor.cpp
--- a/Engine/Source/Runtime/Renderer/Private/MeshPassProcessor.cpp
+++ b/Engine/Source/Runtime/Renderer/Private/MeshPassProcessor.cpp
@@ -121,7 +121,7 @@ private:
 inline void SetTextureParameter(FRHIBatchedShaderParameters& BatchedParameters, const FShaderResourceParameterInfo& Parameter, FRHITexture* TextureRHI)
 {
 #if PLATFORM_SUPPORTS_BINDLESS_RENDERING
-	if (Parameter.Type == EShaderParameterType::BindlessResourceIndex)
+	if (Parameter.Type == EShaderParameterType::BindlessSRV)
 	{
 		check(Parameter.BufferIndex == 0);
 		BatchedParameters.SetBindlessTexture(Parameter.BaseIndex, TextureRHI);
@@ -136,7 +136,7 @@ inline void SetTextureParameter(FRHIBatchedShaderParameters& BatchedParameters,
 inline void SetSrvParameter(FRHIBatchedShaderParameters& BatchedParameters, const FShaderResourceParameterInfo& Parameter, FRHIShaderResourceView* SrvRHI)
 {
 #if PLATFORM_SUPPORTS_BINDLESS_RENDERING
-	if (Parameter.Type == EShaderParameterType::BindlessResourceIndex)
+	if (Parameter.Type == EShaderParameterType::BindlessSRV)
 	{
 		check(Parameter.BufferIndex == 0);
 		BatchedParameters.SetBindlessResourceView(Parameter.BaseIndex, SrvRHI);
@@ -151,7 +151,7 @@ inline void SetSrvParameter(FRHIBatchedShaderParameters& BatchedParameters, cons
 inline void SetSamplerParameter(FRHIBatchedShaderParameters& BatchedParameters, const FShaderResourceParameterInfo& Parameter, FRHISamplerState* SamplerStateRHI)
 {
 #if PLATFORM_SUPPORTS_BINDLESS_RENDERING
-	if (Parameter.Type == EShaderParameterType::BindlessSamplerIndex)
+	if (Parameter.Type == EShaderParameterType::BindlessSampler)
 	{
 		BatchedParameters.SetBindlessSampler(Parameter.BaseIndex, SamplerStateRHI);
 	}
