Fixed race condition with the StopRecording() sync command and commands inserted into a pipe.

diff --git a/Engine/Source/Runtime/RenderCore/Private/RenderingThread.cpp b/Engine/Source/Runtime/RenderCore/Private/RenderingThread.cpp
--- a/Engine/Source/Runtime/RenderCore/Private/RenderingThread.cpp
+++ b/Engine/Source/Runtime/RenderCore/Private/RenderingThread.cpp
@@ -1948,7 +1948,7 @@ private:
 			Pipe->bRecording = false;
 			NumPipesToStopRecording++;
 
-			UE::TScopeLock PipeLock(Pipe->Mutex);
+			Pipe->Mutex.Lock();
 			Pipe->Frame_GameThread = nullptr;
 		}
 
@@ -1982,6 +1982,14 @@ private:
 			RHIResourceLifetimeReleaseRef(RHICmdList, NumPipesToStopRecording);
 		});
 
+		// Wait to unlock the mutex until the sync command has been submitted to the render thread. This avoids
+		// race conditions where a command meant for a specific pipe might be inserted to the render thread pipe
+		// prior to the actual wait command.
+		for (FRenderCommandPipeSetBitIterator BitIt(PipeBits); BitIt; ++BitIt)
+		{
+			AllPipes[BitIt.GetIndex()]->Mutex.Unlock();
+		}
+
 		UE_RENDER_COMMAND_END_REGION(UE_RENDER_COMMAND_PIPE_RECORD_REGION);
 	}
 
