Stable Pipeline Cache files were not serializing the PipelineCache versions.

diff --git a/Engine/Source/Runtime/RHI/Private/PipelineFileCache.cpp b/Engine/Source/Runtime/RHI/Private/PipelineFileCache.cpp
--- a/Engine/Source/Runtime/RHI/Private/PipelineFileCache.cpp
+++ b/Engine/Source/Runtime/RHI/Private/PipelineFileCache.cpp
@@ -80,7 +80,8 @@ enum class EPipelineCacheFileFormatVersions : uint32
 	MoreRenderTargetFlags = 23,
 	FragmentDensityAttachment = 24,
 	AddingDepthClipMode = 25,
-	RemovingLineAA = 26,
+	BeforeStableCacheVersioning = 26,
+	RemovingLineAA = 27,
 };
 
 const uint64 FPipelineCacheFileFormatMagic = 0x5049504543414348; // PIPECACH
@@ -389,8 +390,7 @@ RHI_API FArchive& operator<<(FArchive& Ar, FPipelineFileCacheRasterizerState& Ra
 	Ar << RasterizerStateInitializer.DepthClipMode;
 	Ar << RasterizerStateInitializer.bAllowMSAA;
 
-	// TODO: actually get versioning to work. Look for: MemReader.SetGameNetVer(FPipelineCacheFileFormatCurrentVersion);
-	//if (Ar.GameNetVer() < (uint32)EPipelineCacheFileFormatVersions::RemovingLineAA)
+	if (Ar.GameNetVer() < (uint32)EPipelineCacheFileFormatVersions::RemovingLineAA)
 	{
 		bool bEnableLineAA = false;
 		Ar << bEnableLineAA;
diff --git a/Engine/Source/Runtime/RenderCore/Private/PipelineCacheUtilities.cpp b/Engine/Source/Runtime/RenderCore/Private/PipelineCacheUtilities.cpp
--- a/Engine/Source/Runtime/RenderCore/Private/PipelineCacheUtilities.cpp
+++ b/Engine/Source/Runtime/RenderCore/Private/PipelineCacheUtilities.cpp
@@ -67,7 +67,9 @@ namespace Private
 
 		enum class EVersion : int32
 		{
-			Current = 4
+			AddingPipelineCacheVersion = 5,
+
+			Current = AddingPipelineCacheVersion
 		};
 
 		/** Magic to reject other files */
@@ -483,6 +485,8 @@ bool UE::PipelineCacheUtilities::SaveStablePipelineCacheFile(const FString& Outp
 		UE::PipelineCacheUtilities::Private::FIndexedSHAHashAndFNameProxyArchive MemWriter(PlainMemWriter);
 
 		MemWriter.SetGameNetVer(FPipelineCacheFileFormatCurrentVersion);
+		uint32 CacheVersion = FPipelineCacheFileFormatCurrentVersion;
+		MemWriter << CacheVersion;
 
 		// serialize the stable shader index table in exact order, but without Output hashes
 		// (for now serialize PipelineHash inline, in hopes it will be later changed to a more stable identifier)
@@ -598,9 +602,9 @@ bool UE::PipelineCacheUtilities::LoadStablePipelineCacheFile(const FString& File
 	}
 
 	// start restrictive, as the format isn't really forward compatible, nor needs to be
-	if (Header.Version != SupportedHeader.Version)
+	if (Header.Version > SupportedHeader.Version)
 	{
-		UE_LOG(LogPipelineCacheUtilities, Warning, TEXT("Rejecting %s, old version (%d vs expected %d)."), *Filename, int(Header.Version), int(SupportedHeader.Version));
+		UE_LOG(LogPipelineCacheUtilities, Warning, TEXT("Rejecting %s, version is too new (%d vs expected %d)."), *Filename, int(Header.Version), int(SupportedHeader.Version));
 		return false;
 	}
 
@@ -644,7 +648,17 @@ bool UE::PipelineCacheUtilities::LoadStablePipelineCacheFile(const FString& File
 	FMemoryReader PlainMemReader(UncompressedMemory);
 	UE::PipelineCacheUtilities::Private::FIndexedSHAHashAndFNameProxyArchive MemReader(PlainMemReader);
 
-	MemReader.SetGameNetVer(FPipelineCacheFileFormatCurrentVersion);
+	uint32 CacheVersion = FPipelineCacheFileFormatCurrentVersion;
+	if (Header.Version >= UE::PipelineCacheUtilities::Private::FStablePipelineCacheSerializedHeader::EVersion::AddingPipelineCacheVersion)
+	{
+		MemReader << CacheVersion;
+	}
+	else
+	{
+		CacheVersion = 26; //EPipelineCacheFileFormatVersions::BeforeStableCacheVersioning
+	}
+
+	MemReader.SetGameNetVer(CacheVersion);
 
 	// read the stable keys as saved
 	TArray<FStableShaderKeyAndValue> SavedStableKeys;
