Do not try to open the shaderlibrary from a pak file if using IoStore

#rnx

diff --git a/Engine/Source/Runtime/RenderCore/Private/ShaderCodeLibrary.cpp b/Engine/Source/Runtime/RenderCore/Private/ShaderCodeLibrary.cpp
--- a/Engine/Source/Runtime/RenderCore/Private/ShaderCodeLibrary.cpp
+++ b/Engine/Source/Runtime/RenderCore/Private/ShaderCodeLibrary.cpp
@@ -857,7 +857,9 @@ public:
 			ShaderCodeDirectory.Empty();	// paths don't matter for IoStore-based libraries
 		}
 
-		if (!Library)
+		// Shader library as a ushaderbytecode file is no longer an option for distribution. Some code paths (loose files) still need it,
+		// but treat IoStore being initialized as a sign that the library will only be provided as a IoStore shader code archive.
+		if (!Library && !FIoDispatcher::IsInitialized())
 		{
 			const FName PlatformName = FDataDrivenShaderPlatformInfo::GetName(InShaderPlatform);
 			const FName ShaderFormatName = LegacyShaderPlatformToShaderFormat(InShaderPlatform);
