Fix unitialized variable warning for static analysis.

- ShaderParameterStructureIndex

diff --git a/Engine/Source/Runtime/RenderCore/Public/ShaderCompilerCore.h b/Engine/Source/Runtime/RenderCore/Public/ShaderCompilerCore.h
--- a/Engine/Source/Runtime/RenderCore/Public/ShaderCompilerCore.h
+++ b/Engine/Source/Runtime/RenderCore/Public/ShaderCompilerCore.h
@@ -476,7 +476,7 @@ struct FShaderCompilerInput
 			Environment.Merge(SharedEnvironments[SharedEnvironmentIndex]);
 		}
 
-		int32 ShaderParameterStructureIndex;
+		int32 ShaderParameterStructureIndex = INDEX_NONE;
 		Ar << ShaderParameterStructureIndex;
 		if (ShaderParameterStructureIndex != INDEX_NONE)
 		{
