[shaders] further deprecation (SupportsIndependentPreprocessing/OptionalFinalShaderSource/bIndependentPreprocessed are no longer needed since all backends are migrated)

diff --git a/Engine/Source/Developer/Apple/MetalShaderFormat/Private/MetalShaderCompiler.cpp b/Engine/Source/Developer/Apple/MetalShaderFormat/Private/MetalShaderCompiler.cpp
--- a/Engine/Source/Developer/Apple/MetalShaderFormat/Private/MetalShaderCompiler.cpp
+++ b/Engine/Source/Developer/Apple/MetalShaderFormat/Private/MetalShaderCompiler.cpp
@@ -555,11 +555,7 @@ void BuildMetalShaderOutput(
 		Ar << Header;
 		Ar.Serialize((void*)USFSource, SourceLen + 1 - (USFSource - InShaderSource));
 		
-		if (ShaderInput.ExtraSettings.bExtractShaderSource)
-		{
-			ShaderOutput.OptionalFinalShaderSource = MetalCode;
-		}
-
+		ShaderOutput.ModifiedShaderSource = MetalCode;
 		ShaderOutput.NumInstructions = NumLines;
 		ShaderOutput.NumTextureSamplers = Header.Bindings.NumSamplers;
 		ShaderOutput.bSucceeded = true;
@@ -745,11 +741,7 @@ void BuildMetalShaderOutput(
 			ShaderOutput.NumTextureSamplers = Header.Bindings.NumSamplers;
 		}
 
-		if (ShaderInput.ExtraSettings.bExtractShaderSource)
-		{
-			ShaderOutput.OptionalFinalShaderSource = MetalCode;
-		}
-		
+		ShaderOutput.ModifiedShaderSource = MetalCode;
 		ShaderOutput.NumInstructions = NumLines;
 		ShaderOutput.bSucceeded = bSucceeded;
 	}
diff --git a/Engine/Source/Developer/Apple/MetalShaderFormat/Private/MetalShaderFormat.cpp b/Engine/Source/Developer/Apple/MetalShaderFormat/Private/MetalShaderFormat.cpp
--- a/Engine/Source/Developer/Apple/MetalShaderFormat/Private/MetalShaderFormat.cpp
+++ b/Engine/Source/Developer/Apple/MetalShaderFormat/Private/MetalShaderFormat.cpp
@@ -131,11 +131,6 @@ public:
 		OutFormats.Add(NAME_SF_METAL_MRT_MAC);
 	}
 
-	virtual bool SupportsIndependentPreprocessing() const override final
-	{
-		return true;
-	}
-
 	void CheckShaderFormat(FName Format) const
 	{
 		check(Format == NAME_SF_METAL
diff --git a/Engine/Source/Developer/ShaderFormatOpenGL/Private/OpenGLShaderCompiler.cpp b/Engine/Source/Developer/ShaderFormatOpenGL/Private/OpenGLShaderCompiler.cpp
--- a/Engine/Source/Developer/ShaderFormatOpenGL/Private/OpenGLShaderCompiler.cpp
+++ b/Engine/Source/Developer/ShaderFormatOpenGL/Private/OpenGLShaderCompiler.cpp
@@ -890,13 +890,7 @@ void BuildShaderOutputInternal(
 
 		Ar.Serialize((void*)USFSource, SourceLen + 1 - (USFSource - InShaderSource));
 
-		// extract final source code as requested by the Material Editor
-		if (ShaderInput.ExtraSettings.bExtractShaderSource)
-		{
-			TArray<ANSICHAR> GlslCodeOriginal;
-			GlslCodeOriginal.Append(USFSource, FCStringAnsi::Strlen(USFSource) + 1);
-			ShaderOutput.OptionalFinalShaderSource = FString(GlslCodeOriginal.GetData());
-		}
+		ShaderOutput.ModifiedShaderSource = USFSource;
 
 		if (ShaderInput.Environment.CompilerFlags.Contains(CFLAG_ExtraShaderData))
 		{
diff --git a/Engine/Source/Developer/ShaderFormatOpenGL/Private/ShaderFormatOpenGL.cpp b/Engine/Source/Developer/ShaderFormatOpenGL/Private/ShaderFormatOpenGL.cpp
--- a/Engine/Source/Developer/ShaderFormatOpenGL/Private/ShaderFormatOpenGL.cpp
+++ b/Engine/Source/Developer/ShaderFormatOpenGL/Private/ShaderFormatOpenGL.cpp
@@ -86,11 +86,6 @@ public:
 		}
 	}
 
-	virtual bool SupportsIndependentPreprocessing() const override
-	{
-		return true;
-	}
-
 	virtual bool PreprocessShader(const FShaderCompilerInput& Input, const FShaderCompilerEnvironment& Environment, FShaderPreprocessOutput& PreprocessOutput) const override
 	{
 		CheckFormat(Input.ShaderFormat);
diff --git a/Engine/Source/Developer/ShaderFormatVectorVM/Private/ShaderFormatVectorVM.cpp b/Engine/Source/Developer/ShaderFormatVectorVM/Private/ShaderFormatVectorVM.cpp
--- a/Engine/Source/Developer/ShaderFormatVectorVM/Private/ShaderFormatVectorVM.cpp
+++ b/Engine/Source/Developer/ShaderFormatVectorVM/Private/ShaderFormatVectorVM.cpp
@@ -63,11 +63,6 @@ public:
 		OutFormats.Add(NAME_VVM_1_0);
 	}
 
-	virtual bool SupportsIndependentPreprocessing() const override
-	{
-		return true;
-	}
-
 	virtual bool PreprocessShader(const FShaderCompilerInput& Input, const FShaderCompilerEnvironment& Environment, FShaderPreprocessOutput& PreprocessOutput) const override
 	{
 		CheckFormat(Input.ShaderFormat);
diff --git a/Engine/Source/Developer/TargetPlatform/Public/Interfaces/IShaderFormat.h b/Engine/Source/Developer/TargetPlatform/Public/Interfaces/IShaderFormat.h
--- a/Engine/Source/Developer/TargetPlatform/Public/Interfaces/IShaderFormat.h
+++ b/Engine/Source/Developer/TargetPlatform/Public/Interfaces/IShaderFormat.h
@@ -20,14 +20,6 @@ class IShaderFormat
 {
 public:
 
-	/** 
-	 * Compile the specified shader. Only called if SupportsIndependentPreprocessing returns false.
-	 *
-	 * @param Format The desired format
-	 * @param Input The input to the shader compiler.
-	 * @param Output The output from shader compiler.
-	 * @param WorkingDirectory The working directory.
-	 */
 	UE_DEPRECATED(5.4, "Monolithic CompileShader is deprecated. Please implement separate compilation and preprocess via PreprocessShader/CompilePreprocessedShader")
 	virtual void CompileShader(FName Format, const FShaderCompilerInput& Input, FShaderCompilerOutput& Output, const FString& WorkingDirectory) const {};
 
@@ -135,12 +127,12 @@ public:
 	virtual bool UsesHLSLcc(const FShaderCompilerInput& Input) const { return false; }
 
 	/**
-	 * Executes all shader preprocessing steps; only called if SupportsIndependentPreprocessing returns true
+	 * Execute all shader preprocessing steps, storing the output in the PreprocessOutput struct
 	 */
 	virtual bool PreprocessShader(const FShaderCompilerInput& Input, const FShaderCompilerEnvironment& Environment, FShaderPreprocessOutput& PreprocessOutput) const { return false; };
 
 	/**
-	 * Compile the specified preprocessed shader; only called if SupportsIndependentPreprocessing returns true
+	 * Compile the specified preprocessed shader.
 	 */
 	virtual void CompilePreprocessedShader(
 		const FShaderCompilerInput& Input, 
@@ -149,9 +141,8 @@ public:
 		const FString& WorkingDirectory) const {}
 
 	/**
-	 * Compile the specified preprocessed shaders; only called if SupportsIndependentPreprocessing returns true, and the call to
-	 * RequiresSecondaryCompile given the first preprocess output also returns true. The shader system will pack these outputs together
-	 * in the following format: 
+	 * Compile the specified preprocessed shaders; only called if the call to RequiresSecondaryCompile given the first preprocess 
+	 * output returns true. The shader system will pack these outputs together in the following format: 
 	 * [int32 key][uint32 primary length][uint32 secondary length][full primary shader code][full secondary shader code]
 	 * where "key" is the return value of the GetPackedShaderKey function (this should also be implemented by any backends which
 	 * require secondary compilation and is used by the RHI to differentiate packed shader code from single shaders)
@@ -166,7 +157,6 @@ public:
 
 	/**
 	 * Predicate which should return true if a second preprocess & compilation is required given the initial preprocess output.
-	 * Only called if SupportsIndependentPreprocessing returns true.
 	 * This is generally be determined by analyzing the directives on the given preprocess output (which are set by the presence
 	 * of UESHADERMETADATA directives in the original source), though other conditions are possible and up to the implementation.
 	 *  In the event this returns true, PreprocessShader will be called an additional time, with the bIsSecondary field on the 
@@ -188,7 +178,7 @@ public:
 	virtual int32 GetPackedShaderKey() const { return 0; }
 
 	/* 
-	 * Implement to output debug info for a single compile job; only called if SupportsIndependentPreprocessing returns true.
+	 * Implement to output debug info for a single compile job.
 	 * This will be called for all jobs (including those found in the job cache) but only if debug info is enabled for the job.
 	 * Note that any debug info output in CompilePreprocessedShader will only be done for the job that actually executes the
 	 * compile step, as such any debug outputs that are desirable for all jobs should be written by this function.
@@ -214,10 +204,7 @@ public:
 		const FShaderCompilerOutput& Output,
 		const FShaderCompilerOutput& SecondaryOutput) const {};
 
-	/**
-	 * Return true if preprocessing for this format can be executed independent of compilation (i.e. the format has an
-	 * implementation of the PreprocessShader function). 
-	 */
+	UE_DEPRECATED(5.4, "SupportsIndependentPreprocessing is no longer called now that all shader backends have been migrated.")
 	virtual bool SupportsIndependentPreprocessing() const { return false; }
 
 public:
diff --git a/Engine/Source/Developer/VulkanShaderFormat/Private/VulkanShaderCompiler.cpp b/Engine/Source/Developer/VulkanShaderFormat/Private/VulkanShaderCompiler.cpp
--- a/Engine/Source/Developer/VulkanShaderFormat/Private/VulkanShaderCompiler.cpp
+++ b/Engine/Source/Developer/VulkanShaderFormat/Private/VulkanShaderCompiler.cpp
@@ -1143,12 +1143,16 @@ static void BuildShaderOutput(
 
 	ShaderOutput.bSucceeded = true;
 
+	// guard disassembly of SPIRV code on bExtractShaderSource setting since presumably this isn't that cheap.
+	// this roughly will maintain existing behaviour, except the debug usf will be this version of the code 
+	// instead of the output of  preprocessing if this setting is enabled (which is probably fine since this is only
+	// ever set in editor)
 	if (ShaderInput.ExtraSettings.bExtractShaderSource)
 	{
 		TArray<ANSICHAR> AssemblyText;
 		if (CrossCompiler::FShaderConductorContext::Disassemble(CrossCompiler::EShaderConductorIR::Spirv, SerializedOutput.Spirv.GetByteData(), SerializedOutput.Spirv.GetByteSize(), AssemblyText))
 		{
-			ShaderOutput.OptionalFinalShaderSource = FString(AssemblyText.GetData());
+			ShaderOutput.ModifiedShaderSource = FString(AssemblyText.GetData());
 		}
 	}
 	if (ShaderInput.ExtraSettings.OfflineCompilerPath.Len() > 0)
diff --git a/Engine/Source/Developer/VulkanShaderFormat/Private/VulkanShaderFormat.cpp b/Engine/Source/Developer/VulkanShaderFormat/Private/VulkanShaderFormat.cpp
--- a/Engine/Source/Developer/VulkanShaderFormat/Private/VulkanShaderFormat.cpp
+++ b/Engine/Source/Developer/VulkanShaderFormat/Private/VulkanShaderFormat.cpp
@@ -93,11 +93,6 @@ public:
 		OutFormats.Add(NAME_VULKAN_SM6);
 	}
 
-	virtual bool SupportsIndependentPreprocessing() const override
-	{
-		return true;
-	}
-
 	virtual bool PreprocessShader(const FShaderCompilerInput& Input, const FShaderCompilerEnvironment& Environment, FShaderPreprocessOutput& PreprocessOutput) const
 	{
 		check(InternalGetVersion(Input.ShaderFormat) >= 0);
diff --git a/Engine/Source/Developer/Windows/ShaderFormatD3D/Private/ShaderFormatD3D.cpp b/Engine/Source/Developer/Windows/ShaderFormatD3D/Private/ShaderFormatD3D.cpp
--- a/Engine/Source/Developer/Windows/ShaderFormatD3D/Private/ShaderFormatD3D.cpp
+++ b/Engine/Source/Developer/Windows/ShaderFormatD3D/Private/ShaderFormatD3D.cpp
@@ -116,11 +116,6 @@ public:
 		Output.ShaderDiagnosticDatas = PreprocessOutput.GetDiagnosticDatas();
 	}
 
-	virtual bool SupportsIndependentPreprocessing() const
-	{
-		return true;
-	}
-
 	void AddShaderTargetDefines(FShaderCompilerInput& Input, uint32 ShaderTargetMajor, uint32 ShaderTargetMinor) const
 	{
 		// Inserting our own versions of these defines since we preprocess our shader source before we actually use something that defines them.
diff --git a/Engine/Source/Runtime/Engine/Private/ShaderCompiler/ShaderCompiler.cpp b/Engine/Source/Runtime/Engine/Private/ShaderCompiler/ShaderCompiler.cpp
--- a/Engine/Source/Runtime/Engine/Private/ShaderCompiler/ShaderCompiler.cpp
+++ b/Engine/Source/Runtime/Engine/Private/ShaderCompiler/ShaderCompiler.cpp
@@ -8290,13 +8290,9 @@ void GlobalBeginCompileShader(
 	checkf(Format, TEXT("Shader format %s cannot be found"), *ShaderFormatName.ToString());
 	Format->ModifyShaderCompilerInput(Input);
 
-	if (Format->SupportsIndependentPreprocessing())
+	if (CVarPreprocessedJobCache.GetValueOnAnyThread())
 	{
-		Input.bIndependentPreprocessed = true;
-		if (CVarPreprocessedJobCache.GetValueOnAnyThread())
-		{
-			Input.bCachePreprocessed = true;
-		}
+		Input.bCachePreprocessed = true;
 	}
 
 	// Allow the GBuffer and other shader defines to cause dependend environment changes, but minimizing the #ifdef magic in the shaders, which
diff --git a/Engine/Source/Runtime/RenderCore/Private/ShaderCore.cpp b/Engine/Source/Runtime/RenderCore/Private/ShaderCore.cpp
--- a/Engine/Source/Runtime/RenderCore/Private/ShaderCore.cpp
+++ b/Engine/Source/Runtime/RenderCore/Private/ShaderCore.cpp
@@ -1328,40 +1328,33 @@ public:
 	static void CompileShaderInternal(const IShaderFormat* Compiler, FShaderCompileJob& Job, const FString& WorkingDirectory, FString& OutExceptionMsg, FString& OutExceptionCallstack, int32* CompileCount)
 	{
 		double TimeStart = FPlatformTime::Seconds();
-		if (Compiler->SupportsIndependentPreprocessing())
+		if (!Job.Input.bCachePreprocessed)
 		{
-			if (!Job.Input.bCachePreprocessed)
-			{
-				PreprocessShaderInternal(Compiler, Job);
-			}
+			PreprocessShaderInternal(Compiler, Job);
+		}
 
-			Job.Output.Errors.Append(Job.PreprocessOutput.Errors);
+		Job.Output.Errors.Append(Job.PreprocessOutput.Errors);
 
-			if (Job.PreprocessOutput.bSucceeded)
+		if (Job.PreprocessOutput.bSucceeded)
+		{
+			if (Job.SecondaryPreprocessOutput.IsValid())
 			{
-				if (Job.SecondaryPreprocessOutput.IsValid())
-				{
-					Job.SecondaryOutput = MakeUnique<FShaderCompilerOutput>();
-				}
-				InvokeCompile(Compiler, Job, WorkingDirectory, OutExceptionMsg, OutExceptionCallstack);
-				if (Job.SecondaryOutput.IsValid())
-				{
-					Job.Output.bSucceeded = Job.Output.bSucceeded && Job.SecondaryOutput->bSucceeded;
-					if (Job.Output.bSucceeded)
-					{
-						Job.SecondaryOutput->GenerateOutputHash();
-					}
-					CombineOutputs(Compiler, Job);
-				}
+				Job.SecondaryOutput = MakeUnique<FShaderCompilerOutput>();
 			}
-			else
+			InvokeCompile(Compiler, Job, WorkingDirectory, OutExceptionMsg, OutExceptionCallstack);
+			if (Job.SecondaryOutput.IsValid())
 			{
-				Job.Output.bSucceeded = false;
+				Job.Output.bSucceeded = Job.Output.bSucceeded && Job.SecondaryOutput->bSucceeded;
+				if (Job.Output.bSucceeded)
+				{
+					Job.SecondaryOutput->GenerateOutputHash();
+				}
+				CombineOutputs(Compiler, Job);
 			}
 		}
 		else
 		{
-			InvokeCompile(Compiler, Job, WorkingDirectory, OutExceptionMsg, OutExceptionCallstack);
+			Job.Output.bSucceeded = false;
 		}
 
 		if (Job.Output.bSucceeded)
@@ -1453,27 +1446,18 @@ void CompileShader(const TArray<const IShaderFormat*>& ShaderFormats, FShaderCom
 	Job.bSucceeded = Job.Output.bSucceeded;
 	if (Job.Input.DumpDebugInfoEnabled())
 	{
-		if (Compiler->SupportsIndependentPreprocessing())
+		// if the preprocessed cache is disabled, dump debug output here, since we don't serialize preprocess output back to the cooker from SCW
+		// (if enabled this will occur in the job OnComplete callback)
+		if (!Job.Input.bCachePreprocessed)
 		{
-			// if the shader format supports independent preprocessing and preprocessed cache is disabled, dump debug output here, since we
-			// don't serialize preprocess output back to the cooker from SCW (if enabled this will occur in the job OnComplete callback)
-			if (!Job.Input.bCachePreprocessed)
+			if (Job.SecondaryPreprocessOutput.IsValid() && Job.SecondaryOutput.IsValid())
 			{
-				if (Job.SecondaryPreprocessOutput.IsValid() && Job.SecondaryOutput.IsValid())
-				{
-					Compiler->OutputDebugData(Job.Input, Job.PreprocessOutput, *Job.SecondaryPreprocessOutput, Job.Output, *Job.SecondaryOutput);
-				}
-				else
-				{
-					Compiler->OutputDebugData(Job.Input, Job.PreprocessOutput, Job.Output);
-				}
+				Compiler->OutputDebugData(Job.Input, Job.PreprocessOutput, *Job.SecondaryPreprocessOutput, Job.Output, *Job.SecondaryOutput);
+			}
+			else
+			{
+				Compiler->OutputDebugData(Job.Input, Job.PreprocessOutput, Job.Output);
 			}
-		}
-		else
-		{
-			// write down the output hash as a file
-			FString HashFileName = FPaths::Combine(Job.Input.DumpDebugInfoPath, TEXT("OutputHash.txt"));
-			FFileHelper::SaveStringToFile(Job.Output.OutputHash.ToString(), *HashFileName, FFileHelper::EEncodingOptions::ForceAnsi);
 		}
 	}
 }
@@ -3548,7 +3532,6 @@ FArchive& operator<<(FArchive& Ar, FShaderCompilerInput& Input)
 	Ar << Input.bCompilingForShaderPipeline;
 	Ar << Input.bIncludeUsedOutputs;
 	Ar << Input.bCachePreprocessed;
-	Ar << Input.bIndependentPreprocessed;
 	Ar << Input.UsedOutputs;
 	Ar << Input.DumpDebugInfoRootPath;
 	Ar << Input.DumpDebugInfoPath;
@@ -3859,33 +3842,30 @@ void FShaderCompileJob::SerializeOutput(FArchive& Ar)
 void FShaderCompileJob::OnComplete()
 {
 	const IShaderFormat* ShaderFormat = GetTargetPlatformManagerRef().FindShaderFormat(Input.ShaderFormat);
-	if (ShaderFormat->SupportsIndependentPreprocessing())
+	// For jobs using the preprocessed cache, we need to remap error messages whether or not the job was actually the one that ran
+	// the compilation step. In addition since we always run preprocessing we set the total preprocess time accordingly.
+	if (Input.bCachePreprocessed)
+	{
+		PreprocessOutput.RemapErrors(Output);
+		Output.PreprocessTime = PreprocessOutput.ElapsedTime;
+	}
+	// dump debug info for the job at this point if the preprocessed cache is enabled
+	// this ensures we get debug output for all jobs, including those that were found in the job cache,
+	// or matched another in-flight job's hash and so could share its results
+	if (Input.bCachePreprocessed 
+		&& Input.DumpDebugInfoEnabled()
+		// if we only want debug info for jobs which actually compiled, check the CompileTime
+		// (jobs deserialized from the cache/wait list/ddc will have a compiletime of 0.0)
+		&& (CVarDumpDebugInfoForCacheHits.GetValueOnAnyThread() || Output.CompileTime > 0.0f))
 	{
-		// For jobs using the preprocessed cache, we need to remap error messages whether or not the job was actually the one that ran
-		// the compilation step. In addition since we always run preprocessing we set the total preprocess time accordingly.
-		if (Input.bCachePreprocessed)
+		
+		if (SecondaryPreprocessOutput.IsValid() && SecondaryOutput.IsValid())
 		{
-			PreprocessOutput.RemapErrors(Output);
-			Output.PreprocessTime = PreprocessOutput.ElapsedTime;
+			ShaderFormat->OutputDebugData(Input, PreprocessOutput, *SecondaryPreprocessOutput, Output, *SecondaryOutput);
 		}
-		// dump debug info for the job at this point if the preprocessed cache is enabled
-		// this ensures we get debug output for all jobs, including those that were found in the job cache,
-		// or matched another in-flight job's hash and so could share its results
-		if (Input.bCachePreprocessed 
-			&& Input.DumpDebugInfoEnabled()
-			// if we only want debug info for jobs which actually compiled, check the CompileTime
-			// (jobs deserialized from the cache/wait list/ddc will have a compiletime of 0.0)
-			&& (CVarDumpDebugInfoForCacheHits.GetValueOnAnyThread() || Output.CompileTime > 0.0f))
+		else
 		{
-			
-			if (SecondaryPreprocessOutput.IsValid() && SecondaryOutput.IsValid())
-			{
-				ShaderFormat->OutputDebugData(Input, PreprocessOutput, *SecondaryPreprocessOutput, Output, *SecondaryOutput);
-			}
-			else
-			{
-				ShaderFormat->OutputDebugData(Input, PreprocessOutput, Output);
-			}
+			ShaderFormat->OutputDebugData(Input, PreprocessOutput, Output);
 		}
 	}
 }
@@ -3920,7 +3900,7 @@ void FShaderCompileJob::SerializeWorkerOutput(FArchive& Ar)
 	// edge case for backends which have implemented independent preprocessing API when the preprocessed cache is not enabled.
 	// if no modifications have occurred as part of the compile step, we still need a copy of the source back in the cooker
 	// if bExtractShaderSource is set, so explicitly serialize just that portion of the preprocess output struct here.
-	if (Input.ExtraSettings.bExtractShaderSource && Input.bIndependentPreprocessed && !Input.bCachePreprocessed && Output.ModifiedShaderSource.IsEmpty())
+	if (Input.ExtraSettings.bExtractShaderSource && !Input.bCachePreprocessed && Output.ModifiedShaderSource.IsEmpty())
 	{
 		Ar << PreprocessOutput.EditSource();
 	}
@@ -3954,28 +3934,19 @@ void FShaderCompileJob::SerializeWorkerInput(FArchive& Ar)
 
 const FString& FShaderCompileJob::GetFinalSource() const
 {
-	 // if the backend supports independent preprocessing, any modifications to the source
-	// done as part of the compile step will be written to the "ModifiedShaderSource" field
-	if (Input.bIndependentPreprocessed)
+	 // any modifications to the source done as part of the compile step will be written to the "ModifiedShaderSource" field
+	// always return empty string if source extraction was not requested; this will prevent bloat of material DDC data in the case where debug info is enabled 
+	// or Output.ModifiedShaderSource is unset (since the preprocess output unstripped source will always be set)
+	if (Input.ExtraSettings.bExtractShaderSource)
 	{
-		// always return empty string if source extraction was not requested; this will prevent bloat of material DDC data in the case where debug info is enabled 
-		// or Output.ModifiedShaderSource is unset (since the preprocess output unstripped source will always be set)
-		if (Input.ExtraSettings.bExtractShaderSource)
-		{
-			// if there are no such modifications, return the "unstripped" version of the source code (with comments & line directives maintained),
-			// otherwise return whatever the final modified source is as input to the compiler by the backend.
-			return Output.ModifiedShaderSource.IsEmpty() ? PreprocessOutput.GetUnstrippedSource() : Output.ModifiedShaderSource;
-		}
-		else
-		{
-			static FString Empty;
-			return Empty;
-		}
+		// if there are no such modifications, return the "unstripped" version of the source code (with comments & line directives maintained),
+		// otherwise return whatever the final modified source is as input to the compiler by the backend.
+		return Output.ModifiedShaderSource.IsEmpty() ? PreprocessOutput.GetUnstrippedSource() : Output.ModifiedShaderSource;
 	}
 	else
 	{
-		// backends that do not implement the independent preprocessing API populate this field based on the bExtractShaderSource setting.
-		return Output.OptionalFinalShaderSource;
+		static FString Empty;
+		return Empty;
 	}
 }
 
diff --git a/Engine/Source/Runtime/RenderCore/Public/ShaderCompilerCore.h b/Engine/Source/Runtime/RenderCore/Public/ShaderCompilerCore.h
--- a/Engine/Source/Runtime/RenderCore/Public/ShaderCompilerCore.h
+++ b/Engine/Source/Runtime/RenderCore/Public/ShaderCompilerCore.h
@@ -31,16 +31,16 @@ class FShaderPipelineCompileJob;
 typedef TSharedPtr<TArray<ANSICHAR>, ESPMode::ThreadSafe> FShaderSharedAnsiStringPtr;
 
 // this is for the protocol, not the data, bump if FShaderCompilerInput or ProcessInputFromArchive changes.
-inline const int32 ShaderCompileWorkerInputVersion = 22;
+inline const int32 ShaderCompileWorkerInputVersion = 23;
 // this is for the protocol, not the data, bump if FShaderCompilerOutput or WriteToOutputArchive changes.
-inline const int32 ShaderCompileWorkerOutputVersion = 18;
+inline const int32 ShaderCompileWorkerOutputVersion = 19;
 // this is for the protocol, not the data.
 inline const int32 ShaderCompileWorkerSingleJobHeader = 'S';
 // this is for the protocol, not the data.
 inline const int32 ShaderCompileWorkerPipelineJobHeader = 'P';
 
 // modify this for changes to the FShaderCompilerOutput data structure (in addition to ShaderCompileWorkerOutputVersion)
-inline const int32 FShaderCompilerOutputStructVersion = 2;
+inline const int32 FShaderCompilerOutputStructVersion = 3;
 
 namespace UE::ShaderCompiler
 {
@@ -246,6 +246,7 @@ struct FShaderCompilerInput
 	// Indicates which additional debug outputs should be written for this compile job.
 	EShaderDebugInfoFlags DebugInfoFlags;
 
+	UE_DEPRECATED(5.4, "bIndependentPreprocessed member no longer used now that all backends have been migrated to the new IShaderFormat API")
 	// True if the backend for this job implements the independent preprocessing API.
 	bool bIndependentPreprocessed;
 
@@ -298,7 +299,6 @@ struct FShaderCompilerInput
 		Target(SF_NumFrequencies, SP_NumPlatforms),
 		bSkipPreprocessedCache(false),
 		DebugInfoFlags(EShaderDebugInfoFlags::Default),
-		bIndependentPreprocessed(false),
 		bCachePreprocessed(false),
 		bCompilingForShaderPipeline(false),
 		bIncludeUsedOutputs(false)
@@ -642,17 +642,13 @@ struct FShaderCompilerOutput
 
 	TArray<FShaderDiagnosticData> ShaderDiagnosticDatas;
 
-	/** This field should be set by backends which do not implement the independent preprocessing API to contain the "final" shader source as 
-	 * passed to the platform compiler. For backends that do implement this API this is superceded by ModifiedShaderSource (and will eventually
-	 * be deprecated).
-	 */
+	UE_DEPRECATED(5.4, "OptionalFinalShaderSource is no longer used; set ModifiedShaderSource instead if the shader backend makes additional source code manipulations.")
 	FString OptionalFinalShaderSource;
 
-	/** Use this field to store the shader source code if it's modified as part of the shader format's compilation process. This field is only 
-	 * currently required for shader formats which implement the independent preprocessing API and should only be set when additional manipulation 
-	 * is required that is not part of the implementation of PreprocessShader. This version of the source, if set, will be what is written as part
-	 * of the debug dumps of preprocessed source, as well as used in place of OptionalFinalShaderSource for upstream code which explicitly requests
-	 * the final source code for other purposes (i.e. when ExtraSettings.bExtractShaderSource is set on the FShaderCompilerInput struct)
+	/** Use this field to store the shader source code if it's modified as part of the shader format's compilation process. This should only be set when 
+	 * additional manipulation of source code  is required that is not part of the implementation of PreprocessShader. This version of the source, if set, 
+	 * will be what is written as part of the debug dumps of preprocessed source, as well as used for upstream code which explicitly requests the final
+	 * source code for other purposes (i.e. when ExtraSettings.bExtractShaderSource is set on the FShaderCompilerInput struct)
 	 */
 	FString ModifiedShaderSource;
 
@@ -682,7 +678,6 @@ struct FShaderCompilerOutput
 		Ar << Output.bSupportsQueryingUsedAttributes << Output.UsedAttributes;
 		Ar << Output.CompileTime;
 		Ar << Output.PreprocessTime;
-		Ar << Output.OptionalFinalShaderSource;
 		Ar << Output.bSerializeModifiedSource;
 		if (Output.bSerializeModifiedSource)
 		{
