Fix Nanite splitscreen issues.

- Make sure ViewRect.Min is at least 8x8-aligned when Nanite is enabled.

diff --git a/Engine/Source/Runtime/RenderCore/Private/RenderUtils.cpp b/Engine/Source/Runtime/RenderCore/Private/RenderUtils.cpp
--- a/Engine/Source/Runtime/RenderCore/Private/RenderUtils.cpp
+++ b/Engine/Source/Runtime/RenderCore/Private/RenderUtils.cpp
@@ -1128,7 +1128,7 @@ RENDERCORE_API FBufferRHIRef& GetUnitCubeAABBVertexBuffer()
 }
 #endif // RHI_RAYTRACING
 
-RENDERCORE_API void QuantizeSceneBufferSize(const FIntPoint& InBufferSize, FIntPoint& OutBufferSize)
+RENDERCORE_API void QuantizeSceneBufferSize(const FIntPoint& InBufferSize, FIntPoint& OutBufferSize, const uint32 SuggestedDivisor)
 {
 	// Ensure sizes are dividable by SUBSTRATE_TILE_SIZE (==8) 2d tiles to make it more convenient.
 	const uint32 SubstrateDividableBy = SUBSTRATE_TILE_SIZE;
@@ -1138,7 +1138,7 @@ RENDERCORE_API void QuantizeSceneBufferSize(const FIntPoint& InBufferSize, FIntP
 	const uint32 LegacyDividableBy = 4;
 	static_assert(LegacyDividableBy % 4 == 0, "A lot of graphic algorithms where previously assuming DividableBy == 4");
 
-	const uint32 DividableBy = Substrate::IsSubstrateEnabled() ? SubstrateDividableBy : LegacyDividableBy;
+	const uint32 DividableBy = FMath::Max(Substrate::IsSubstrateEnabled() ? SubstrateDividableBy : LegacyDividableBy, SuggestedDivisor);
 
 	const uint32 Mask = ~(DividableBy - 1);
 	OutBufferSize.X = (InBufferSize.X + DividableBy - 1) & Mask;
diff --git a/Engine/Source/Runtime/RenderCore/Public/RenderUtils.h b/Engine/Source/Runtime/RenderCore/Public/RenderUtils.h
--- a/Engine/Source/Runtime/RenderCore/Public/RenderUtils.h
+++ b/Engine/Source/Runtime/RenderCore/Public/RenderUtils.h
@@ -362,7 +362,7 @@ RENDERCORE_API FBufferRHIRef& GetUnitCubeAABBVertexBuffer();
 * rendering pipeline. Currently ensures that sizes are multiples of 4 so that they can safely
 * be halved in size several times.
 */
-RENDERCORE_API void QuantizeSceneBufferSize(const FIntPoint& InBufferSize, FIntPoint& OutBufferSize);
+RENDERCORE_API void QuantizeSceneBufferSize(const FIntPoint& InBufferSize, FIntPoint& OutBufferSize, const uint32 SuggestedDivisor = 0);
 
 /**
 * Checks if virtual texturing enabled and supported
diff --git a/Engine/Source/Runtime/Renderer/Private/SceneRendering.cpp b/Engine/Source/Runtime/Renderer/Private/SceneRendering.cpp
--- a/Engine/Source/Runtime/Renderer/Private/SceneRendering.cpp
+++ b/Engine/Source/Runtime/Renderer/Private/SceneRendering.cpp
@@ -2831,7 +2831,13 @@ FIntPoint FSceneRenderer::ApplyResolutionFraction(const FSceneViewFamily& ViewFa
 FIntPoint FSceneRenderer::QuantizeViewRectMin(const FIntPoint& ViewRectMin)
 {
 	FIntPoint Out;
-	QuantizeSceneBufferSize(ViewRectMin, Out);
+
+	// Some code paths of Nanite require that view rect is aligned on 8x8 boundary.
+	static const auto EnableNaniteCVar = IConsoleManager::Get().FindConsoleVariable(TEXT("r.Nanite"));
+	const bool bNaniteEnabled = (EnableNaniteCVar != nullptr) ? (EnableNaniteCVar->GetInt() != 0) : true;
+	const int kMinimumNaniteDivisor = 8;	// HTILE size
+
+	QuantizeSceneBufferSize(ViewRectMin, Out, bNaniteEnabled ? kMinimumNaniteDivisor : 0);
 	return Out;
 }
 
