Removing FShaderParameterParser from FShaderPreprocessOutput.

diff --git a/Engine/Source/Developer/ShaderFormatOpenGL/Private/OpenGLShaderCompiler.cpp b/Engine/Source/Developer/ShaderFormatOpenGL/Private/OpenGLShaderCompiler.cpp
--- a/Engine/Source/Developer/ShaderFormatOpenGL/Private/OpenGLShaderCompiler.cpp
+++ b/Engine/Source/Developer/ShaderFormatOpenGL/Private/OpenGLShaderCompiler.cpp
@@ -10,6 +10,7 @@
 #include "Serialization/MemoryWriter.h"
 #include "ShaderCompilerDefinitions.h"
 #include "ShaderFormatOpenGL.h"
+#include "ShaderParameterParser.h"
 #include "ShaderPreprocessTypes.h"
 #include "SpirvReflectCommon.h"
 #include <algorithm>
diff --git a/Engine/Source/Developer/VulkanShaderFormat/Private/VulkanShaderCompiler.cpp b/Engine/Source/Developer/VulkanShaderFormat/Private/VulkanShaderCompiler.cpp
--- a/Engine/Source/Developer/VulkanShaderFormat/Private/VulkanShaderCompiler.cpp
+++ b/Engine/Source/Developer/VulkanShaderFormat/Private/VulkanShaderCompiler.cpp
@@ -246,7 +246,7 @@ static const FString kHitGroupSystemRootConstantsSymbolName = TEXT("HitGroupSyst
 // A collection of states and data that is locked in at the top level call and doesn't change throughout the compilation process
 struct FVulkanShaderCompilerInternalState
 {
-	FVulkanShaderCompilerInternalState(const FShaderCompilerInput& InInput, const FShaderParameterParser& InParameterParser)
+	FVulkanShaderCompilerInternalState(const FShaderCompilerInput& InInput, const FShaderParameterParser* InParameterParser)
 		: Input(InInput)
 		, ParameterParser(InParameterParser)
 		, Version(FormatToVersion(Input.ShaderFormat))
@@ -265,7 +265,7 @@ struct FVulkanShaderCompilerInternalState
 	}
 
 	const FShaderCompilerInput& Input;
-	const FShaderParameterParser& ParameterParser;
+	const FShaderParameterParser* ParameterParser;
 
 	const EVulkanShaderVersion Version;
 	const CrossCompiler::FShaderConductorOptions::ETargetEnvironment MinimumTargetEnvironment;
@@ -1425,7 +1425,8 @@ static bool BuildShaderOutputFromSpirv(
 				check(SpvResult == SPV_REFLECT_RESULT_SUCCESS);
 
 				const int32 ReflectionSlot = SerializedOutput.Spirv.ReflectionInfo.Add(FVulkanSpirv::FEntry(ResourceName, BindingIndex));
-				const FShaderParameterParser::FParsedShaderParameter* ParsedParam = InternalState.ParameterParser.FindParameterInfosUnsafe(ResourceName);
+				check(InternalState.ParameterParser);
+				const FShaderParameterParser::FParsedShaderParameter* ParsedParam = InternalState.ParameterParser->FindParameterInfosUnsafe(ResourceName);
 
 				auto AddShaderValidationType = [] (uint32_t VulkanBindingIndex, const FShaderParameterParser::FParsedShaderParameter* ParsedParam, FShaderCompilerOutput& Output) {
 					/*if (ParsedParam)
@@ -2155,7 +2156,7 @@ static void RemoveUnusedBindlessHeaps(FString& PreprocessedShaderSource, const T
 
 bool PreprocessVulkanShader(const FShaderCompilerInput& Input, const FShaderCompilerEnvironment& Environment, FShaderPreprocessOutput& PreprocessOutput)
 {
-	const FVulkanShaderCompilerInternalState InternalState(Input, PreprocessOutput.GetParameterParser());
+	const FVulkanShaderCompilerInternalState InternalState(Input, nullptr);
 
 	PRAGMA_DISABLE_DEPRECATION_WARNINGS		// FShaderCompilerDefinitions will be made internal in the future, marked deprecated until then
 
@@ -2552,7 +2553,7 @@ void CompileVulkanShader(const FShaderCompilerInput& Input, const FString& InPre
 		return;
 	}
 
-	FVulkanShaderCompilerInternalState InternalState(Input, ShaderParameterParser);
+	FVulkanShaderCompilerInternalState InternalState(Input, &ShaderParameterParser);
 
 	const EHlslShaderFrequency HlslFrequency = InternalState.GetHlslShaderFrequency();
 	if (HlslFrequency == HSF_InvalidFrequency)
diff --git a/Engine/Source/Runtime/RenderCore/Private/ShaderParameterParser.cpp b/Engine/Source/Runtime/RenderCore/Private/ShaderParameterParser.cpp
--- a/Engine/Source/Runtime/RenderCore/Private/ShaderParameterParser.cpp
+++ b/Engine/Source/Runtime/RenderCore/Private/ShaderParameterParser.cpp
@@ -6,6 +6,19 @@
 #include "String/RemoveFrom.h"
 #include "Misc/StringBuilder.h"
 
+inline FStringView StripTemplateFromType(const FStringView& Input)
+{
+	FStringView UntemplatedType = FStringView(Input);
+	if (int32 Index = Input.Find(TEXT("<")); Index != INDEX_NONE)
+	{
+		// Remove the template argument but don't forget to clean up the type name
+		const int32 NumChars = Input.Len() - Index;
+		UntemplatedType = Input.LeftChop(NumChars).TrimEnd();
+	}
+
+	return UntemplatedType;
+}
+
 template<typename TParameterFunction>
 static void IterateShaderParameterMembersInternal(
 	const FShaderParametersMetadata& ParametersMetadata,
@@ -134,7 +147,6 @@ static void AddNoteToDisplayShaderParameterMemberOnCppSide(
 	CompilerOutput.Errors.Add(Error);
 }
 
-FShaderParameterParser::FShaderParameterParser() = default;
 FShaderParameterParser::~FShaderParameterParser() = default;
 
 PRAGMA_DISABLE_DEPRECATION_WARNINGS
@@ -1424,76 +1436,6 @@ void FShaderParameterParser::ValidateShaderParameterTypes(
 		});
 }
 
-void SerializeStringView(FArchive& Ar, FStringView& StringView, const TCHAR* BasePtr)
-{
-	if (Ar.IsSaving())
-	{
-		uint32 Len = StringView.Len();
-		uint64 Offset = StringView.GetData() - BasePtr;
-		Ar << Len;
-		Ar << Offset;
-	}
-	else
-	{
-		uint32 Len;
-		uint64 Offset;
-		Ar << Len;
-		Ar << Offset;
-		StringView = FStringView(BasePtr + Offset, Len);
-	}
-}
-void SerializeParam(FArchive& Ar, FShaderParameterParser::FParsedShaderParameter& Param, const TCHAR* StringViewBase)
-{
-	Ar << Param.BaseType;
-	Ar << Param.PrecisionModifier;
-	Ar << Param.NumRows;
-	Ar << Param.NumColumns;
-	Ar << Param.MemberSize;
-	SerializeStringView(Ar, Param.ParsedName, StringViewBase);
-	SerializeStringView(Ar, Param.ParsedType, StringViewBase);
-	SerializeStringView(Ar, Param.ParsedArraySize, StringViewBase);
-	Ar << Param.ConstantBufferOffset;
-	Ar << Param.ParsedPragmaLineOffset;
-	Ar << Param.ParsedLineOffset;
-	Ar << Param.ParsedCharOffsetStart;
-	Ar << Param.ParsedCharOffsetEnd;
-	Ar << Param.BindlessConversionType;
-	Ar << Param.ConstantBufferParameterType;
-	Ar << Param.bGloballyCoherent;
-	Ar << Param.bIsBindable;
-	Ar << Param.ParsedTypeDecl;
-}
-FArchive& operator<<(FArchive& Ar, FShaderParameterParser& Parser)
-{
-	Ar << Parser.bMovedLoosedParametersToRootConstantBuffer;
-	Ar << Parser.bModifiedShader;
-	Ar << Parser.OriginalParsedShader;
-	if (Ar.IsSaving())
-	{
-		int32 ParsedParamCount = Parser.ParsedParameters.Num();
-		Ar << ParsedParamCount;
-		for (TPair<FString, FShaderParameterParser::FParsedShaderParameter>& ParsedParam : Parser.ParsedParameters)
-		{
-			Ar << ParsedParam.Key;
-			SerializeParam(Ar, ParsedParam.Value, Parser.OriginalParsedShader.GetCharArray().GetData());
-		}
-	}
-	else
-	{
-		uint32 ParsedParameterCount;
-		Ar << ParsedParameterCount;
-		Parser.ParsedParameters.Reserve(ParsedParameterCount);
-		for (uint32 CurParam = 0; CurParam < ParsedParameterCount; ++CurParam)
-		{
-			FString Name;
-			Ar << Name;
-			FShaderParameterParser::FParsedShaderParameter& ParsedParam = Parser.ParsedParameters.Add(Name);
-			SerializeParam(Ar, ParsedParam, Parser.OriginalParsedShader.GetCharArray().GetData());
-		}
-	}
-	return Ar;
-}
-
 void FShaderParameterParser::ExtractFileAndLine(int32 PragmaLineOffset, int32 LineOffset, FString& OutFile, FString& OutLine) const
 {
 	if (PragmaLineOffset == -1)
diff --git a/Engine/Source/Runtime/RenderCore/Private/ShaderPreprocessTypes.cpp b/Engine/Source/Runtime/RenderCore/Private/ShaderPreprocessTypes.cpp
--- a/Engine/Source/Runtime/RenderCore/Private/ShaderPreprocessTypes.cpp
+++ b/Engine/Source/Runtime/RenderCore/Private/ShaderPreprocessTypes.cpp
@@ -299,7 +299,6 @@ FArchive& operator<<(FArchive& Ar, FShaderPreprocessOutput& PreprocessOutput)
 	Ar << PreprocessOutput.Errors;
 	Ar << PreprocessOutput.PragmaDirectives;
 	Ar << PreprocessOutput.PreprocessedSource;
-	Ar << PreprocessOutput.ParameterParser;
 	Ar << PreprocessOutput.ShaderDiagnosticDatas;
 	return Ar;
 }
diff --git a/Engine/Source/Runtime/RenderCore/Public/ShaderCompilerCore.h b/Engine/Source/Runtime/RenderCore/Public/ShaderCompilerCore.h
--- a/Engine/Source/Runtime/RenderCore/Public/ShaderCompilerCore.h
+++ b/Engine/Source/Runtime/RenderCore/Public/ShaderCompilerCore.h
@@ -31,7 +31,7 @@ class FShaderPipelineCompileJob;
 typedef TSharedPtr<TArray<ANSICHAR>, ESPMode::ThreadSafe> FShaderSharedAnsiStringPtr;
 
 // this is for the protocol, not the data, bump if FShaderCompilerInput or ProcessInputFromArchive changes.
-inline const int32 ShaderCompileWorkerInputVersion = 23;
+inline const int32 ShaderCompileWorkerInputVersion = 24;
 // this is for the protocol, not the data, bump if FShaderCompilerOutput or WriteToOutputArchive changes.
 inline const int32 ShaderCompileWorkerOutputVersion = 19;
 // this is for the protocol, not the data.
diff --git a/Engine/Source/Runtime/RenderCore/Public/ShaderParameterParser.h b/Engine/Source/Runtime/RenderCore/Public/ShaderParameterParser.h
--- a/Engine/Source/Runtime/RenderCore/Public/ShaderParameterParser.h
+++ b/Engine/Source/Runtime/RenderCore/Public/ShaderParameterParser.h
@@ -28,19 +28,6 @@ enum class EBindlessParameterMode : uint8
 	Vulkan,
 };
 
-inline FStringView StripTemplateFromType(const FStringView& Input)
-{
-	FStringView UntemplatedType = FStringView(Input);
-	if (int32 Index = Input.Find(TEXT("<")); Index != INDEX_NONE)
-	{
-		// Remove the template argument but don't forget to clean up the type name
-		const int32 NumChars = Input.Len() - Index;
-		UntemplatedType = Input.LeftChop(NumChars).TrimEnd();
-	}
-
-	return UntemplatedType;
-}
-
 /** Validates and moves all the shader loose data parameter defined in the root scope of the shader into the root uniform buffer. */
 class FShaderParameterParser
 {
@@ -87,8 +74,6 @@ public:
 		friend class FShaderParameterParser;
 	};
 
-	RENDERCORE_API FShaderParameterParser();
-
 	UE_DEPRECATED(5.3, "Use FShaderParameterParser constructor which accepts FShaderCompilerFlags")
 	RENDERCORE_API FShaderParameterParser(const TCHAR* InConstantBufferType);
 	UE_DEPRECATED(5.3, "Use FShaderParameterParser constructor which accepts FShaderCompilerFlags")
@@ -102,8 +87,6 @@ public:
 
 	RENDERCORE_API virtual ~FShaderParameterParser();
 
-	FShaderParameterParser& operator=(FShaderParameterParser&&) = default;
-
 	static constexpr const TCHAR* kBindlessSRVPrefix = TEXT("BindlessSRV_");
 	static constexpr const TCHAR* kBindlessUAVPrefix = TEXT("BindlessUAV_");
 	static constexpr const TCHAR* kBindlessSamplerPrefix = TEXT("BindlessSampler_");
@@ -185,8 +168,6 @@ public:
 
 	bool DidModifyShader() const { return bModifiedShader; }
 
-	friend FArchive& operator<<(FArchive& Ar, FShaderParameterParser& Parser);
-
 protected:
 	/** Parses the preprocessed shader code */
 	RENDERCORE_API bool ParseParameters(
diff --git a/Engine/Source/Runtime/RenderCore/Public/ShaderPreprocessTypes.h b/Engine/Source/Runtime/RenderCore/Public/ShaderPreprocessTypes.h
--- a/Engine/Source/Runtime/RenderCore/Public/ShaderPreprocessTypes.h
+++ b/Engine/Source/Runtime/RenderCore/Public/ShaderPreprocessTypes.h
@@ -6,7 +6,6 @@
 #include "Containers/UnrealString.h"
 #include "ShaderCompilerCore.h"
 #include "ShaderCore.h"
-#include "ShaderParameterParser.h"
 
 struct FShaderCompilerEnvironment;
 struct FShaderCompilerError;
@@ -153,32 +152,6 @@ public:
 		return MakeArrayView(Errors);
 	}
 
-	inline bool ParseAndModify(
-		const FShaderCompilerInput& Input, 
-		const FShaderCompilerEnvironment& Environment, 
-		const TCHAR* ConstantBufferType, 
-		const TArrayView<const TCHAR* const> ExtraSRVTypes = {}, 
-		const TArrayView<const TCHAR* const> ExtraUAVTypes = {}, 
-		EBindlessParameterMode BindlessParameterMode = EBindlessParameterMode::Default)
-	{
-		FShaderParameterParser Tmp(Environment.CompilerFlags, ConstantBufferType, ExtraSRVTypes, ExtraUAVTypes);
-		ParameterParser = MoveTemp(Tmp);
-		return ParameterParser.ParseAndModify(Input, Errors, PreprocessedSource, BindlessParameterMode);
-	}
-
-	inline bool ParseAndModify(
-		const FShaderCompilerInput& Input,
-		const FShaderCompilerEnvironment& Environment,
-		EBindlessParameterMode BindlessParameterMode)
-	{
-		return ParseAndModify(Input, Environment, nullptr, {}, {}, BindlessParameterMode);
-	}
-
-	inline const FShaderParameterParser& GetParameterParser() const
-	{
-		return ParameterParser;
-	}
-
 	double GetElapsedTime() const
 	{
 		return ElapsedTime;
@@ -225,7 +198,9 @@ private:
 	bool bSucceeded = false;
 	bool bIsSecondary = false;
 
-	FShaderParameterParser ParameterParser;
-
 	TArray<FShaderDiagnosticData> ShaderDiagnosticDatas;
 };
+
+#if UE_ENABLE_INCLUDE_ORDER_DEPRECATED_IN_5_4
+#include "ShaderParameterParser.h"
+#endif
