Fix checkerboard materials happening in games that were cooked multi-process.

- New asset associations for already-encountered shadermaps were not being sent to the cook director from the workers,
  since the shadermap wasn't new.
- Now the workers will send all encountered shadermaps after the last send (but not their shader code), so new assets that share the same shadermap are known to the director.

diff --git a/Engine/Source/Runtime/RenderCore/Private/ShaderCodeLibrary.cpp b/Engine/Source/Runtime/RenderCore/Private/ShaderCodeLibrary.cpp
--- a/Engine/Source/Runtime/RenderCore/Private/ShaderCodeLibrary.cpp
+++ b/Engine/Source/Runtime/RenderCore/Private/ShaderCodeLibrary.cpp
@@ -1361,8 +1361,9 @@ struct FEditorShaderCodeArchive
 			// the cook because of serialization path being reused for other purposes than actual saving, so counting them every time artificially inflates number of shaders.
 			CodeStats.NumShaders += Code->ShaderEntries.Num();
 			CodeStats.NumShaderMaps++;
-			MarkShaderMapDirty(ShaderMapIndex);
 		}
+		// always mark the shadermap dirty, because it might have gotten new asset associations
+		MarkShaderMapDirty(ShaderMapIndex);
 		return ShaderMapIndex;
 	}
 
@@ -1572,6 +1573,8 @@ struct FEditorShaderCodeArchive
 			const FShaderMapAssetPaths* const SourceAssetPaths = SourceArchive.ShaderCodeToAssets.Find(SourceShaderMapHash);
 			int32 TargetShaderMapIndex;
 			const bool bIsNewShaderMap = TargetArchive.FindOrAddShaderMap(SourceShaderMapHash, TargetShaderMapIndex, SourceAssetPaths);
+			// always mark the shadermap dirty, because it might have gotten new asset associations
+			MarkShaderMapDirty(TargetShaderMapIndex);
 			if (!bIsNewShaderMap)
 			{
 				// ShaderMap has already been loaded in this process, e.g. from another CookWorker loading the same Material
@@ -1599,7 +1602,6 @@ struct FEditorShaderCodeArchive
 
 			CodeStats.NumShaders += TargetEntry.NumShaders; // Sum of shader counts used by each ShaderMap, without removing duplicates
 			CodeStats.NumShaderMaps++;
-			MarkShaderMapDirty(TargetShaderMapIndex);
 		}
 		return bOk;
 	}
@@ -1660,8 +1662,9 @@ struct FEditorShaderCodeArchive
 				}
 				SerializedShaders.ShaderIndices[ShaderMapEntry.ShaderIndicesOffset + i] = ShaderIndex;
 			}
-			MarkShaderMapDirty(ShaderMapIndex);
 		}
+		// always mark the shadermap dirty, because it might have gotten new asset associations
+		MarkShaderMapDirty(ShaderMapIndex);
 		return ShaderMapIndex;
 	}
 
@@ -1703,8 +1706,9 @@ struct FEditorShaderCodeArchive
 				}
 				SerializedShaders.ShaderIndices[ShaderMapEntry.ShaderIndicesOffset + i] = ShaderIndex;
 			}
-			MarkShaderMapDirty(ShaderMapIndex);
 		}
+		// always mark the shadermap dirty, because it might have gotten new asset associations
+		MarkShaderMapDirty(ShaderMapIndex);
 		return ShaderMapIndex;
 	}
 
