Disable warnings caused by fastgen PGO

diff --git a/Engine/Source/Runtime/RenderCore/Public/RenderGraphResources.inl b/Engine/Source/Runtime/RenderCore/Public/RenderGraphResources.inl
--- a/Engine/Source/Runtime/RenderCore/Public/RenderGraphResources.inl
+++ b/Engine/Source/Runtime/RenderCore/Public/RenderGraphResources.inl
@@ -2,6 +2,8 @@
 
 #pragma once
 
+PRAGMA_DISABLE_BUFFER_OVERRUN_WARNING
+
 inline void FRDGSubresourceState::Finalize()
 {
 	ensureMsgf(!EnumHasAnyFlags(GetPipelines(), ERHIPipeline::AsyncCompute), TEXT("Resource should not be on the async compute pipeline!"));
@@ -262,4 +264,6 @@ inline FGraphicsPipelineRenderTargetsInfo ExtractRenderTargetsInfo(const FRender
 	RenderTargetsInfo.bHasFragmentDensityAttachment = RenderTargets.ShadingRateTexture != nullptr;
 
 	return RenderTargetsInfo;
-}
\ No newline at end of file
+}
+
+PRAGMA_ENABLE_BUFFER_OVERRUN_WARNING
diff --git a/Engine/Source/Runtime/RenderCore/Public/ShaderParameterMacros.h b/Engine/Source/Runtime/RenderCore/Public/ShaderParameterMacros.h
--- a/Engine/Source/Runtime/RenderCore/Public/ShaderParameterMacros.h
+++ b/Engine/Source/Runtime/RenderCore/Public/ShaderParameterMacros.h
@@ -16,6 +16,8 @@
 #include "Misc/AssertionMacros.h"
 #include "RHICommandList.h"
 
+PRAGMA_DISABLE_BUFFER_OVERRUN_WARNING
+
 class FRDGTexture;
 class FRDGTextureSRV;
 class FRDGTextureUAV;
@@ -1920,3 +1922,5 @@ constexpr uint32 CalcPackedComponentIndex(uint32 ElementIndex)	{ return (Element
 
 #define GET_SCALAR_ARRAY_ELEMENT(PackedArray, ElementIndex) \
 	PackedArray[CalcPackedArrayIndex(ElementIndex)][CalcPackedComponentIndex(ElementIndex)]
+
+PRAGMA_ENABLE_BUFFER_OVERRUN_WARNING
\ No newline at end of file
