RootParametersStructure needs to be part of the preprocessed shader cache key.
- Shaders can share generated code but generate different Root Parameters. See: UnifiedBuffer.cpp
- Including the hash of the root parameters structure will make sure these shaders compile seperately.

diff --git a/Engine/Source/Runtime/RenderCore/Private/ShaderCore.cpp b/Engine/Source/Runtime/RenderCore/Private/ShaderCore.cpp
--- a/Engine/Source/Runtime/RenderCore/Private/ShaderCore.cpp
+++ b/Engine/Source/Runtime/RenderCore/Private/ShaderCore.cpp
@@ -3657,6 +3657,13 @@ FShaderCommonCompileJob::FInputHash FShaderCompileJob::GetInputHash()
 		{
 			Hasher << const_cast<FString&>(SecondaryPreprocessOutput->GetSource());
 		}
+
+		if (Input.RootParametersStructure)
+		{
+			uint32 LayoutHash = Input.RootParametersStructure->GetLayoutHash();
+			Hasher << LayoutHash;
+		}
+
 		InputHash = Hasher.Finalize();
 	}
 	else
