Fix not being able to run with loose files.

diff --git a/Engine/Source/Runtime/RenderCore/Private/ShaderCodeLibrary.cpp b/Engine/Source/Runtime/RenderCore/Private/ShaderCodeLibrary.cpp
--- a/Engine/Source/Runtime/RenderCore/Private/ShaderCodeLibrary.cpp
+++ b/Engine/Source/Runtime/RenderCore/Private/ShaderCodeLibrary.cpp
@@ -37,6 +37,8 @@ ShaderCodeLibrary.cpp: Bound shader state cache implementation.
 #include "ShaderPipelineCache.h"
 #include "DataDrivenShaderPlatformInfo.h"
 #include "String/ParseTokens.h"
+#include "IO/IoChunkId.h"
+#include "IO/IoDispatcher.h"
 
 #if WITH_EDITORONLY_DATA
 #include "Interfaces/IShaderFormat.h"
@@ -851,15 +853,16 @@ public:
 			ShaderCodeDirectory = ShaderCodeDir;
 		}
 
-		if (!Library && FIoDispatcher::IsInitialized())
+		const bool bRunningWithIoStore = FIoDispatcher::IsInitialized() && FIoDispatcher::Get().DoesChunkExist(CreateIoChunkId(0, 0, EIoChunkType::ScriptObjects));
+
+		if (!Library && bRunningWithIoStore)
 		{
 			Library = FIoStoreShaderCodeArchive::Create(InShaderPlatform, InLibraryName, FIoDispatcher::Get());
 			ShaderCodeDirectory.Empty();	// paths don't matter for IoStore-based libraries
 		}
 
-		// Shader library as a ushaderbytecode file is no longer an option for distribution. Some code paths (loose files) still need it,
-		// but treat IoStore being initialized as a sign that the library will only be provided as a IoStore shader code archive.
-		if (!Library && !FIoDispatcher::IsInitialized())
+		// Shader library as a ushaderbytecode file is no longer an option for distribution. Some cases (a build using loose files) still require its support though.
+		if (!Library && !bRunningWithIoStore)
 		{
 			const FName PlatformName = FDataDrivenShaderPlatformInfo::GetName(InShaderPlatform);
 			const FName ShaderFormatName = LegacyShaderPlatformToShaderFormat(InShaderPlatform);
