Reduce shader library startup times.

- Plugins that we have cooked with the build will not get their own libraries due to the way the cooker works (all content is cooked at once).
- Do not open any plugins as chunked. Rationale: DLC plugins have chunking turned off (extra code added to validate this), whereas baked-in plugins will not have a separate shader library and will not go through that route.
- Also, do not attempt to load iterative shader library or discover libraries on disk in Test - this is a development feature, whereas Test is used for perf metrics.

diff --git a/Engine/Plugins/Experimental/GameFeatures/Source/GameFeatures/Private/GameFeaturePluginStateMachine.cpp b/Engine/Plugins/Experimental/GameFeatures/Source/GameFeatures/Private/GameFeaturePluginStateMachine.cpp
--- a/Engine/Plugins/Experimental/GameFeatures/Source/GameFeatures/Private/GameFeaturePluginStateMachine.cpp
+++ b/Engine/Plugins/Experimental/GameFeatures/Source/GameFeatures/Private/GameFeaturePluginStateMachine.cpp
@@ -1878,7 +1878,7 @@ struct FGameFeaturePluginState_Mounting : public FGameFeaturePluginState
 	ESubState StartedSubStates = ESubState::None;
 	ESubState CompletedSubStates = ESubState::None;
 	bool bCheckedRealtimeMode = false;
-	bool bForceMonolithicShaderLibrary = false;
+	bool bForceMonolithicShaderLibrary = true;	// use monolithic unless a DLC plugin is chunked
 
 	void OnInstallBundleCompleted(FInstallBundleRequestResultInfo BundleResult)
 	{
@@ -2038,11 +2038,12 @@ struct FGameFeaturePluginState_Mounting : public FGameFeaturePluginState
 		}
 
 		// Pre-mount
-		bool bOpenPluginShaderLibrary = true;
+		// Normally the shader library itself listens to a "New Plugin mounted" (and "New Pakfile mounted") callback and the library is opened automatically. This switch governs whether the manual behavior is wanted.
+		bool bManuallyOpenPluginShaderLibrary = true;
 		{
 			FGameFeaturePreMountingContext Context;
 			UGameFeaturesSubsystem::Get().OnGameFeaturePreMounting(StateProperties.PluginName, StateProperties.PluginIdentifier, Context);
-			bOpenPluginShaderLibrary = Context.bOpenPluginShaderLibrary;
+			bManuallyOpenPluginShaderLibrary = Context.bOpenPluginShaderLibrary;
 		}
 
 		checkf(!StateProperties.PluginInstalledFilename.IsEmpty(), TEXT("PluginInstalledFilename must be set by the Mounting. PluginURL: %s"), *StateProperties.PluginIdentifier.GetFullPluginURL());
@@ -2078,7 +2079,7 @@ struct FGameFeaturePluginState_Mounting : public FGameFeaturePluginState
 			return;
 		}
 
-		if (bOpenPluginShaderLibrary)
+		if (bManuallyOpenPluginShaderLibrary)
 		{
 			// We want to control opening the shader lib
 			FShaderCodeLibrary::DontOpenPluginShaderLibraryOnMount(StateProperties.PluginName);
@@ -2087,7 +2088,7 @@ struct FGameFeaturePluginState_Mounting : public FGameFeaturePluginState
 		if (!UseAsyncLoading())
 		{
 			verify(IPluginManager::Get().MountExplicitlyLoadedPlugin(StateProperties.PluginName));
-			if (bOpenPluginShaderLibrary)
+			if (bManuallyOpenPluginShaderLibrary)
 			{
 				TSharedPtr<IPlugin> Plugin = IPluginManager::Get().FindPlugin(StateProperties.PluginName);
 				FShaderCodeLibrary::OpenPluginShaderLibrary(*Plugin, bForceMonolithicShaderLibrary);
@@ -2100,11 +2101,11 @@ struct FGameFeaturePluginState_Mounting : public FGameFeaturePluginState
 
 		// Now load the shader lib in the background
 		TSharedPtr<IPlugin> Plugin = IPluginManager::Get().FindPlugin(StateProperties.PluginName);
-		if (bOpenPluginShaderLibrary && Plugin->CanContainContent() && Plugin->IsEnabled())
+		if (bManuallyOpenPluginShaderLibrary && Plugin->CanContainContent() && Plugin->IsEnabled())
 		{
-			UE::Tasks::Launch(UE_SOURCE_LOCATION, [this, PluginName=Plugin->GetName(), PluginDir=Plugin->GetContentDir()]
+			UE::Tasks::Launch(UE_SOURCE_LOCATION, [this, Plugin]
 			{
-				FShaderCodeLibrary::OpenLibrary(PluginName, PluginDir, bForceMonolithicShaderLibrary);
+				FShaderCodeLibrary::OpenPluginShaderLibrary(*Plugin, bForceMonolithicShaderLibrary);
 
 				ExecuteOnGameThread(UE_SOURCE_LOCATION, [this]
 				{
diff --git a/Engine/Source/Editor/UnrealEd/Private/ShaderLibraryChunkDataGenerator.cpp b/Engine/Source/Editor/UnrealEd/Private/ShaderLibraryChunkDataGenerator.cpp
--- a/Engine/Source/Editor/UnrealEd/Private/ShaderLibraryChunkDataGenerator.cpp
+++ b/Engine/Source/Editor/UnrealEd/Private/ShaderLibraryChunkDataGenerator.cpp
@@ -2,6 +2,7 @@
 
 #include "ShaderLibraryChunkDataGenerator.h"
 
+#include "CookOnTheSide/CookLog.h"
 #include "CookOnTheSide/CookOnTheFlyServer.h"
 #include "Interfaces/ITargetPlatform.h"
 #include "IPlatformFileSandboxWrapper.h"
@@ -25,6 +26,18 @@ FShaderLibraryChunkDataGenerator::FShaderLibraryChunkDataGenerator(UCookOnTheFly
 	PlatformIniFile.GetBool(TEXT("DevOptions.Shaders"), TEXT("bDoNotChunkShaderLib"), bOptedOut);
 	PlatformIniFile.GetBool(TEXT("DevOptions.Shaders"), TEXT("bChunkShaderWhenCookingDLC"), bChunkShadersWhenCookingDLC);
 
+	// Make sure that chunking DLC plugins is only allowed when -OverrideShaderLibraryName is set. As of Oct 2023 chunking DLC shader libraries is used as a part of bespoke build process
+	// that splits a monolithic cook into a number of "fake" DLC cooks, but overrides the library name to match the game in the end. This flow is not intended to be supported for true DLC plugins
+	// and as such shader code library is not atm prepared to look for chunked parts when trying to open a DLC plugin library
+	if (bChunkShadersWhenCookingDLC)
+	{
+		FString OverrideShaderLibraryName;
+		if (!FParse::Value(FCommandLine::Get(), TEXT("OverrideShaderLibraryName="), OverrideShaderLibraryName))
+		{
+			UE_LOG(LogCook, Error, TEXT("bChunkShaderWhenCookingDLC is true, but -OverrideShaderLibraryName is not set. Chunk a real DLC plugin content isn't currently supported. bChunkShaderWhenCookingDLC is only usable in conjuction with -OverrideShaderLibraryName for a bespoke build process."));
+		}
+	}
+
 	// Disable chunking for DLC - this causes problems as the main game can be optionally (for faster iteration) cooked with -fastcook. Fastcook disables chunking,
 	// so the game has no idea about ChunkIDs and cannot find DLC's chunked libs. If DLC lib is monolithic, both monolithic and chunked games will try to open it.
 	if (!bChunkShadersWhenCookingDLC && COTFS.IsCookingDLC())
diff --git a/Engine/Source/Runtime/RenderCore/Private/ShaderCodeLibrary.cpp b/Engine/Source/Runtime/RenderCore/Private/ShaderCodeLibrary.cpp
--- a/Engine/Source/Runtime/RenderCore/Private/ShaderCodeLibrary.cpp
+++ b/Engine/Source/Runtime/RenderCore/Private/ShaderCodeLibrary.cpp
@@ -58,7 +58,7 @@ ShaderCodeLibrary.cpp: Bound shader state cache implementation.
 #define UE_SHADERLIB_WITH_INTROSPECTION			!UE_BUILD_SHIPPING
 
 // In some development-only scenario (e.g. LaunchOn), the library is chunked, but the build isn't pak'd. We need to find the chunk files manually then.
-#define UE_SHADERLIB_SUPPORT_CHUNK_DISCOVERY	!UE_BUILD_SHIPPING
+#define UE_SHADERLIB_SUPPORT_CHUNK_DISCOVERY	(!UE_BUILD_SHIPPING && !UE_BUILD_TEST)
 
 DEFINE_LOG_CATEGORY(LogShaderLibrary);
 
@@ -2605,11 +2605,14 @@ public:
 #if UE_SHADERLIB_SUPPORT_CHUNK_DISCOVERY
 				if (!bResult)
 				{
+					static const bool bRunningWithPakFile = (FPlatformFileManager::Get().FindPlatformFile(TEXT("PakFile")) != nullptr);
+					static const bool bRunningWithIoStore = FIoDispatcher::IsInitialized() && FIoDispatcher::Get().DoesChunkExist(CreateIoChunkId(0, 0, EIoChunkType::ScriptObjects));
+
 					// Some deployment flows (e.g. Launch on) avoid pak files despite project packaging settings. 
 					// In case we run under such circumstances, we need to discover the components ourselves
-					if (FPlatformFileManager::Get().FindPlatformFile(TEXT("PakFile")) == nullptr)
+					if (!bRunningWithPakFile && !bRunningWithIoStore)
 					{
-						UE_LOG(LogShaderLibrary, Display, TEXT("Running without a pakfile and did not find a monolithic library '%s' - attempting disk search for its chunks"), *Name);
+						UE_LOG(LogShaderLibrary, Display, TEXT("Running without a pakfile/IoStore and did not find a monolithic library '%s' - attempting disk search for its chunks"), *Name);
 
 						TArray<FString> UshaderbytecodeFiles;
 						FString SearchMask = Directory / FString::Printf(TEXT("ShaderArchive-*%s*.ushaderbytecode"), *Name);
@@ -2693,7 +2696,8 @@ public:
 			{
 				if (bAddNewNamedLibrary)
 				{
-					UE_LOG(LogShaderLibrary, Display, TEXT("Tried to open shader library '%s', but could not find it neither as a monolithic library nor as a chunked one."), *Name);
+					UE_LOG(LogShaderLibrary, Display, TEXT("Tried to open shader library '%s', but could not find it%s"), *Name, 
+						bMonolithicOnly ? TEXT(" (only tried to open it as a monolithic library).") : TEXT(" neither as a monolithic library nor as a chunked one."));
 
 					check(Library->GetNumComponents() == 0);
 					delete Library;
@@ -3475,7 +3479,7 @@ void FShaderCodeLibrary::InitForRuntime(EShaderPlatform ShaderPlatform)
 			UE::ShaderLibrary::Private::OnPluginMountedDelegateHandle = IPluginManager::Get().OnNewPluginMounted().AddStatic(&FShaderCodeLibraryPluginMountedCallback);
 			UE::ShaderLibrary::Private::OnPluginUnmountedDelegateHandle = IPluginManager::Get().OnPluginUnmounted().AddStatic(&FShaderCodeLibraryPluginUnmountedCallback);
 		
-#if !UE_BUILD_SHIPPING
+#if (!UE_BUILD_SHIPPING && !UE_BUILD_TEST)	// test builds are supposed to be closer to Shipping than to Development, and as such not have development features
 			// support shared cooked builds by also opening the shared cooked build shader code file
 			FShaderLibrariesCollection::Impl->OpenLibrary(TEXT("Global_SC"), FPaths::ProjectContentDir());
 #endif
@@ -4085,6 +4089,14 @@ void FShaderCodeLibrary::OpenPluginShaderLibrary(IPlugin& Plugin, bool bMonolith
 	if (Plugin.CanContainContent() && Plugin.IsEnabled())
 	{
 		// load any shader libraries that may exist in this plugin
+		if (!bMonolithicOnly)
+		{
+			// Chunked libraries in plugins that are not built in (i.e. ones that were cooked separately as DLC) are not supported atm. This is because the main game can be cooked without chunks (-fastcook), but still needs
+			// to load the same plugins, so it would not know which ChunkIDs to try.
+			// Plugins that are built-in can be chunked, but their shaders don't go into a separate library (cooker doesn't separate that atm), everything goes into main project's library.
+			UE_LOG(LogShaderLibrary, Warning, TEXT("Opening a chunked shader library for plugin '%s' is ignored. Chunked libraries for plugins are not supported."), *Plugin.GetName());
+		}
+		bMonolithicOnly = true;
 		FShaderCodeLibrary::OpenLibrary(Plugin.GetName(), Plugin.GetContentDir(), bMonolithicOnly);
 	}
 }
diff --git a/Engine/Source/Runtime/RenderCore/Public/ShaderCodeLibrary.h b/Engine/Source/Runtime/RenderCore/Public/ShaderCodeLibrary.h
--- a/Engine/Source/Runtime/RenderCore/Public/ShaderCodeLibrary.h
+++ b/Engine/Source/Runtime/RenderCore/Public/ShaderCodeLibrary.h
@@ -256,8 +256,8 @@ struct FShaderCodeLibrary
 	static RENDERCORE_API void DontOpenPluginShaderLibraryOnMount(const FString& PluginName);
 	
 	// Open the plugin's shader library
-	// @param bMonolithicOnly If true, only attempt to open a monolithic library (no chunks)
-	static RENDERCORE_API void OpenPluginShaderLibrary(IPlugin& Plugin, bool bMonolithicOnly = false);
+	// @param bMonolithicOnly If true, only attempt to open a monolithic library (no chunks) - which is a default behavior for DLC plugins, see FShaderLibraryChunkDataGenerator.
+	static RENDERCORE_API void OpenPluginShaderLibrary(IPlugin& Plugin, bool bMonolithicOnly = true);
 };
 
 #if WITH_EDITOR
