SCL: downgrade a common warning to display.

#preflgiht none

diff --git a/Engine/Source/Runtime/RenderCore/Private/ShaderCodeLibrary.cpp b/Engine/Source/Runtime/RenderCore/Private/ShaderCodeLibrary.cpp
--- a/Engine/Source/Runtime/RenderCore/Private/ShaderCodeLibrary.cpp
+++ b/Engine/Source/Runtime/RenderCore/Private/ShaderCodeLibrary.cpp
@@ -4094,7 +4094,7 @@ void FShaderCodeLibrary::OpenPluginShaderLibrary(IPlugin& Plugin, bool bMonolith
 			// Chunked libraries in plugins that are not built in (i.e. ones that were cooked separately as DLC) are not supported atm. This is because the main game can be cooked without chunks (-fastcook), but still needs
 			// to load the same plugins, so it would not know which ChunkIDs to try.
 			// Plugins that are built-in can be chunked, but their shaders don't go into a separate library (cooker doesn't separate that atm), everything goes into main project's library.
-			UE_LOG(LogShaderLibrary, Warning, TEXT("Opening a chunked shader library for plugin '%s' is ignored. Chunked libraries for plugins are not supported."), *Plugin.GetName());
+			UE_LOG(LogShaderLibrary, Display, TEXT("Opening a chunked shader library for plugin '%s' is ignored. Chunked libraries for plugins are not supported."), *Plugin.GetName());
 		}
 		bMonolithicOnly = true;
 		FShaderCodeLibrary::OpenLibrary(Plugin.GetName(), Plugin.GetContentDir(), bMonolithicOnly);
