Shader bundle dispatch optimizations on XSX - do not allocate and NOP fill buffer records for invalid dispatches, and also added wave64-wide redundant descriptor table filtering for the CBV/SRV/UAV/SMP tables

diff --git a/Engine/Source/Runtime/RenderCore/Private/ShaderBundles.cpp b/Engine/Source/Runtime/RenderCore/Private/ShaderBundles.cpp
--- a/Engine/Source/Runtime/RenderCore/Private/ShaderBundles.cpp
+++ b/Engine/Source/Runtime/RenderCore/Private/ShaderBundles.cpp
@@ -23,6 +23,7 @@ void FDispatchShaderBundleCS::ModifyCompilationEnvironment(const FGlobalShaderPe
 
 void FDispatchShaderBundle::Dispatch(
 	FRHIShaderBundle* ShaderBundle,
+	uint32 RecordCount,
 	FRHIComputeCommandList& RHICmdList,
 	FRHIShaderResourceView* RecordArgBufferSRV,
 	FRHIShaderResourceView* RecordDataBufferSRV,
@@ -30,7 +31,7 @@ void FDispatchShaderBundle::Dispatch(
 )
 {
 	check(RHISupportsShaderBundleDispatch(GMaxRHIShaderPlatform) && GRHISupportsShaderBundleDispatch);
-	check(ShaderBundle && ShaderBundle->NumRecords > 0);
+	check(ShaderBundle && ShaderBundle->NumRecords > 0 && RecordCount <= ShaderBundle->NumRecords);
 
 	RHICmdList.ClearUAVUint(ExecutionBufferUAV, FUintVector4(0, 0, 0, 0));
 
@@ -38,13 +39,13 @@ void FDispatchShaderBundle::Dispatch(
 
 	FDispatchShaderBundleCS::FParameters Parameters;
 
-	Parameters.RecordCount = ShaderBundle->NumRecords;
+	Parameters.RecordCount = RecordCount;
 	Parameters.PlatformData = ShaderBundle->GetPlatformData();
 	Parameters.RecordArgBuffer = RecordArgBufferSRV;
 	Parameters.RecordDataBuffer = RecordDataBufferSRV;
 	Parameters.RWExecutionBuffer = ExecutionBufferUAV;
 
-	const FIntVector GroupCount = FComputeShaderUtils::GetGroupCount(ShaderBundle->NumRecords, FDispatchShaderBundleCS::ThreadGroupSizeX);
+	const FIntVector GroupCount = FComputeShaderUtils::GetGroupCount(RecordCount, FDispatchShaderBundleCS::ThreadGroupSizeX);
 
 	FRHIComputeShader* ComputeShaderRHI = ComputeShader.GetComputeShader();
 	SetComputePipelineState(RHICmdList, ComputeShaderRHI);
diff --git a/Engine/Source/Runtime/RenderCore/Public/ShaderBundles.h b/Engine/Source/Runtime/RenderCore/Public/ShaderBundles.h
--- a/Engine/Source/Runtime/RenderCore/Public/ShaderBundles.h
+++ b/Engine/Source/Runtime/RenderCore/Public/ShaderBundles.h
@@ -45,6 +45,7 @@ class FDispatchShaderBundle
 public:
 	static RENDERCORE_API void Dispatch(
 		FRHIShaderBundle* ShaderBundle,
+		uint32 RecordCount,
 		FRHIComputeCommandList& RHICmdList,
 		FRHIShaderResourceView* RecordArgBufferSRV,
 		FRHIShaderResourceView* RecordDataBufferSRV,
