Add check for UBMT_RDG_BUFFER_SRV when validating ShaderParameters

diff --git a/Engine/Source/Runtime/RenderCore/Private/ShaderParameterStruct.cpp b/Engine/Source/Runtime/RenderCore/Private/ShaderParameterStruct.cpp
--- a/Engine/Source/Runtime/RenderCore/Private/ShaderParameterStruct.cpp
+++ b/Engine/Source/Runtime/RenderCore/Private/ShaderParameterStruct.cpp
@@ -620,7 +620,8 @@ void ValidateShaderParameterResourcesRHI(const void* Contents, const FRHIUniform
 
 		const bool bSRV =
 			Parameter.MemberType == UBMT_SRV ||
-			Parameter.MemberType == UBMT_RDG_TEXTURE_SRV;
+			Parameter.MemberType == UBMT_RDG_TEXTURE_SRV ||
+			Parameter.MemberType == UBMT_RDG_BUFFER_SRV;
 
 		// Allow null SRV's in uniform buffers for feature levels that don't support SRV's in shaders
 		if (GMaxRHIFeatureLevel <= ERHIFeatureLevel::ES3_1 && bSRV)
