Fix warning V781: The value of the 'LineNumberEnd' index is checked after it was used. Perhaps there is a mistake in program logic.
#rnx

diff --git a/Engine/Source/Runtime/RenderCore/Private/ShaderPreprocessTypes.cpp b/Engine/Source/Runtime/RenderCore/Private/ShaderPreprocessTypes.cpp
--- a/Engine/Source/Runtime/RenderCore/Private/ShaderPreprocessTypes.cpp
+++ b/Engine/Source/Runtime/RenderCore/Private/ShaderPreprocessTypes.cpp
@@ -52,7 +52,7 @@ void FShaderDiagnosticRemapper::Remap(FShaderCompilerError& Diagnostic) const
 		}
 		int32 LineNumberEnd = LineNumberStart;
 		int32 StrippedLineNum = 0;
-		while (FChar::IsDigit(OriginalMessage[LineNumberEnd]) && LineNumberEnd < OriginalMessage.Len())
+		while (LineNumberEnd < OriginalMessage.Len() && FChar::IsDigit(OriginalMessage[LineNumberEnd]))
 		{
 			StrippedLineNum = StrippedLineNum * 10 + (OriginalMessage[LineNumberEnd++] - TEXT('0'));
 		}
