Fix warning V779: Unreachable code detected. It is possible that an error is present.
#rnx

diff --git a/Engine/Source/Runtime/RenderCore/Private/UnifiedBuffer.cpp b/Engine/Source/Runtime/RenderCore/Private/UnifiedBuffer.cpp
--- a/Engine/Source/Runtime/RenderCore/Private/UnifiedBuffer.cpp
+++ b/Engine/Source/Runtime/RenderCore/Private/UnifiedBuffer.cpp
@@ -65,16 +65,16 @@ class FByteBufferShader : public FGlobalShader
 		{
 			return FDataDrivenShaderPlatformInfo::GetSupportsByteBufferComputeShaders(Parameters.Platform);
 		}
+		// Don't compile structured buffer size variations unless we need them
+		else if (ResourceType != EByteBufferResourceType::StructuredBuffer && static_cast<EByteBufferStructuredSize>(PermutationVector.Get<StructuredElementSizeDim>()) != EByteBufferStructuredSize::Uint4)
+		{
+			return false;
+		}
 		else
 		{
 			return true;
 		}
 
-		// Don't compile structured buffer size variations unless we need them
-		if (ResourceType != EByteBufferResourceType::StructuredBuffer && static_cast<EByteBufferStructuredSize>(PermutationVector.Get<StructuredElementSizeDim>()) != EByteBufferStructuredSize::Uint4)
-		{
-			return false;
-		}
 	}
 
 	BEGIN_SHADER_PARAMETER_STRUCT( FParameters, )
