Bindless Conversion Changes
- Removing the need for COMPILER_DEFINE on GetResourceFromHeap and GetSamplerFromHeap
- FShaderParameterParser constructor now requires a FPlatformConfiguration object that each shader format can configure.
- FPlatformConfiguration::GenerateBindlessAccess allows each platform to generate the code for accessing a bindless resource or sampler.
- Added EShaderParameterParserConfigurationFlags to specify support for stable constant buffers and bindless.
- Moved EBindlessParameterMode into EShaderParameterParserConfigurationFlags::BindlessUsesArrays
- While moving things around, I was able to convert some TCHAR* uses to FStringView

diff --git a/Engine/Shaders/Public/Platform/D3D/D3DCommon.ush b/Engine/Shaders/Public/Platform/D3D/D3DCommon.ush
--- a/Engine/Shaders/Public/Platform/D3D/D3DCommon.ush
+++ b/Engine/Shaders/Public/Platform/D3D/D3DCommon.ush
@@ -136,6 +136,6 @@ void UEReportAssertWithPayload(uint AssertID, uint4 Payload)
 
 #if PLATFORM_SUPPORTS_BINDLESS
 	// SM6.6 path
-	COMPILER_DEFINE GetResourceFromHeap(Type, Index) ResourceDescriptorHeap[Index]
-	COMPILER_DEFINE GetSamplerFromHeap(Type, Index) SamplerDescriptorHeap[Index]
+	#define GetResourceFromHeap(Type, Index) ResourceDescriptorHeap[Index]
+	#define GetSamplerFromHeap(Type, Index) SamplerDescriptorHeap[Index]
 #endif
diff --git a/Engine/Source/Developer/Apple/MetalShaderFormat/Private/MetalShaderCompiler.cpp b/Engine/Source/Developer/Apple/MetalShaderFormat/Private/MetalShaderCompiler.cpp
--- a/Engine/Source/Developer/Apple/MetalShaderFormat/Private/MetalShaderCompiler.cpp
+++ b/Engine/Source/Developer/Apple/MetalShaderFormat/Private/MetalShaderCompiler.cpp
@@ -882,8 +882,9 @@ void CompileMetalShader(const FShaderCompilerInput& Input, const FString& InPrep
 	FString EntryPointName = Input.EntryPointName;
 	FString PreprocessedSource = InPreprocessedSource;
 
-	FShaderParameterParser ShaderParameterParser(Input.Environment.CompilerFlags, nullptr, {}, {});
-	if (!ShaderParameterParser.ParseAndModify(Input, Output.Errors, PreprocessedSource, EBindlessParameterMode::Default))
+	FShaderParameterParser::FPlatformConfiguration PlatformConfiguration;
+	FShaderParameterParser ShaderParameterParser(PlatformConfiguration);
+	if (!ShaderParameterParser.ParseAndModify(Input, Output.Errors, PreprocessedSource))
 	{
 		// The FShaderParameterParser will add any relevant errors.
 		return;
diff --git a/Engine/Source/Developer/ShaderCompilerCommon/Private/ShaderCompilerCommon.cpp b/Engine/Source/Developer/ShaderCompilerCommon/Private/ShaderCompilerCommon.cpp
--- a/Engine/Source/Developer/ShaderCompilerCommon/Private/ShaderCompilerCommon.cpp
+++ b/Engine/Source/Developer/ShaderCompilerCommon/Private/ShaderCompilerCommon.cpp
@@ -373,14 +373,6 @@ const TCHAR* ParseHLSLSymbolName(const TCHAR* SearchString, FString& SymbolName)
 	return Result.GetData() + Result.Len();
 }
 
-EShaderParameterType UE::ShaderCompilerCommon::ParseParameterType(
-	FStringView InType,
-	TArrayView<const TCHAR* const> InExtraSRVTypes,
-	TArrayView<const TCHAR* const> InExtraUAVTypes)
-{
-	return FShaderParameterParser::ParseParameterType(InType, InExtraSRVTypes, InExtraUAVTypes);
-}
-
 FStringView UE::ShaderCompilerCommon::RemoveConstantBufferPrefix(FStringView InName)
 {
 	return UE::String::RemoveFromStart(InName, FStringView(UE::ShaderCompilerCommon::kUniformBufferConstantBufferPrefix));
@@ -391,21 +383,6 @@ FString UE::ShaderCompilerCommon::RemoveConstantBufferPrefix(const FString& InNa
 	return FString(RemoveConstantBufferPrefix(FStringView(InName)));
 }
 
-EShaderParameterType UE::ShaderCompilerCommon::ParseAndRemoveBindlessParameterPrefix(FStringView& InName)
-{
-	return FShaderParameterParser::ParseAndRemoveBindlessParameterPrefix(InName);
-}
-
-EShaderParameterType UE::ShaderCompilerCommon::ParseAndRemoveBindlessParameterPrefix(FString& InName)
-{
-	return FShaderParameterParser::ParseAndRemoveBindlessParameterPrefix(InName);
-}
-
-bool UE::ShaderCompilerCommon::RemoveBindlessParameterPrefix(FString& InName)
-{
-	return FShaderParameterParser::RemoveBindlessParameterPrefix(InName);
-}
-
 bool UE::ShaderCompilerCommon::ValidatePackedResourceCounts(FShaderCompilerOutput& Output, const FShaderCodePackedResourceCounts& PackedResourceCounts)
 {
 	if (Output.bSucceeded)
diff --git a/Engine/Source/Developer/ShaderCompilerCommon/Public/ShaderCompilerCommon.h b/Engine/Source/Developer/ShaderCompilerCommon/Public/ShaderCompilerCommon.h
--- a/Engine/Source/Developer/ShaderCompilerCommon/Public/ShaderCompilerCommon.h
+++ b/Engine/Source/Developer/ShaderCompilerCommon/Public/ShaderCompilerCommon.h
@@ -64,27 +64,9 @@ extern SHADERCOMPILERCOMMON_API int16 GetNumUniformBuffersUsed(const FShaderComp
 
 namespace UE::ShaderCompilerCommon
 {
-	UE_DEPRECATED(5.3, "Deprecated; use FShaderCompilerInput::ShouldUseStableContantBuffer directly")
-	inline bool ShouldUseStableConstantBuffer(const FShaderCompilerInput& Input)
-	{
-		return Input.ShouldUseStableConstantBuffer();
-	}
-
-	UE_DEPRECATED(5.3, "Deprecated; use FShaderParameterParser::ParseParameterType instead")
-	extern SHADERCOMPILERCOMMON_API EShaderParameterType ParseParameterType(FStringView InType, TArrayView<const TCHAR* const> InExtraSRVTypes, TArrayView<const TCHAR* const> InExtraUAVTypes);
-
 	extern SHADERCOMPILERCOMMON_API FStringView          RemoveConstantBufferPrefix(FStringView InName);
 	extern SHADERCOMPILERCOMMON_API FString              RemoveConstantBufferPrefix(const FString& InName);
 
-	UE_DEPRECATED(5.3, "Deprecated; use FShaderParameterParser::ParseAndRemoveBindlessParameterPrefix instead")
-	extern SHADERCOMPILERCOMMON_API EShaderParameterType ParseAndRemoveBindlessParameterPrefix(FStringView& InName);
-	
-	UE_DEPRECATED(5.3, "Deprecated; use FShaderParameterParser::ParseAndRemoveBindlessParameterPrexix instead")
-	extern SHADERCOMPILERCOMMON_API EShaderParameterType ParseAndRemoveBindlessParameterPrefix(FString& InName);
-
-	UE_DEPRECATED(5.3, "Deprecated; use FShaderParameterParser::RemoveBindlessParameterPrefix instead")
-	extern SHADERCOMPILERCOMMON_API bool                 RemoveBindlessParameterPrefix(FString& InName);
-
 	extern SHADERCOMPILERCOMMON_API bool                 ValidatePackedResourceCounts(FShaderCompilerOutput& Output, const FShaderCodePackedResourceCounts& PackedResourceCounts);
 
 	/*
diff --git a/Engine/Source/Developer/ShaderFormatOpenGL/Private/OpenGLShaderCompiler.cpp b/Engine/Source/Developer/ShaderFormatOpenGL/Private/OpenGLShaderCompiler.cpp
--- a/Engine/Source/Developer/ShaderFormatOpenGL/Private/OpenGLShaderCompiler.cpp
+++ b/Engine/Source/Developer/ShaderFormatOpenGL/Private/OpenGLShaderCompiler.cpp
@@ -3341,8 +3341,9 @@ void CompileOpenGLShader(const FShaderCompilerInput& Input, const FString& InPre
 	FString EntryPointName = Input.EntryPointName;
 	FString PreprocessedSource = InPreprocessedSource;
 
-	FShaderParameterParser ShaderParameterParser(Input.Environment.CompilerFlags, nullptr, {}, {});
-	if (!ShaderParameterParser.ParseAndModify(Input, Output.Errors, PreprocessedSource, EBindlessParameterMode::Default))
+	FShaderParameterParser::FPlatformConfiguration PlatformConfiguration;
+	FShaderParameterParser ShaderParameterParser(PlatformConfiguration);
+	if (!ShaderParameterParser.ParseAndModify(Input, Output.Errors, PreprocessedSource))
 	{
 		// The FShaderParameterParser will add any relevant errors.
 		return;
diff --git a/Engine/Source/Developer/VulkanShaderFormat/Private/VulkanShaderCompiler.cpp b/Engine/Source/Developer/VulkanShaderFormat/Private/VulkanShaderCompiler.cpp
--- a/Engine/Source/Developer/VulkanShaderFormat/Private/VulkanShaderCompiler.cpp
+++ b/Engine/Source/Developer/VulkanShaderFormat/Private/VulkanShaderCompiler.cpp
@@ -2538,6 +2538,28 @@ static bool CompileShaderGroup(
 	return bSuccess;
 }
 
+struct FPS5ShaderParameterParserPlatformConfiguration : public FShaderParameterParser::FPlatformConfiguration
+{
+	FPS5ShaderParameterParserPlatformConfiguration(const FShaderCompilerInput& Input)
+		: FShaderParameterParser::FPlatformConfiguration()
+	{
+		EnumAddFlags(Flags, EShaderParameterParserConfigurationFlags::SupportsBindless | EShaderParameterParserConfigurationFlags::BindlessUsesArrays);
+
+		// Create a _RootShaderParameters and bind it in slot 0 like any other uniform buffer
+		if (Input.Target.GetFrequency() == SF_RayGen && Input.RootParametersStructure != nullptr)
+		{
+			ConstantBufferType = TEXTVIEW("cbuffer");
+			EnumAddFlags(Flags, EShaderParameterParserConfigurationFlags::UseStableConstantBuffer);
+		}
+	}
+
+	virtual FString GenerateBindlessAccess(EBindlessConversionType BindlessType, FStringView ShaderTypeString, FStringView IndexString) const final
+	{
+		checkf(false, TEXT("Vulkan does not use GenerateBindlessAccess"));
+		return FString();
+	}
+};
+
 void CompileVulkanShader(const FShaderCompilerInput& Input, const FString& InPreprocessedSource, FShaderCompilerOutput& Output, const class FString& WorkingDirectory)
 {
 	check(IsVulkanShaderFormat(Input.ShaderFormat));
@@ -2545,11 +2567,9 @@ void CompileVulkanShader(const FShaderCompilerInput& Input, const FString& InPre
 	FString EntryPointName = Input.EntryPointName;
 	FString PreprocessedSource = InPreprocessedSource;
 
-	// Create a _RootShaderParameters and bind it in slot 0 like any other uniform buffer
-	const TCHAR* ConstantBufferType = (Input.Target.GetFrequency() == SF_RayGen && Input.RootParametersStructure != nullptr) ? TEXT("cbuffer") : nullptr;
-
-	FShaderParameterParser ShaderParameterParser(Input.Environment.CompilerFlags, ConstantBufferType, {}, {});
-	if (!ShaderParameterParser.ParseAndModify(Input, Output.Errors, PreprocessedSource, EBindlessParameterMode::Vulkan))
+	FPS5ShaderParameterParserPlatformConfiguration PlatformConfiguration(Input);
+	FShaderParameterParser ShaderParameterParser(PlatformConfiguration);
+	if (!ShaderParameterParser.ParseAndModify(Input, Output.Errors, PreprocessedSource))
 	{
 		// The FShaderParameterParser will add any relevant errors.
 		return;
diff --git a/Engine/Source/Developer/Windows/ShaderFormatD3D/Private/D3DShaderCompiler.cpp b/Engine/Source/Developer/Windows/ShaderFormatD3D/Private/D3DShaderCompiler.cpp
--- a/Engine/Source/Developer/Windows/ShaderFormatD3D/Private/D3DShaderCompiler.cpp
+++ b/Engine/Source/Developer/Windows/ShaderFormatD3D/Private/D3DShaderCompiler.cpp
@@ -1389,6 +1389,26 @@ bool PreprocessD3DShader(
 	return true;
 }
 
+struct FD3DShaderParameterParserPlatformConfiguration : public FShaderParameterParser::FPlatformConfiguration
+{
+	FD3DShaderParameterParserPlatformConfiguration()
+		: FShaderParameterParser::FPlatformConfiguration(TEXTVIEW("cbuffer"), EShaderParameterParserConfigurationFlags::UseStableConstantBuffer|EShaderParameterParserConfigurationFlags::SupportsBindless)
+	{
+	}
+
+	virtual FString GenerateBindlessAccess(EBindlessConversionType BindlessType, FStringView ShaderTypeString, FStringView IndexString) const final
+	{
+		// GetResourceFromHeap(Type, Index) ResourceDescriptorHeap[Index]
+		// GetSamplerFromHeap(Type, Index)  SamplerDescriptorHeap[Index]
+
+		const TCHAR* HeapString = BindlessType == EBindlessConversionType::Sampler ? TEXT("SamplerDescriptorHeap") : TEXT("ResourceDescriptorHeap");
+
+		return FString::Printf(TEXT("%s[%.*s]"),
+			HeapString,
+			IndexString.Len(), IndexString.GetData()
+		);
+	}
+};
 
 void CompileD3DShader(const FShaderCompilerInput& Input, const FString& InPreprocessedSource, FShaderCompilerOutput& Output, const FString& WorkingDirectory, ELanguage Language)
 {
@@ -1406,8 +1426,9 @@ void CompileD3DShader(const FShaderCompilerInput& Input, const FString& InPrepro
 	FString EntryPointName = Input.EntryPointName;
 	FString PreprocessedSource = InPreprocessedSource;
 
-	FShaderParameterParser ShaderParameterParser(Input.Environment.CompilerFlags, TEXT("cbuffer"), {}, {});
-	if (!ShaderParameterParser.ParseAndModify(Input, Output.Errors, PreprocessedSource, EBindlessParameterMode::Default))
+	FD3DShaderParameterParserPlatformConfiguration PlatformConfiguration;
+	FShaderParameterParser ShaderParameterParser(PlatformConfiguration);
+	if (!ShaderParameterParser.ParseAndModify(Input, Output.Errors, PreprocessedSource))
 	{
 		// The FShaderParameterParser will add any relevant errors.
 		return;
diff --git a/Engine/Source/Runtime/RenderCore/Private/ShaderParameterParser.cpp b/Engine/Source/Runtime/RenderCore/Private/ShaderParameterParser.cpp
--- a/Engine/Source/Runtime/RenderCore/Private/ShaderParameterParser.cpp
+++ b/Engine/Source/Runtime/RenderCore/Private/ShaderParameterParser.cpp
@@ -149,76 +149,54 @@ static void AddNoteToDisplayShaderParameterMemberOnCppSide(
 
 FShaderParameterParser::~FShaderParameterParser() = default;
 
-PRAGMA_DISABLE_DEPRECATION_WARNINGS
-FShaderParameterParser::FShaderParameterParser(const TCHAR* InConstantBufferType)
-	: FShaderParameterParser(FShaderCompilerFlags(), InConstantBufferType)
-{}
-FShaderParameterParser::FShaderParameterParser(
-	const TCHAR* InConstantBufferType,
-	TConstArrayView<const TCHAR*> InExtraSRVTypes,
-	TConstArrayView<const TCHAR*> InExtraUAVTypes)
-	: FShaderParameterParser(FShaderCompilerFlags(), InConstantBufferType, InExtraSRVTypes, InExtraUAVTypes)
-{}
-PRAGMA_ENABLE_DEPRECATION_WARNINGS
-
-FShaderParameterParser::FShaderParameterParser(
-	FShaderCompilerFlags CompilerFlags,
-	const TCHAR* InConstantBufferType,
-	TConstArrayView<const TCHAR*> InExtraSRVTypes,
-	TConstArrayView<const TCHAR*> InExtraUAVTypes)
-	: ConstantBufferType(InConstantBufferType)
-	, ExtraSRVTypes(InExtraSRVTypes)
-	, ExtraUAVTypes(InExtraUAVTypes)
-	, bBindlessResources(CompilerFlags.Contains(CFLAG_BindlessResources))
-	, bBindlessSamplers(CompilerFlags.Contains(CFLAG_BindlessSamplers))
-{}
-
-static const TCHAR* const s_AllSRVTypes[] =
+FShaderParameterParser::FShaderParameterParser(const FPlatformConfiguration& InPlatformConfiguration)
+	: PlatformConfiguration(InPlatformConfiguration)
 {
-	TEXT("Texture1D"),
-	TEXT("Texture1DArray"),
-	TEXT("Texture2D"),
-	TEXT("Texture2DArray"),
-	TEXT("Texture2DMS"),
-	TEXT("Texture2DMSArray"),
-	TEXT("Texture3D"),
-	TEXT("TextureCube"),
-	TEXT("TextureCubeArray"),
-
-	TEXT("Buffer"),
-	TEXT("ByteAddressBuffer"),
-	TEXT("StructuredBuffer"),
-	TEXT("RaytracingAccelerationStructure"),
+}
+
+static const FStringView s_AllSRVTypes[] =
+{
+	TEXTVIEW("Texture1D"),
+	TEXTVIEW("Texture1DArray"),
+	TEXTVIEW("Texture2D"),
+	TEXTVIEW("Texture2DArray"),
+	TEXTVIEW("Texture2DMS"),
+	TEXTVIEW("Texture2DMSArray"),
+	TEXTVIEW("Texture3D"),
+	TEXTVIEW("TextureCube"),
+	TEXTVIEW("TextureCubeArray"),
+
+	TEXTVIEW("Buffer"),
+	TEXTVIEW("ByteAddressBuffer"),
+	TEXTVIEW("StructuredBuffer"),
+	TEXTVIEW("RaytracingAccelerationStructure"),
 };
 
-static const TCHAR* const s_AllUAVTypes[] =
+static const FStringView s_AllUAVTypes[] =
 {
-	TEXT("AppendStructuredBuffer"),
-	TEXT("RWBuffer"),
-	TEXT("RWByteAddressBuffer"),
-	TEXT("RWStructuredBuffer"),
-	TEXT("RWTexture1D"),
-	TEXT("RWTexture1DArray"),
-	TEXT("RWTexture2D"),
-	TEXT("RWTexture2DArray"),
-	TEXT("RWTexture3D"),
-	TEXT("RasterizerOrderedTexture2D"),
+	TEXTVIEW("AppendStructuredBuffer"),
+	TEXTVIEW("RWBuffer"),
+	TEXTVIEW("RWByteAddressBuffer"),
+	TEXTVIEW("RWStructuredBuffer"),
+	TEXTVIEW("RWTexture1D"),
+	TEXTVIEW("RWTexture1DArray"),
+	TEXTVIEW("RWTexture2D"),
+	TEXTVIEW("RWTexture2DArray"),
+	TEXTVIEW("RWTexture3D"),
+	TEXTVIEW("RasterizerOrderedTexture2D"),
 };
 
-static const TCHAR* const s_AllSamplerTypes[] =
+static const FStringView s_AllSamplerTypes[] =
 {
-	TEXT("SamplerState"),
-	TEXT("SamplerComparisonState"),
+	TEXTVIEW("SamplerState"),
+	TEXTVIEW("SamplerComparisonState"),
 };
 
-EShaderParameterType FShaderParameterParser::ParseParameterType(
-	FStringView InType,
-	TConstArrayView<const TCHAR*> InExtraSRVTypes,
-	TConstArrayView<const TCHAR*> InExtraUAVTypes)
+EShaderParameterType FShaderParameterParser::ParseParameterType(FStringView InType)
 {
-	TConstArrayView<const TCHAR*> AllSamplerTypes(s_AllSamplerTypes);
-	TConstArrayView<const TCHAR*> AllSRVTypes(s_AllSRVTypes);
-	TConstArrayView<const TCHAR*> AllUAVTypes(s_AllUAVTypes);
+	TConstArrayView<FStringView> AllSamplerTypes(s_AllSamplerTypes);
+	TConstArrayView<FStringView> AllSRVTypes(s_AllSRVTypes);
+	TConstArrayView<FStringView> AllUAVTypes(s_AllUAVTypes);
 
 	if (AllSamplerTypes.Contains(InType))
 	{
@@ -227,12 +205,12 @@ EShaderParameterType FShaderParameterParser::ParseParameterType(
 
 	FStringView UntemplatedType = StripTemplateFromType(InType);
 
-	if (AllSRVTypes.Contains(UntemplatedType) || InExtraSRVTypes.Contains(UntemplatedType))
+	if (AllSRVTypes.Contains(UntemplatedType) || PlatformConfiguration.ExtraSRVTypes.Contains(UntemplatedType))
 	{
 		return EShaderParameterType::SRV;
 	}
 
-	if (AllUAVTypes.Contains(UntemplatedType) || InExtraUAVTypes.Contains(UntemplatedType))
+	if (AllUAVTypes.Contains(UntemplatedType) || PlatformConfiguration.ExtraUAVTypes.Contains(UntemplatedType))
 	{
 		return EShaderParameterType::UAV;
 	}
@@ -421,7 +399,7 @@ bool FShaderParameterParser::ParseParameters(
 
 				if ((bBindlessResources || bBindlessSamplers) && ParsedParameterType == EShaderParameterType::LooseData)
 				{
-					ParsedParameterType = ParseParameterType(Type, ExtraSRVTypes, ExtraUAVTypes);
+					ParsedParameterType = ParseParameterType(Type);
 
 					if (bBindlessResources && ParsedParameterType == EShaderParameterType::SRV)
 					{
@@ -954,48 +932,52 @@ FString FShaderParameterParser::GenerateBindlessParameterDeclaration(const FPars
 	const TCHAR* StorageClass = ParsedParameter.bGloballyCoherent ? TEXT("globallycoherent ") : TEXT("");
 	const FStringView IndexPrefix = ::GetBindlessParameterPrefix(ParsedParameter.BindlessConversionType);
 
+	TStringBuilder<64> IndexString;
+	IndexString << IndexPrefix << Name;
+
 	TStringBuilder<512> Result;
 
 	// If we weren't going to be added to a root constant buffer, that means we need to declare our index before we declare our getter.
 	if (ParsedParameter.ConstantBufferParameterType == EShaderParameterType::Num)
 	{
 		// e.g. `uint BindlessResource_##Name;`
-		Result << TEXT("uint ") << IndexPrefix << Name << TEXT("; ");
+		Result << TEXT("uint ") << IndexString << TEXT("; ");
 	}
 
-	FString RewriteType(Type);
+	// Add the typedef to keep return types shortened
+	//  `typedef Type SafeType##Name;`
 
-	FString TypedefText = TEXT("");
-	// Vulkan requires these typedefs for both resources and samplers, while other platforms only require resources.
-	if (ParsedParameter.BindlessConversionType != EBindlessConversionType::Sampler || BindlessParameterMode == EBindlessParameterMode::Vulkan)
-	{
-		RewriteType = FString::Printf(TEXT("SafeType%.*s"), Name.Len(), Name.GetData());
-		TypedefText = FString::Printf(TEXT("typedef %.*s %s;"), Type.Len(), Type.GetData(), *RewriteType);
-	}
+	TStringBuilder<64> TypedefName;
+	TypedefName << TEXT("SafeType") << Name;
+
+	Result << TEXT("typedef ") << Type << TEXT(" ") << TypedefName << TEXT(";");
 
-	Result << TypedefText;
+	// Full type to use for return types. Makes sure globallycoherent is used where needed. Should be using the typedef name.
+	TStringBuilder<64> FullType;
+	FullType << StorageClass << TypedefName;
 
-	if (BindlessParameterMode == EBindlessParameterMode::Vulkan)
+	if (EnumHasAnyFlags(PlatformConfiguration.Flags, EShaderParameterParserConfigurationFlags::BindlessUsesArrays))
 	{
 		const FStringView HeapPrefix = GetBindlessArrayHeapPrefix(ParsedParameter.BindlessConversionType);
 
 		// Declare a heap for the RewriteType
 		// e.g. `SafeType##Name ResourceDescriptorHeap_SafeType##Name[];`
-		Result << StorageClass << RewriteType << TEXT(" ") << HeapPrefix << RewriteType << TEXT("[]; ");
+		Result << FullType << TEXT(" ") << HeapPrefix << TypedefName << TEXT("[]; ");
 		// :todo-jn: specify the descriptor set and binding directly in source instead of patching SPIRV
 
 		// e.g. `static const SafeType##Name Name = ResourceDescriptorHeap_SafeType##Name[BindlessResource_##Name];`
-		Result << TEXT("static const ") << StorageClass << RewriteType << TEXT(" ") << Name << TEXT(" = ") << HeapPrefix << RewriteType << TEXT("[") << IndexPrefix << Name << TEXT("];");
+		Result << TEXT("static const ") << FullType << TEXT(" ") << Name << TEXT(" = ") << HeapPrefix << TypedefName << TEXT("[") << IndexString << TEXT("];");
 	}
 	else
 	{
+		const FString BindlessAccess = PlatformConfiguration.GenerateBindlessAccess(ParsedParameter.BindlessConversionType, FullType, IndexString);
+
 		const TCHAR* Kind = bIsSampler ? TEXT("Sampler") : TEXT("Resource");
 
 		// e.g. `Type GetBindlessResource##Name() { return GetResourceFromHeap(Type, BindlessResource_##Name); } static const Type Name = GetBindlessResource##Name()`
 		// or   `Type GetBindlessSampler##Name() { return GetSamplerFromHeap(Type, BindlessSampler_##Name); } static const Type Name = GetBindlessSampler##Name()`
-		Result << StorageClass << RewriteType << TEXT(" GetBindless") << Kind << Name << TEXT("()");
-		Result << TEXT("{ return Get") << Kind << TEXT("FromHeap(") << StorageClass << RewriteType << TEXT(", ") << IndexPrefix << Name << TEXT("); } ");
-		Result << TEXT("static const ") << StorageClass << RewriteType << TEXT(" ") << Name << TEXT(" = GetBindless") << Kind << Name << TEXT("();");
+		Result << FullType << TEXT(" GetBindless") << Kind << Name << TEXT("() { return ") << BindlessAccess << TEXT("; } ");
+		Result << TEXT("static const ") << FullType << TEXT(" ") << Name << TEXT(" = GetBindless") << Kind << Name << TEXT("();");
 	}
 
 	return Result.ToString();
@@ -1105,7 +1087,9 @@ bool FShaderParameterParser::MoveShaderParametersToRootConstantBuffer(
 	// Generate the root cbuffer content.
 	if (RootParametersStructure && bNeedToMoveToRootConstantBuffer)
 	{
-		FString RootCBufferContent;
+		FStringBuilderBase ConstantBufferCode;
+
+		ConstantBufferCode << PlatformConfiguration.ConstantBufferType << TEXT(" ") << FShaderParametersMetadata::kRootUniformBufferBindingName << TEXT("\n{\n");
 
 		IterateShaderParameterMembers(
 			*RootParametersStructure,
@@ -1125,51 +1109,46 @@ bool FShaderParameterParser::MoveShaderParametersToRootConstantBuffer(
 				{
 					const FStringView Prefix = GetBindlessParameterPrefix(ParsedParameter->ConstantBufferParameterType);
 
-					RootCBufferContent.Append(FString::Printf(
+					ConstantBufferCode.Appendf(
 						TEXT("uint %.*s%.*s : packoffset(c%d%s);\n"),
 						SVARG(Prefix),
 						SVARG(ParsedParameter->ParsedName),
 						ConstantRegister,
-						ConstantSwizzle));
+						ConstantSwizzle
+					);
 				}
 				else if (ParsedParameter->ConstantBufferParameterType == EShaderParameterType::LooseData)
 				{
 					if (!ParsedParameter->ParsedArraySize.IsEmpty())
 					{
-						RootCBufferContent.Append(FString::Printf(
+						ConstantBufferCode.Appendf(
 							TEXT("%.*s %s[%.*s] : packoffset(c%d%s);\n"),
 							SVARG(ParsedParameter->ParsedType),
 							ShaderBindingName,
 							SVARG(ParsedParameter->ParsedArraySize),
 							ConstantRegister,
-							ConstantSwizzle));
+							ConstantSwizzle
+						);
 					}
 					else
 					{
-						RootCBufferContent.Append(FString::Printf(
+						ConstantBufferCode.Appendf(
 							TEXT("%.*s %s : packoffset(c%d%s);\n"),
 							SVARG(ParsedParameter->ParsedType),
 							ShaderBindingName,
 							ConstantRegister,
-							ConstantSwizzle));
+							ConstantSwizzle
+						);
 					}
 				}
 #undef SVARG
-
 			}
 		});
 
-		FString CBufferCodeBlock = FString::Printf(
-			TEXT("%s %s\n")
-			TEXT("{\n")
-			TEXT("%s")
-			TEXT("}\n\n"),
-			ConstantBufferType,
-			FShaderParametersMetadata::kRootUniformBufferBindingName,
-			*RootCBufferContent);
+		ConstantBufferCode << TEXT("}\n\n");
 
 		FString NewShaderCode = (
-			MakeInjectedShaderCodeBlock(TEXT("MoveShaderParametersToRootConstantBuffer"), CBufferCodeBlock) +
+			MakeInjectedShaderCodeBlock(TEXT("MoveShaderParametersToRootConstantBuffer"), *ConstantBufferCode) +
 			PreprocessedShaderSource);
 
 		PreprocessedShaderSource = MoveTemp(NewShaderCode);
@@ -1181,36 +1160,30 @@ bool FShaderParameterParser::MoveShaderParametersToRootConstantBuffer(
 	return bSuccess;
 }
 
-PRAGMA_DISABLE_DEPRECATION_WARNINGS
-bool FShaderParameterParser::ParseAndModify(
-	const FShaderCompilerInput& CompilerInput,
-	FShaderCompilerOutput& CompilerOutput,
-	FString& PreprocessedShaderSource)
+bool FShaderParameterParser::ParseAndModify(const FShaderCompilerInput& CompilerInput, TArray<FShaderCompilerError>& OutErrors, FString& PreprocessedShaderSource)
 {
-	// assume that if this deprecated overload is still being called, that also the deprecated constructor overloads
-	// are also still being used (and so re-set these member variables according to the environment in the input struct)
 	bBindlessResources = CompilerInput.Environment.CompilerFlags.Contains(CFLAG_BindlessResources);
 	bBindlessSamplers = CompilerInput.Environment.CompilerFlags.Contains(CFLAG_BindlessSamplers);
-	return ParseAndModify(CompilerInput, CompilerOutput.Errors, PreprocessedShaderSource);
-}
-PRAGMA_ENABLE_DEPRECATION_WARNINGS
 
-bool FShaderParameterParser::ParseAndModify(
-	const FShaderCompilerInput& CompilerInput,
-	TArray<FShaderCompilerError>& OutErrors,
-	FString& PreprocessedShaderSource,
-	EBindlessParameterMode InBindlessParameterMode)
-{
+	const bool bUseStableConstantBuffer = EnumHasAnyFlags(PlatformConfiguration.Flags, EShaderParameterParserConfigurationFlags::UseStableConstantBuffer);
+	const bool bSupportsBindless = EnumHasAnyFlags(PlatformConfiguration.Flags, EShaderParameterParserConfigurationFlags::SupportsBindless);
+
 	const bool bHasRootParameters = (CompilerInput.RootParametersStructure != nullptr);
-	BindlessParameterMode = InBindlessParameterMode;
+	const bool bRootParametersModification = bUseStableConstantBuffer && (CompilerInput.IsRayTracingShader() || CompilerInput.ShouldUseStableConstantBuffer());
+	const bool bBindlessModifications = bSupportsBindless && (bBindlessResources || bBindlessSamplers);
+
+	const bool bShouldModify = bRootParametersModification || bBindlessModifications;
+
+	// Always parse if we have root parameters since we need that data during reflection validation
+	const bool bShouldParse = bHasRootParameters || bShouldModify;
 
 	// The shader doesn't have any parameter binding through shader structure, therefore don't do anything.
-	if (!(bBindlessResources || bBindlessSamplers || bHasRootParameters))
+	if (!bShouldParse)
 	{
 		return true;
 	}
 
-	bNeedToMoveToRootConstantBuffer = ConstantBufferType != nullptr && (CompilerInput.IsRayTracingShader() || CompilerInput.ShouldUseStableConstantBuffer());
+	bNeedToMoveToRootConstantBuffer = bRootParametersModification;
 	OriginalParsedShader = PreprocessedShaderSource;
 
 	if (!ParseParameters(CompilerInput.RootParametersStructure, OutErrors))
@@ -1218,31 +1191,39 @@ bool FShaderParameterParser::ParseAndModify(
 		return false;
 	}
 
-	RemoveMovingParametersFromSource(PreprocessedShaderSource);
-	ApplyBindlessModifications(PreprocessedShaderSource);
-
 	bool bResult = true;
 
-	if (bNeedToMoveToRootConstantBuffer)
+	if (bShouldModify)
 	{
-		bResult = MoveShaderParametersToRootConstantBuffer(CompilerInput.RootParametersStructure, PreprocessedShaderSource);
-	}
+		RemoveMovingParametersFromSource(PreprocessedShaderSource);
 
-#if DO_GUARD_SLOW
-	if (bResult)
-	{
-		if (DidModifyShader())
+		if (bSupportsBindless)
 		{
-			checkSlow(PreprocessedShaderSource != OriginalParsedShader);
+			ApplyBindlessModifications(PreprocessedShaderSource);
 		}
-		else
+
+		if (bNeedToMoveToRootConstantBuffer)
 		{
-			checkSlow(PreprocessedShaderSource == OriginalParsedShader);
+			bResult = MoveShaderParametersToRootConstantBuffer(CompilerInput.RootParametersStructure, PreprocessedShaderSource);
+		}
+
+#if DO_GUARD_SLOW
+		if (bResult)
+		{
+			if (DidModifyShader())
+			{
+				checkSlow(PreprocessedShaderSource != OriginalParsedShader);
+			}
+			else
+			{
+				checkSlow(PreprocessedShaderSource == OriginalParsedShader);
+			}
 		}
-	}
 #endif
-	
-	return true;
+
+	}
+
+	return bResult;
 }
 
 void FShaderParameterParser::ValidateShaderParameterType(
diff --git a/Engine/Source/Runtime/RenderCore/Public/ShaderParameterParser.h b/Engine/Source/Runtime/RenderCore/Public/ShaderParameterParser.h
--- a/Engine/Source/Runtime/RenderCore/Public/ShaderParameterParser.h
+++ b/Engine/Source/Runtime/RenderCore/Public/ShaderParameterParser.h
@@ -22,16 +22,46 @@ enum class EBindlessConversionType : uint8
 	Sampler,
 };
 
-enum class EBindlessParameterMode : uint8
+enum class EShaderParameterParserConfigurationFlags
 {
-	Default,
-	Vulkan,
+	None                    = 0,
+	UseStableConstantBuffer = 1 << 0,
+	SupportsBindless        = 1 << 1,
+	// "Vulkan" style
+	BindlessUsesArrays      = 1 << 2,
 };
+ENUM_CLASS_FLAGS(EShaderParameterParserConfigurationFlags)
 
 /** Validates and moves all the shader loose data parameter defined in the root scope of the shader into the root uniform buffer. */
 class FShaderParameterParser
 {
 public:
+	struct FPlatformConfiguration
+	{
+		FPlatformConfiguration() = default;
+		FPlatformConfiguration(FStringView InConstantBufferType, EShaderParameterParserConfigurationFlags InFlags = EShaderParameterParserConfigurationFlags::None)
+			: ConstantBufferType(InConstantBufferType)
+			, Flags(InFlags)
+		{
+			if (InConstantBufferType.Len())
+			{
+				EnumAddFlags(Flags, EShaderParameterParserConfigurationFlags::UseStableConstantBuffer);
+			}
+		}
+
+		/** Generate shader code for accessing a bindless resource or sampler */
+		virtual FString GenerateBindlessAccess(EBindlessConversionType BindlessType, FStringView ShaderTypeString, FStringView IndexString) const
+		{
+			checkf(false, TEXT("Platforms that support bindless must override GenerateBindlessAccess"));
+			return FString();
+		}
+
+		TConstArrayView<FStringView> ExtraSRVTypes;
+		TConstArrayView<FStringView> ExtraUAVTypes;
+		FStringView ConstantBufferType;
+		EShaderParameterParserConfigurationFlags Flags = EShaderParameterParserConfigurationFlags::None;
+	};
+
 	struct FParsedShaderParameter
 	{
 	public:
@@ -74,16 +104,8 @@ public:
 		friend class FShaderParameterParser;
 	};
 
-	UE_DEPRECATED(5.3, "Use FShaderParameterParser constructor which accepts FShaderCompilerFlags")
-	RENDERCORE_API FShaderParameterParser(const TCHAR* InConstantBufferType);
-	UE_DEPRECATED(5.3, "Use FShaderParameterParser constructor which accepts FShaderCompilerFlags")
-	RENDERCORE_API FShaderParameterParser(const TCHAR* InConstantBufferType, TConstArrayView<const TCHAR*> InExtraSRVTypes, TConstArrayView<const TCHAR*> InExtraUAVTypes);
-
-	RENDERCORE_API FShaderParameterParser(
-		FShaderCompilerFlags CompilerFlags,
-		const TCHAR* InConstantBufferType = nullptr,
-		TConstArrayView<const TCHAR*> InExtraSRVTypes = {},
-		TConstArrayView<const TCHAR*> InExtraUAVTypes = {});
+	FShaderParameterParser() = delete;
+	RENDERCORE_API FShaderParameterParser(const FPlatformConfiguration& InPlatformConfiguration);
 
 	RENDERCORE_API virtual ~FShaderParameterParser();
 
@@ -96,25 +118,15 @@ public:
 	static constexpr const TCHAR* kBindlessUAVArrayPrefix = TEXT("UAVDescriptorHeap_");
 	static constexpr const TCHAR* kBindlessSamplerArrayPrefix = TEXT("SamplerDescriptorHeap_");
 
-	static RENDERCORE_API EShaderParameterType ParseParameterType(FStringView InType, TConstArrayView<const TCHAR*> InExtraSRVTypes, TConstArrayView<const TCHAR*> InExtraUAVTypes);
 	static RENDERCORE_API EShaderParameterType ParseAndRemoveBindlessParameterPrefix(FStringView& InName);
 	static RENDERCORE_API EShaderParameterType ParseAndRemoveBindlessParameterPrefix(FString& InName);
 	static RENDERCORE_API bool RemoveBindlessParameterPrefix(FString& InName);
 	static RENDERCORE_API FStringView GetBindlessParameterPrefix(EShaderParameterType InShaderParameterType);
 
-	/** Parses the preprocessed shader code and applies the necessary modifications to it. */
-	UE_DEPRECATED(5.3, "Use ParseAndModify overload accepting array of FShaderCompilerError instead of passing FShaderCompilerOutput")
-	RENDERCORE_API bool ParseAndModify(
-		const FShaderCompilerInput& CompilerInput,
-		FShaderCompilerOutput& CompilerOutput,
-		FString& PreprocessedShaderSource
-	);
-
 	RENDERCORE_API bool ParseAndModify(
 		const FShaderCompilerInput& CompilerInput,
 		TArray<FShaderCompilerError>& OutErrors,
-		FString& PreprocessedShaderSource,
-		EBindlessParameterMode BindlessParameterMode = EBindlessParameterMode::Default
+		FString& PreprocessedShaderSource
 	);
 
 	/** Gets parsing information from a parameter binding name. */
@@ -169,6 +181,8 @@ public:
 	bool DidModifyShader() const { return bModifiedShader; }
 
 protected:
+	RENDERCORE_API EShaderParameterType ParseParameterType(FStringView InType);
+
 	/** Parses the preprocessed shader code */
 	RENDERCORE_API bool ParseParameters(
 		const FShaderParametersMetadata* RootParametersStructure,
@@ -198,10 +212,7 @@ protected:
 	*/
 	RENDERCORE_API FString GenerateBindlessParameterDeclaration(const FParsedShaderParameter& ParsedParameter) const;
 
-	const TCHAR* ConstantBufferType = nullptr;
-
-	TConstArrayView<const TCHAR*> ExtraSRVTypes;
-	TConstArrayView<const TCHAR*> ExtraUAVTypes;
+	const FPlatformConfiguration& PlatformConfiguration;
 
 	FString OriginalParsedShader;
 
@@ -209,7 +220,6 @@ protected:
 
 	bool bBindlessResources = false;
 	bool bBindlessSamplers = false;
-	EBindlessParameterMode BindlessParameterMode = EBindlessParameterMode::Default;
 
 	/** Indicates that parameters should be moved to the root cosntant buffer. */
 	bool bNeedToMoveToRootConstantBuffer = false;
