[shaders] include parameter structure hash in job input hash for both caching methods; use a more robust hashing mechanism than re-hashing the uint32 (this is based on GetTypeHash which is not intended to be persisted per the comments in the header)

diff --git a/Engine/Source/Runtime/RenderCore/Private/ShaderCore.cpp b/Engine/Source/Runtime/RenderCore/Private/ShaderCore.cpp
--- a/Engine/Source/Runtime/RenderCore/Private/ShaderCore.cpp
+++ b/Engine/Source/Runtime/RenderCore/Private/ShaderCore.cpp
@@ -3660,8 +3660,7 @@ FShaderCommonCompileJob::FInputHash FShaderCompileJob::GetInputHash()
 
 		if (Input.RootParametersStructure)
 		{
-			uint32 LayoutHash = Input.RootParametersStructure->GetLayoutHash();
-			Hasher << LayoutHash;
+			const_cast<FShaderParametersMetadata*>(Input.RootParametersStructure)->SerializeLayout(Hasher);
 		}
 
 		InputHash = Hasher.Finalize();
@@ -3681,6 +3680,11 @@ FShaderCommonCompileJob::FInputHash FShaderCompileJob::GetInputHash()
 			Input.DebugGroupName = MoveTemp(DebugGroupNameTmp);
 			Input.Environment.SerializeEverythingButFiles(Archive);
 
+			if (Input.RootParametersStructure)
+			{
+				const_cast<FShaderParametersMetadata*>(Input.RootParametersStructure)->SerializeLayout(Archive);
+			}
+
 			// hash the source file so changes to files during the development are picked up
 			const FSHAHash& SourceHash = GetShaderFileHash(*Input.VirtualSourceFilePath, Input.Target.GetPlatform());
 			Archive << const_cast<FSHAHash&>(SourceHash);
diff --git a/Engine/Source/Runtime/RenderCore/Private/ShaderParameterMetadata.cpp b/Engine/Source/Runtime/RenderCore/Private/ShaderParameterMetadata.cpp
--- a/Engine/Source/Runtime/RenderCore/Private/ShaderParameterMetadata.cpp
+++ b/Engine/Source/Runtime/RenderCore/Private/ShaderParameterMetadata.cpp
@@ -456,6 +456,55 @@ void FShaderParametersMetadata::InitializeAllUniformBufferStructs()
 	}
 }
 
+void FShaderParametersMetadata::FMember::SerializeLayout(FArchive& Ar)
+{
+	// This is only used at the moment for writing to a buffer for hashing; deserialize not supported
+	// due to various const members, but also not needed.
+	check(Ar.IsSaving());
+
+	Ar << Offset;
+	Ar << reinterpret_cast<uint8&>(BaseType);
+	Ar.Serialize(const_cast<TCHAR*>(Name), FCString::Strlen(Name));
+	Ar << NumElements;
+
+	const bool bIsRHIResource = (
+		BaseType == UBMT_TEXTURE ||
+		BaseType == UBMT_SRV ||
+		BaseType == UBMT_SAMPLER);
+	const bool bIsRDGResource = IsRDGResourceReferenceShaderParameterType(BaseType);
+
+	if (BaseType == UBMT_INT32 ||
+		BaseType == UBMT_UINT32 ||
+		BaseType == UBMT_FLOAT32)
+	{
+		Ar << reinterpret_cast<uint8&>(Precision);
+		Ar << NumRows;
+		Ar << NumColumns;
+	}
+	else if (BaseType == UBMT_INCLUDED_STRUCT || BaseType == UBMT_NESTED_STRUCT)
+	{
+		const_cast<FShaderParametersMetadata*>(Struct)->SerializeLayout(Ar);
+	}
+	else if (bIsRHIResource || bIsRDGResource)
+	{
+		Ar.Serialize(const_cast<TCHAR*>(ShaderType), FCString::Strlen(ShaderType));
+	}
+}
+
+void FShaderParametersMetadata::SerializeLayout(FArchive& Ar) 
+{
+	// This is only used at the moment for writing to a buffer for hashing; deserialize not supported
+	check(Ar.IsSaving());
+
+	Ar << const_cast<uint32&>(Size);
+
+	for (FMember& CurrentMember : Members)
+	{
+		CurrentMember.SerializeLayout(Ar);
+	}
+}
+
+
 void FShaderParametersMetadata::InitializeLayout(FRHIUniformBufferLayoutInitializer* OutLayoutInitializer)
 {
 	check(!IsLayoutInitialized());
diff --git a/Engine/Source/Runtime/RenderCore/Public/ShaderParameterMetadata.h b/Engine/Source/Runtime/RenderCore/Public/ShaderParameterMetadata.h
--- a/Engine/Source/Runtime/RenderCore/Public/ShaderParameterMetadata.h
+++ b/Engine/Source/Runtime/RenderCore/Public/ShaderParameterMetadata.h
@@ -251,6 +251,8 @@ public:
 			return ElementSize;
 		}
 
+		void SerializeLayout(FArchive& Ar);
+
 		static RENDERCORE_API void GenerateShaderParameterType(
 			FString& Result,
 			bool bSupportsPrecisionModifier,
@@ -389,6 +391,8 @@ public:
 		return LayoutHash;	
 	}
 
+	void SerializeLayout(FArchive& Ar);
+
 	/** Iterate recursively over all FShaderParametersMetadata. */
 	template<typename TParameterFunction>
 	void IterateStructureMetadataDependencies(TParameterFunction Lambda) const
