Fixed: Compile errors in RenderUtils.cpp
Moved FReadOnlyCVARCache vars out of global scope

diff --git a/Engine/Source/Runtime/RenderCore/Private/ReadOnlyCVARCache.cpp b/Engine/Source/Runtime/RenderCore/Private/ReadOnlyCVARCache.cpp
--- a/Engine/Source/Runtime/RenderCore/Private/ReadOnlyCVARCache.cpp
+++ b/Engine/Source/Runtime/RenderCore/Private/ReadOnlyCVARCache.cpp
@@ -9,27 +9,6 @@
 #include "Interfaces/ITargetPlatformManagerModule.h"
 #endif
 
-static bool bEnablePointLightShadows;
-static bool bEnableStationarySkylight;
-static bool bEnableLowQualityLightmaps;
-static bool bAllowStaticLighting;
-static bool bSupportSkyAtmosphere;
-
-// Mobile specific
-static bool bMobileHDR;
-static bool bMobileAllowMovableDirectionalLights;
-static bool bMobileAllowDistanceFieldShadows;
-static bool bMobileEnableStaticAndCSMShadowReceivers;
-static bool bMobileEnableMovableLightCSMShaderCulling;
-static bool bMobileSupportsGPUScene;
-static int32 MobileSkyLightPermutationValue;
-static int32 MobileEarlyZPassValue;
-static int32 MobileForwardLocalLightsValue;
-static bool bMobileEnableNoPrecomputedLightingCSMShader;
-static bool bMobileDeferredShadingValue;
-static bool bMobileEnableMovableSpotlightsShadowValue;
-
-
 static int32 MobileEarlyZPassIniValue(const FStaticShaderPlatform Platform)
 {
 	static FShaderPlatformCachedIniValue<int32> CVar(TEXT("r.Mobile.EarlyZPass"));
@@ -59,6 +38,26 @@ static bool MobileEnableMovableSpotlightsShadowIniValue(const FStaticShaderPlatf
 
 bool FReadOnlyCVARCache::bInitialized = false;
 
+bool FReadOnlyCVARCache::bEnablePointLightShadows = true;
+bool FReadOnlyCVARCache::bEnableStationarySkylight = true;
+bool FReadOnlyCVARCache::bEnableLowQualityLightmaps = true;
+bool FReadOnlyCVARCache::bAllowStaticLighting = true;
+bool FReadOnlyCVARCache::bSupportSkyAtmosphere = true;
+
+// Mobile specific
+bool FReadOnlyCVARCache::bMobileHDR = true;
+bool FReadOnlyCVARCache::bMobileAllowMovableDirectionalLights = true;
+bool FReadOnlyCVARCache::bMobileAllowDistanceFieldShadows = true;
+bool FReadOnlyCVARCache::bMobileEnableStaticAndCSMShadowReceivers = true;
+bool FReadOnlyCVARCache::bMobileEnableMovableLightCSMShaderCulling = true;
+bool FReadOnlyCVARCache::bMobileSupportsGPUScene = false;
+int32 FReadOnlyCVARCache::MobileSkyLightPermutationValue = 0;
+int32 FReadOnlyCVARCache::MobileEarlyZPassValue = 0;
+int32 FReadOnlyCVARCache::MobileForwardLocalLightsValue = 1;
+bool FReadOnlyCVARCache::bMobileEnableNoPrecomputedLightingCSMShader = false;
+bool FReadOnlyCVARCache::bMobileDeferredShadingValue = false;
+bool FReadOnlyCVARCache::bMobileEnableMovableSpotlightsShadowValue = false;
+
 void FReadOnlyCVARCache::Initialize()
 {
 	UE_LOG(LogInit, Log, TEXT("Initializing FReadOnlyCVARCache"));
diff --git a/Engine/Source/Runtime/RenderCore/Public/ReadOnlyCVARCache.h b/Engine/Source/Runtime/RenderCore/Public/ReadOnlyCVARCache.h
--- a/Engine/Source/Runtime/RenderCore/Public/ReadOnlyCVARCache.h
+++ b/Engine/Source/Runtime/RenderCore/Public/ReadOnlyCVARCache.h
@@ -34,5 +34,25 @@ struct FReadOnlyCVARCache
 
 private:
 	static bool bInitialized;
+
+	static bool bEnablePointLightShadows;
+	static bool bEnableStationarySkylight;
+	static bool bEnableLowQualityLightmaps;
+	static bool bAllowStaticLighting;
+	static bool bSupportSkyAtmosphere;
+
+	// Mobile specific
+	static bool bMobileHDR;
+	static bool bMobileAllowMovableDirectionalLights;
+	static bool bMobileAllowDistanceFieldShadows;
+	static bool bMobileEnableStaticAndCSMShadowReceivers;
+	static bool bMobileEnableMovableLightCSMShaderCulling;
+	static bool bMobileSupportsGPUScene;
+	static int32 MobileSkyLightPermutationValue;
+	static int32 MobileEarlyZPassValue;
+	static int32 MobileForwardLocalLightsValue;
+	static bool bMobileEnableNoPrecomputedLightingCSMShader;
+	static bool bMobileDeferredShadingValue;
+	static bool bMobileEnableMovableSpotlightsShadowValue;
 };
 
