[shaders] fix for static analysis warning

diff --git a/Engine/Source/Runtime/RenderCore/Private/ShaderCore.cpp b/Engine/Source/Runtime/RenderCore/Private/ShaderCore.cpp
--- a/Engine/Source/Runtime/RenderCore/Private/ShaderCore.cpp
+++ b/Engine/Source/Runtime/RenderCore/Private/ShaderCore.cpp
@@ -3579,8 +3579,9 @@ FShaderCommonCompileJob::FInputHash FShaderPipelineCompileJob::GetInputHash()
 	{
 		if (StageJobs[Index])
 		{
-			FShaderCommonCompileJob::FInputHash StageHash = StageJobs[Index]->GetInputHash();
-			CombinedHash += int256(StageHash.GetBytes(), sizeof(StageHash.GetBytes()));
+			const FShaderCommonCompileJob::FInputHash::ByteArray& StageHash = StageJobs[Index]->GetInputHash().GetBytes();
+			static_assert(sizeof(FShaderCommonCompileJob::FInputHash::ByteArray) == sizeof(int256));
+			CombinedHash += int256(StageHash, sizeof(StageHash));
 		}
 	}
 
