ShaderMaps in loose cooked builds: Fix not being able to run with loose files when using -zenstore. Fall back to looking for .ushaderbytecode files on disk even when using iostore, since zenstore uses iostore, but can run in non-packaged mode where it needs to read the shadermaps off of disk.
#rnx

diff --git a/Engine/Source/Runtime/RenderCore/Private/ShaderCodeLibrary.cpp b/Engine/Source/Runtime/RenderCore/Private/ShaderCodeLibrary.cpp
--- a/Engine/Source/Runtime/RenderCore/Private/ShaderCodeLibrary.cpp
+++ b/Engine/Source/Runtime/RenderCore/Private/ShaderCodeLibrary.cpp
@@ -113,6 +113,30 @@ namespace UE
 				return FString::Printf(TEXT("%s_Chunk%d"), *BaseName, ChunkId);
 			}
 
+			bool IsRunningWithPakFile()
+			{
+				static const bool bRunningWithPakFile =
+					(FPlatformFileManager::Get().FindPlatformFile(TEXT("PakFile")) != nullptr);
+				return bRunningWithPakFile;
+			}
+			bool IsRunningWithIoStore()
+			{
+				static const bool bRunningWithIoStore =
+					FIoDispatcher::IsInitialized()
+					&& FIoDispatcher::Get().DoesChunkExist(CreateIoChunkId(0, 0, EIoChunkType::ScriptObjects));
+				return bRunningWithIoStore;
+			}
+			bool IsRunningWithZenStore()
+			{
+				static const bool bRunningWithZenStore =
+					FPlatformFileManager::Get().FindPlatformFile(TEXT("StorageServer")) != nullptr;
+				return bRunningWithZenStore;
+			}
+			bool ShouldLookForLooseCookedChunks()
+			{
+				return IsRunningWithZenStore() || !IsRunningWithIoStore();
+			}
+
 			// [RCL] TODO 2020-11-20: Separate runtime and editor-only code (tracked as UE-103486)
 			/** Descriptor used to pass the pak file information to the library as we cannot store IPakFile ref */
 			struct FMountedPakFileInfo
@@ -853,16 +877,14 @@ public:
 			ShaderCodeDirectory = ShaderCodeDir;
 		}
 
-		const bool bRunningWithIoStore = FIoDispatcher::IsInitialized() && FIoDispatcher::Get().DoesChunkExist(CreateIoChunkId(0, 0, EIoChunkType::ScriptObjects));
-
-		if (!Library && bRunningWithIoStore)
+		if (!Library && UE::ShaderLibrary::Private::IsRunningWithIoStore())
 		{
 			Library = FIoStoreShaderCodeArchive::Create(InShaderPlatform, InLibraryName, FIoDispatcher::Get());
 			ShaderCodeDirectory.Empty();	// paths don't matter for IoStore-based libraries
 		}
 
 		// Shader library as a ushaderbytecode file is no longer an option for distribution. Some cases (a build using loose files) still require its support though.
-		if (!Library && !bRunningWithIoStore)
+		if (!Library && UE::ShaderLibrary::Private::ShouldLookForLooseCookedChunks())
 		{
 			const FName PlatformName = FDataDrivenShaderPlatformInfo::GetName(InShaderPlatform);
 			const FName ShaderFormatName = LegacyShaderPlatformToShaderFormat(InShaderPlatform);
@@ -2605,12 +2627,9 @@ public:
 #if UE_SHADERLIB_SUPPORT_CHUNK_DISCOVERY
 				if (!bResult)
 				{
-					static const bool bRunningWithPakFile = (FPlatformFileManager::Get().FindPlatformFile(TEXT("PakFile")) != nullptr);
-					static const bool bRunningWithIoStore = FIoDispatcher::IsInitialized() && FIoDispatcher::Get().DoesChunkExist(CreateIoChunkId(0, 0, EIoChunkType::ScriptObjects));
-
 					// Some deployment flows (e.g. Launch on) avoid pak files despite project packaging settings. 
 					// In case we run under such circumstances, we need to discover the components ourselves
-					if (!bRunningWithPakFile && !bRunningWithIoStore)
+					if (!IsRunningWithPakFile() && ShouldLookForLooseCookedChunks())
 					{
 						UE_LOG(LogShaderLibrary, Display, TEXT("Running without a pakfile/IoStore and did not find a monolithic library '%s' - attempting disk search for its chunks"), *Name);
 
