Fix error C2872: 'FSphere3d': ambiguous symbol
#rnx

diff --git a/Engine/Source/Runtime/RenderCore/Public/ShaderParameterUtils.h b/Engine/Source/Runtime/RenderCore/Public/ShaderParameterUtils.h
--- a/Engine/Source/Runtime/RenderCore/Public/ShaderParameterUtils.h
+++ b/Engine/Source/Runtime/RenderCore/Public/ShaderParameterUtils.h
@@ -504,7 +504,7 @@ GUARD_SETSHADERVALUE(FVector4)
 GUARD_SETSHADERVALUE(FPlane4)
 GUARD_SETSHADERVALUE(FQuat4)
 // Secondary
-GUARD_SETSHADERVALUE(FSphere3)
+GUARD_SETSHADERVALUE(::FSphere3)
 GUARD_SETSHADERVALUE(FBox3)
 
 /**
