Do not attempt to open plugin shader libraries on the server.

diff --git a/Engine/Source/Runtime/RenderCore/Private/ShaderCodeLibrary.cpp b/Engine/Source/Runtime/RenderCore/Private/ShaderCodeLibrary.cpp
--- a/Engine/Source/Runtime/RenderCore/Private/ShaderCodeLibrary.cpp
+++ b/Engine/Source/Runtime/RenderCore/Private/ShaderCodeLibrary.cpp
@@ -3408,7 +3408,7 @@ FShaderLibrariesCollection* FShaderLibrariesCollection::Impl = nullptr;
 
 static void FShaderCodeLibraryPluginMountedCallback(IPlugin& Plugin)
 {
-	if (UE::ShaderLibrary::Private::PluginsToIgnoreOnMount.Remove(Plugin.GetName()) == 0)
+	if (FApp::CanEverRender() && UE::ShaderLibrary::Private::PluginsToIgnoreOnMount.Remove(Plugin.GetName()) == 0)
 	{
 		FShaderCodeLibrary::OpenPluginShaderLibrary(Plugin);
 	}
@@ -4100,12 +4100,15 @@ void FShaderCodeLibrary::UnregisterSharedShaderCodeRequestDelegate_Handle(FDeleg
 void FShaderCodeLibrary::DontOpenPluginShaderLibraryOnMount(const FString& PluginName)
 {
 	check(IsInGameThread());
-	UE::ShaderLibrary::Private::PluginsToIgnoreOnMount.Add(PluginName);
+	if (FApp::CanEverRender())
+	{
+		UE::ShaderLibrary::Private::PluginsToIgnoreOnMount.Add(PluginName);
+	}
 }
 
 void FShaderCodeLibrary::OpenPluginShaderLibrary(IPlugin& Plugin, bool bMonolithicOnly)
 {
-	if (Plugin.CanContainContent() && Plugin.IsEnabled())
+	if (Plugin.CanContainContent() && Plugin.IsEnabled() && FApp::CanEverRender())
 	{
 		// load any shader libraries that may exist in this plugin
 		if (!bMonolithicOnly)
