Fix assertion failure about CurrentlyDeleting == this when one of the RHICreate*Shader functions fails in D3D12.

The engine will likely still crash in this case because shader creation failures are fatal, but it should do that at the correct spot now, where the appropriate debug information can be logged. Also made FShaderMapResource_InlineCode::CreateRHIShaderOrCrash print the hash of the failing shader, so we have actionable information in the log.

#rnx

diff --git a/Engine/Source/Runtime/D3D12RHI/Private/D3D12Shaders.cpp b/Engine/Source/Runtime/D3D12RHI/Private/D3D12Shaders.cpp
--- a/Engine/Source/Runtime/D3D12RHI/Private/D3D12Shaders.cpp
+++ b/Engine/Source/Runtime/D3D12RHI/Private/D3D12Shaders.cpp
@@ -128,7 +128,8 @@ TShaderType* CreateStandardShader(TArrayView<const uint8> InCode)
 
 	if (!InitShaderCommon(ShaderCode, Offset, Shader))
 	{
-		delete Shader;
+		Shader->AddRef();
+		Shader->Release();
 		return nullptr;
 	}
 
diff --git a/Engine/Source/Runtime/RenderCore/Private/ShaderResource.cpp b/Engine/Source/Runtime/RenderCore/Private/ShaderResource.cpp
--- a/Engine/Source/Runtime/RenderCore/Private/ShaderResource.cpp
+++ b/Engine/Source/Runtime/RenderCore/Private/ShaderResource.cpp
@@ -657,7 +657,7 @@ FRHIShader* FShaderMapResource_InlineCode::CreateRHIShaderOrCrash(int32 ShaderIn
 	}
 	if (UNLIKELY(RHIShader == nullptr))
 	{
-		UE_LOG(LogShaders, Fatal, TEXT("FShaderMapResource_InlineCode::InitRHI is unable to create a shader (frequency %d)"), static_cast<int32>(Frequency));
+		UE_LOG(LogShaders, Fatal, TEXT("FShaderMapResource_InlineCode::InitRHI is unable to create a shader: frequency=%d, hash=%s."), static_cast<int32>(Frequency), *ShaderHash.ToString());
 		// unreachable
 		return nullptr;
 	}
