Use FOptionalTaskTagScope instead of FTaskTagScope to avoid overwriting the current tag when task retraction kicks in.

This fixes some invalid tag ensures in FD3D11DynamicRHI::RHICreateUniformBuffer.

#rnx

diff --git a/Engine/Source/Runtime/RenderCore/Private/RenderGraphBuilder.cpp b/Engine/Source/Runtime/RenderCore/Private/RenderGraphBuilder.cpp
--- a/Engine/Source/Runtime/RenderCore/Private/RenderGraphBuilder.cpp
+++ b/Engine/Source/Runtime/RenderCore/Private/RenderGraphBuilder.cpp
@@ -2559,7 +2559,7 @@ UE::Tasks::FTask FRDGBuilder::SubmitBufferUploads()
 
 		return UE::Tasks::Launch(TEXT("FRDGBuilder::SubmitBufferUploads"), [this, SubmitUploadsLambda]
 		{
-			FTaskTagScope Scope(ETaskTag::EParallelRenderingThread);
+			FOptionalTaskTagScope Scope(ETaskTag::EParallelRenderingThread);
 			RHICmdListBufferUploads = new FRHICommandList(FRHIGPUMask::All());
 			RHICmdListBufferUploads->SwitchPipeline(ERHIPipeline::Graphics);
 			SubmitUploadsLambda(*RHICmdListBufferUploads);
@@ -2579,7 +2579,7 @@ UE::Tasks::FTask FRDGBuilder::SubmitBufferUploads()
 void FRDGBuilder::SetupParallelExecute()
 {
 	SCOPED_NAMED_EVENT(SetupParallelExecute, FColor::Emerald);
-	FTaskTagScope Scope(ETaskTag::EParallelRenderingThread);
+	FOptionalTaskTagScope Scope(ETaskTag::EParallelRenderingThread);
 
 	TArray<FRDGPass*, TInlineAllocator<64, FRDGArrayAllocator>> ParallelPassCandidates;
 	uint32 ParallelPassCandidatesWorkload = 0;
@@ -2790,7 +2790,7 @@ UE::Tasks::FTask FRDGBuilder::CreateUniformBuffers()
 		Task = UE::Tasks::Launch(TEXT("FRDGBuilder::CreateUniformBuffer"),
 			[CreateUniformBuffersFunction]
 		{
-			FTaskTagScope Scope(ETaskTag::EParallelRenderingThread);
+			FOptionalTaskTagScope Scope(ETaskTag::EParallelRenderingThread);
 			CreateUniformBuffersFunction();
 
 		}, LowLevelTasks::ETaskPriority::High);
