Create a dummy BlackCubeArrayTexture without querying project cvars as cubemap arrays are supported on all platforms

diff --git a/Engine/Source/Runtime/RenderCore/Private/GlobalRenderResources.cpp b/Engine/Source/Runtime/RenderCore/Private/GlobalRenderResources.cpp
--- a/Engine/Source/Runtime/RenderCore/Private/GlobalRenderResources.cpp
+++ b/Engine/Source/Runtime/RenderCore/Private/GlobalRenderResources.cpp
@@ -414,31 +414,28 @@ public:
 	// FResource interface.
 	virtual void InitRHI(FRHICommandListBase& RHICmdList) override
 	{
-		if (SupportsTextureCubeArray(GetFeatureLevel()))
-		{
-			const TCHAR* Name = TEXT("BlackCubeArray");
+		const TCHAR* Name = TEXT("BlackCubeArray");
 
-			const FRHITextureCreateDesc Desc =
-				FRHITextureCreateDesc::CreateCubeArray(TEXT("BlackCubeArray"), 1, 1, PF_B8G8R8A8)
-				.SetFlags(ETextureCreateFlags::ShaderResource)
-				.SetClassName(TEXT("FBlackCubeArrayTexture"));
-
-			// Create the texture RHI.
-			TextureRHI = RHICreateTexture(Desc);
+		const FRHITextureCreateDesc Desc =
+			FRHITextureCreateDesc::CreateCubeArray(TEXT("BlackCubeArray"), 1, 1, PF_B8G8R8A8)
+			.SetFlags(ETextureCreateFlags::ShaderResource)
+			.SetClassName(TEXT("FBlackCubeArrayTexture"));
 
-			for (uint32 FaceIndex = 0; FaceIndex < 6; FaceIndex++)
-			{
-				uint32 DestStride;
-				FColor* DestBuffer = (FColor*)RHILockTextureCubeFace(TextureRHI, FaceIndex, 0, 0, RLM_WriteOnly, DestStride, false);
-				// Note: alpha is used by reflection environment to say how much of the foreground texture is visible, so 0 says it is completely invisible
-				*DestBuffer = FColor(0, 0, 0, 0);
-				RHIUnlockTextureCubeFace(TextureRHI, FaceIndex, 0, 0, false);
-			}
+		// Create the texture RHI.
+		TextureRHI = RHICreateTexture(Desc);
 
-			// Create the sampler state RHI resource.
-			FSamplerStateInitializerRHI SamplerStateInitializer(SF_Point, AM_Wrap, AM_Wrap, AM_Wrap);
-			SamplerStateRHI = GetOrCreateSamplerState(SamplerStateInitializer);
+		for (uint32 FaceIndex = 0; FaceIndex < 6; FaceIndex++)
+		{
+			uint32 DestStride;
+			FColor* DestBuffer = (FColor*)RHILockTextureCubeFace(TextureRHI, FaceIndex, 0, 0, RLM_WriteOnly, DestStride, false);
+			// Note: alpha is used by reflection environment to say how much of the foreground texture is visible, so 0 says it is completely invisible
+			*DestBuffer = FColor(0, 0, 0, 0);
+			RHIUnlockTextureCubeFace(TextureRHI, FaceIndex, 0, 0, false);
 		}
+
+		// Create the sampler state RHI resource.
+		FSamplerStateInitializerRHI SamplerStateInitializer(SF_Point, AM_Wrap, AM_Wrap, AM_Wrap);
+		SamplerStateRHI = GetOrCreateSamplerState(SamplerStateInitializer);
 	}
 
 	virtual uint32 GetSizeX() const override { return 1; }
diff --git a/Engine/Source/Runtime/Renderer/Private/SystemTextures.cpp b/Engine/Source/Runtime/Renderer/Private/SystemTextures.cpp
--- a/Engine/Source/Runtime/Renderer/Private/SystemTextures.cpp
+++ b/Engine/Source/Runtime/Renderer/Private/SystemTextures.cpp
@@ -291,10 +291,7 @@ void FSystemTextures::InitializeCommonTextures(FRHICommandListImmediate& RHICmdL
 	// Create Cube BlackDummy textures
 	{
 		CubeBlackDummy = CreateRenderTarget(GBlackTextureCube->TextureRHI, TEXT("CubeBlackDummy"));
-		if (SupportsTextureCubeArray(GetFeatureLevel()))
-		{
-			CubeArrayBlackDummy = CreateRenderTarget(GBlackCubeArrayTexture->TextureRHI, TEXT("CubeArrayBlackDummy"));
-		}
+		CubeArrayBlackDummy = CreateRenderTarget(GBlackCubeArrayTexture->TextureRHI, TEXT("CubeArrayBlackDummy"));
 	}
 }
 
