Fix ray tracing shader compilation on SM5 - RHISupportsRayTracing was returning true and we were trying to compile those shaders but they crashed because we also check GRHISupportsRayTracing before compiling them which was set to 0 due to the recent change. Also, D3D12 ray tracing validation shaders were using wave operations which are not supported in SM5.
* Update builtin shaders to use ShouldCompileRayTracingShadersForProject instead of always being compiled.
* Modify ShouldCompileRayTracingShadersForProject to respect r.Raytracing.RequireSM6 cvar.
* Don't compile or use D3D12 ray tracing validation shaders on SM5.

diff --git a/Engine/Source/Runtime/D3D12RHI/Private/D3D12RayTracing.cpp b/Engine/Source/Runtime/D3D12RHI/Private/D3D12RayTracing.cpp
--- a/Engine/Source/Runtime/D3D12RHI/Private/D3D12RayTracing.cpp
+++ b/Engine/Source/Runtime/D3D12RHI/Private/D3D12RayTracing.cpp
@@ -399,6 +399,13 @@ const FD3D12ShaderIdentifier FD3D12ShaderIdentifier::Null = { 0, 0, 0, 0 };
 
 static_assert(sizeof(FD3D12ShaderIdentifier) == D3D12_SHADER_IDENTIFIER_SIZE_IN_BYTES, "Unexpected shader identifier size");
 
+static bool ShouldRunRayTracingGPUValidation()
+{
+	// Wave ops are required to run ray tracing validation shaders
+	const bool bSupportsWaveOps = GRHISupportsWaveOperations && RHISupportsWaveOperations(GMaxRHIShaderPlatform);
+	return GD3D12RayTracingGPUValidation && bSupportsWaveOps;
+}
+
 static D3D12_RAYTRACING_ACCELERATION_STRUCTURE_BUILD_FLAGS TranslateRayTracingAccelerationStructureFlags(ERayTracingAccelerationStructureFlags Flags)
 {
 	uint32 Result = {};
@@ -3672,7 +3679,7 @@ void FD3D12RayTracingScene::BuildAccelerationStructure(FD3D12CommandContext& Com
 		0
 	);
 
-	if (GD3D12RayTracingGPUValidation)
+	if (ShouldRunRayTracingGPUValidation())
 	{
 		for (uint32 LayerIndex = 0; LayerIndex < NumLayers; ++LayerIndex)
 		{
@@ -3957,7 +3964,7 @@ void FD3D12CommandContext::RHIBuildAccelerationStructures(const TArrayView<const
 		}		
 	}
 
-	if (GD3D12RayTracingGPUValidation)
+	if (ShouldRunRayTracingGPUValidation())
 	{
 		TRHICommandList_RecursiveHazardous<FD3D12CommandContext> RHICmdList(this);
 
diff --git a/Engine/Source/Runtime/RenderCore/Private/RayTracingValidationShaders.cpp b/Engine/Source/Runtime/RenderCore/Private/RayTracingValidationShaders.cpp
--- a/Engine/Source/Runtime/RenderCore/Private/RayTracingValidationShaders.cpp
+++ b/Engine/Source/Runtime/RenderCore/Private/RayTracingValidationShaders.cpp
@@ -12,6 +12,12 @@ IMPLEMENT_GLOBAL_SHADER(FRayTracingValidateGeometryBuildParamsCS, "/Engine/Priva
 
 void FRayTracingValidateGeometryBuildParamsCS::Dispatch(FRHICommandList& RHICmdList, const FRayTracingGeometryBuildParams& Params)
 {
+	const bool bSupportsWaveOps = GRHISupportsWaveOperations && RHISupportsWaveOperations(GMaxRHIShaderPlatform);
+	if (!ensureMsgf(bSupportsWaveOps, TEXT("Wave operations are required to run ray tracing GPU validation shaders.")))
+	{
+		return;
+	}
+
 	const FRayTracingGeometryInitializer& Initializer = Params.Geometry->GetInitializer();
 
 	TShaderMapRef<FRayTracingValidateGeometryBuildParamsCS> ComputeShader(GetGlobalShaderMap(GMaxRHIFeatureLevel));
@@ -95,6 +101,12 @@ void FRayTracingValidateSceneBuildParamsCS::Dispatch(FRHICommandList& RHICmdList
 	uint32 NumHitGroups, uint32 NumInstances, 
 	FRHIBuffer* InstanceBuffer, uint32 InstanceBufferOffset, uint32 InstanceBufferStride)
 {
+	const bool bSupportsWaveOps = GRHISupportsWaveOperations && RHISupportsWaveOperations(GMaxRHIShaderPlatform);
+	if (!ensureMsgf(bSupportsWaveOps, TEXT("Wave operations are required to run ray tracing GPU validation shaders.")))
+	{
+		return;
+	}
+
 	TShaderMapRef<FRayTracingValidateSceneBuildParamsCS> ComputeShader(GetGlobalShaderMap(GMaxRHIFeatureLevel));
 	FRHIComputeShader* ShaderRHI = ComputeShader.GetComputeShader();
 	SetComputePipelineState(RHICmdList, ShaderRHI);
diff --git a/Engine/Source/Runtime/RenderCore/Private/RenderUtils.cpp b/Engine/Source/Runtime/RenderCore/Private/RenderUtils.cpp
--- a/Engine/Source/Runtime/RenderCore/Private/RenderUtils.cpp
+++ b/Engine/Source/Runtime/RenderCore/Private/RenderUtils.cpp
@@ -1589,7 +1589,9 @@ bool ShouldCompileRayTracingShadersForProject(EShaderPlatform ShaderPlatform)
 {
 	if (RHISupportsRayTracingShaders(ShaderPlatform))
 	{
-		return IsRayTracingEnabledForProject(ShaderPlatform);
+		const bool bRayTracingRequireSM6 = (GRayTracingRequireSM6 != 0);
+
+		return IsRayTracingEnabledForProject(ShaderPlatform) && (IsFeatureLevelSupported(ShaderPlatform, ERHIFeatureLevel::SM6) || !bRayTracingRequireSM6);		
 	}
 	else
 	{
diff --git a/Engine/Source/Runtime/RenderCore/Public/BuiltInRayTracingShaders.h b/Engine/Source/Runtime/RenderCore/Public/BuiltInRayTracingShaders.h
--- a/Engine/Source/Runtime/RenderCore/Public/BuiltInRayTracingShaders.h
+++ b/Engine/Source/Runtime/RenderCore/Public/BuiltInRayTracingShaders.h
@@ -19,8 +19,7 @@ protected:
 public:
 	static bool ShouldCompilePermutation(const FGlobalShaderPermutationParameters& Parameters)
 	{
-		// Built-in ray tracing shaders are always compiled for RHIs that support them, regardless of whether RT is enabled for the project.
-		return RHISupportsRayTracingShaders(Parameters.Platform);
+		return ShouldCompileRayTracingShadersForProject(Parameters.Platform);
 	}
 
 	FBuiltInRayTracingShader() = default;
diff --git a/Engine/Source/Runtime/RenderCore/Public/RayTracingValidationShaders.h b/Engine/Source/Runtime/RenderCore/Public/RayTracingValidationShaders.h
--- a/Engine/Source/Runtime/RenderCore/Public/RayTracingValidationShaders.h
+++ b/Engine/Source/Runtime/RenderCore/Public/RayTracingValidationShaders.h
@@ -29,7 +29,7 @@ public:
 
 	static bool ShouldCompilePermutation(const FGlobalShaderPermutationParameters& Parameters)
 	{
-		return FBuiltInRayTracingShader::ShouldCompilePermutation(Parameters);
+		return FBuiltInRayTracingShader::ShouldCompilePermutation(Parameters) && RHISupportsWaveOperations(Parameters.Platform);
 	}
 
 	// Large thread group to handle large meshes with a single 1D dispatch
@@ -73,7 +73,7 @@ public:
 
 	static bool ShouldCompilePermutation(const FGlobalShaderPermutationParameters& Parameters)
 	{
-		return FBuiltInRayTracingShader::ShouldCompilePermutation(Parameters);
+		return FBuiltInRayTracingShader::ShouldCompilePermutation(Parameters) && RHISupportsWaveOperations(Parameters.Platform);
 	}
 
 	// Large thread group to handle large meshes with a single 1D dispatch
