[shaders] fix for static analysis warning

diff --git a/Engine/Source/Runtime/RenderCore/Private/ShaderCore.cpp b/Engine/Source/Runtime/RenderCore/Private/ShaderCore.cpp
--- a/Engine/Source/Runtime/RenderCore/Private/ShaderCore.cpp
+++ b/Engine/Source/Runtime/RenderCore/Private/ShaderCore.cpp
@@ -3578,8 +3578,10 @@ FShaderCommonCompileJob::FInputHash FShaderPipelineCompileJob::GetInputHash()
 	{
 		if (StageJobs[Index])
 		{
-			FShaderCommonCompileJob::FInputHash StageHash = StageJobs[Index]->GetInputHash();
-			CombinedHash += int256(StageHash.GetBytes(), sizeof(StageHash.GetBytes()));
+			const FShaderCommonCompileJob::FInputHash StageHash = StageJobs[Index]->GetInputHash();
+			const FShaderCommonCompileJob::FInputHash::ByteArray& StageHashBytes = StageHash.GetBytes();
+			static_assert(sizeof(StageHashBytes) == sizeof(int256));
+			CombinedHash += int256(StageHashBytes, sizeof(StageHashBytes));
 		}
 	}
 
