remove device id check for HDR on windows:
* d3d12 / d3d11-editor are already using DXGI to check whether HDR is supported on the display.
* d3d11-cooked still relies on vendor extensions, but we propagate the information from FD3D11DynamicRHI::SetupDisplayHDRMetaData to GRHISupportsHDROutput. Because of that, we don't even call ChooseHDRDeviceAndColorGamut on vendors we do not support
Remove unnecessary check in EnsureColorSpace: it's valid to call CheckColorSpaceSupport with SDR
fix d3d11-cooked backbuffer format: both vendor extensions expect FP16 ScRGB, not HDR10 format

diff --git a/Engine/Source/Runtime/Core/Public/Windows/WindowsPlatformMisc.h b/Engine/Source/Runtime/Core/Public/Windows/WindowsPlatformMisc.h
--- a/Engine/Source/Runtime/Core/Public/Windows/WindowsPlatformMisc.h
+++ b/Engine/Source/Runtime/Core/Public/Windows/WindowsPlatformMisc.h
@@ -290,21 +290,18 @@ struct FWindowsPlatformMisc
 
 	FORCEINLINE static void ChooseHDRDeviceAndColorGamut(uint32 DeviceId, uint32 DisplayNitLevel, EDisplayOutputFormat& OutputDevice, EDisplayColorGamut& ColorGamut)
 	{
-		if (DeviceId == 0x1002 /*AMD*/ || DeviceId == 0x10DE /*NVIDIA*/)
-		{
-			// needs to match GRHIHDRDisplayOutputFormat chosen in FD3D12DynamicRHI::Init
+		// needs to match GRHIHDRDisplayOutputFormat chosen in FD3D12DynamicRHI::Init
 #if WITH_EDITOR
 		// ScRGB, 1000 or 2000 nits
-			OutputDevice = (DisplayNitLevel == 1000) ? EDisplayOutputFormat::HDR_ACES_1000nit_ScRGB : EDisplayOutputFormat::HDR_ACES_2000nit_ScRGB;
-			// Rec709
-			ColorGamut = EDisplayColorGamut::sRGB_D65;
+		OutputDevice = (DisplayNitLevel == 1000) ? EDisplayOutputFormat::HDR_ACES_1000nit_ScRGB : EDisplayOutputFormat::HDR_ACES_2000nit_ScRGB;
+		// Rec709
+		ColorGamut = EDisplayColorGamut::sRGB_D65;
 #else
 		// ST-2084, 1000 or 2000 nits
-			OutputDevice = (DisplayNitLevel == 1000) ? EDisplayOutputFormat::HDR_ACES_1000nit_ST2084 : EDisplayOutputFormat::HDR_ACES_2000nit_ST2084;
-			// Rec2020
-			ColorGamut = EDisplayColorGamut::Rec2020_D65;
+		OutputDevice = (DisplayNitLevel == 1000) ? EDisplayOutputFormat::HDR_ACES_1000nit_ST2084 : EDisplayOutputFormat::HDR_ACES_2000nit_ST2084;
+		// Rec2020
+		ColorGamut = EDisplayColorGamut::Rec2020_D65;
 #endif
-		}
 	}
 
 	/**
diff --git a/Engine/Source/Runtime/D3D12RHI/Private/Windows/WindowsD3D12Viewport.cpp b/Engine/Source/Runtime/D3D12RHI/Private/Windows/WindowsD3D12Viewport.cpp
--- a/Engine/Source/Runtime/D3D12RHI/Private/Windows/WindowsD3D12Viewport.cpp
+++ b/Engine/Source/Runtime/D3D12RHI/Private/Windows/WindowsD3D12Viewport.cpp
@@ -504,7 +504,6 @@ void FD3D12Viewport::EnsureColorSpace(EDisplayColorGamut DisplayGamut, EDisplayO
 
 	if (ColorSpace != NewColorSpace)
 	{
-		check(CurrentOutputSupportsHDR());
 		uint32 ColorSpaceSupport = 0;
 		HRESULT hr = SwapChain4->CheckColorSpaceSupport(NewColorSpace, &ColorSpaceSupport);
 		FString NewColorSpaceName = GetDXGIColorSpaceString(NewColorSpace);
diff --git a/Engine/Source/Runtime/Engine/Private/GameUserSettings.cpp b/Engine/Source/Runtime/Engine/Private/GameUserSettings.cpp
--- a/Engine/Source/Runtime/Engine/Private/GameUserSettings.cpp
+++ b/Engine/Source/Runtime/Engine/Private/GameUserSettings.cpp
@@ -945,7 +945,7 @@ void UGameUserSettings::EnableHDRDisplayOutputInternal(bool bEnable, int32 Displ
 		if (bEnable)
 		{
 #if PLATFORM_WINDOWS
-			if (IsRHIDeviceNVIDIA() || IsRHIDeviceAMD())
+			if (GRHIHDRNeedsVendorExtensions)
 			{
 				// Force exclusive fullscreen
 				SetPreferredFullscreenMode(0);
diff --git a/Engine/Source/Runtime/RHI/Public/RHIGlobals.h b/Engine/Source/Runtime/RHI/Public/RHIGlobals.h
--- a/Engine/Source/Runtime/RHI/Public/RHIGlobals.h
+++ b/Engine/Source/Runtime/RHI/Public/RHIGlobals.h
@@ -633,6 +633,9 @@ struct FRHIGlobals
 
 	/** True if the RHI supports shaders with barycentrics */
 	bool SupportsBarycentricsSemantic = false;
+
+	/** True if HDR requires vendor specific extensions */
+	bool HDRNeedsVendorExtensions = false;
 };
 
 extern RHI_API FRHIGlobals GRHIGlobals;
@@ -798,6 +801,7 @@ extern RHI_API FRHIGlobals GRHIGlobals;
 #define GRHISupportsReservedResources                          GRHIGlobals.ReservedResources.Supported UE_DEPRECATED_MACRO(5.4, "GRHISupportsReservedResources has been deprecated - please use GRHIGlobals.ReservedResources.Supported instead.")
 #define GShaderPlatformForFeatureLevel                         GRHIGlobals.ShaderPlatformForFeatureLevel
 #define GRHIIsDebugLayerEnabled                                GRHIGlobals.IsDebugLayerEnabled
+#define GRHIHDRNeedsVendorExtensions						   GRHIGlobals.HDRNeedsVendorExtensions
 
 // Utility Getters
 
diff --git a/Engine/Source/Runtime/RenderCore/Private/RenderCore.cpp b/Engine/Source/Runtime/RenderCore/Private/RenderCore.cpp
--- a/Engine/Source/Runtime/RenderCore/Private/RenderCore.cpp
+++ b/Engine/Source/Runtime/RenderCore/Private/RenderCore.cpp
@@ -524,6 +524,20 @@ bool HdrHasWindowParamsFromCVars(void* OSWindow, FHDRMetaData& HDRMetaData)
 	return false;
 }
 
+static void HDRGetDeviceAndColorGamut(uint32 DeviceId, uint32 DisplayNitLevel, EDisplayOutputFormat& OutDisplayOutputFormat, EDisplayColorGamut& OutDisplayColorGamut)
+{
+	if (GRHIHDRNeedsVendorExtensions)
+	{
+		// all our implementations of HDR with vendor extensions happen with FP16 / ScRGB. See FD3D11DynamicRHI::EnableHDR / SetHDRMonitorMode*
+		OutDisplayOutputFormat = EDisplayOutputFormat::HDR_ACES_1000nit_ScRGB;
+		OutDisplayColorGamut = EDisplayColorGamut::sRGB_D65;
+	}
+	else
+	{
+		FPlatformMisc::ChooseHDRDeviceAndColorGamut(DeviceId, DisplayNitLevel, OutDisplayOutputFormat, OutDisplayColorGamut);
+	}
+}
+
 RENDERCORE_API void HDRGetMetaData(EDisplayOutputFormat& OutDisplayOutputFormat, EDisplayColorGamut& OutDisplayColorGamut, bool& OutbHDRSupported, 
 								   const FVector2D& WindowTopLeft, const FVector2D& WindowBottomRight, void* OSWindow)
 {
@@ -575,8 +589,8 @@ RENDERCORE_API void HDRGetMetaData(EDisplayOutputFormat& OutDisplayOutputFormat,
 
 	if (OutbHDRSupported)
 	{
-		FPlatformMisc::ChooseHDRDeviceAndColorGamut(GRHIVendorId, CVarHDRDisplayMaxLuminance.GetValueOnAnyThread(), OutDisplayOutputFormat, OutDisplayColorGamut);
-}
+		HDRGetDeviceAndColorGamut(GRHIVendorId, CVarHDRDisplayMaxLuminance.GetValueOnAnyThread(), OutDisplayOutputFormat, OutDisplayColorGamut);
+	}
 
 }
 
@@ -609,7 +623,7 @@ RENDERCORE_API void HDRConfigureCVars(bool bIsHDREnabled, uint32 DisplayNits, bo
 	// If we are turning it off, we'll reset back to 0/0
 	if (bIsHDREnabled)
 	{
-		FPlatformMisc::ChooseHDRDeviceAndColorGamut(GRHIVendorId, DisplayNits, OutputDevice, ColorGamut);
+		HDRGetDeviceAndColorGamut(GRHIVendorId, DisplayNits, OutputDevice, ColorGamut);
 	}
 
 	//CVarHDRDisplayMaxLuminance is ECVF_SetByCode as it's only a mean of communicating the information from UGameUserSettings to the rest of the engine
diff --git a/Engine/Source/Runtime/Windows/D3D11RHI/Private/Windows/WindowsD3D11Device.cpp b/Engine/Source/Runtime/Windows/D3D11RHI/Private/Windows/WindowsD3D11Device.cpp
--- a/Engine/Source/Runtime/Windows/D3D11RHI/Private/Windows/WindowsD3D11Device.cpp
+++ b/Engine/Source/Runtime/Windows/D3D11RHI/Private/Windows/WindowsD3D11Device.cpp
@@ -2310,6 +2310,10 @@ void FD3D11DynamicRHI::InitD3DDevice()
 
 		SetupAfterDeviceCreation();
 		GRHISupportsHDROutput = SetupDisplayHDRMetaData();
+#if !WITH_EDITOR
+		// cooked game D3D11 still needs to rely on vendor extensions to trigger HDR, which will then require exclusive fullscreen state / special RT formats
+		GRHIHDRNeedsVendorExtensions = true;
+#endif
 
 		// Add device overclock state to crash context
 		const bool bIsGPUOverclocked = IsDeviceOverclocked();
