Removed FRHICommandList parameter from FRayTracingGeometryManager::RequestBuildAccelerationStructure()
- Also made it possible to update some other functions to only require FRHICommandListBase.

diff --git a/Engine/Source/Runtime/Engine/Private/Rendering/RayTracingGeometryManager.cpp b/Engine/Source/Runtime/Engine/Private/Rendering/RayTracingGeometryManager.cpp
--- a/Engine/Source/Runtime/Engine/Private/Rendering/RayTracingGeometryManager.cpp
+++ b/Engine/Source/Runtime/Engine/Private/Rendering/RayTracingGeometryManager.cpp
@@ -57,7 +57,7 @@ static float GetInitialBuildPriority(ERTAccelerationStructureBuildPriority InBui
 	}
 }
 
-FRayTracingGeometryManager::BuildRequestIndex FRayTracingGeometryManager::RequestBuildAccelerationStructure(FRHICommandList& RHICmdList, FRayTracingGeometry* InGeometry, ERTAccelerationStructureBuildPriority InPriority, EAccelerationStructureBuildMode InBuildMode)
+FRayTracingGeometryManager::BuildRequestIndex FRayTracingGeometryManager::RequestBuildAccelerationStructure(FRayTracingGeometry* InGeometry, ERTAccelerationStructureBuildPriority InPriority, EAccelerationStructureBuildMode InBuildMode)
 {
 	BuildRequest Request;
 	Request.BuildPriority = GetInitialBuildPriority(InPriority);
diff --git a/Engine/Source/Runtime/Engine/Private/SkeletalMeshDeformerHelpers.cpp b/Engine/Source/Runtime/Engine/Private/SkeletalMeshDeformerHelpers.cpp
--- a/Engine/Source/Runtime/Engine/Private/SkeletalMeshDeformerHelpers.cpp
+++ b/Engine/Source/Runtime/Engine/Private/SkeletalMeshDeformerHelpers.cpp
@@ -300,6 +300,6 @@ void FSkeletalMeshDeformerHelpers::ResetVertexFactoryBufferOverrides(FSkeletalMe
 
 	TArray<FBufferRHIRef> VertexBuffers;
 	VertexBuffers.Init(VertexBuffer, NumSections);
-	MeshObjectGPU->UpdateRayTracingGeometry(FRHICommandList::Get(RHICmdList), LODModel, LODIndex, VertexBuffers);
+	MeshObjectGPU->UpdateRayTracingGeometry(RHICmdList, LODModel, LODIndex, VertexBuffers);
 #endif // RHI_RAYTRACING
 }
diff --git a/Engine/Source/Runtime/Engine/Private/SkeletalRenderGPUSkin.cpp b/Engine/Source/Runtime/Engine/Private/SkeletalRenderGPUSkin.cpp
--- a/Engine/Source/Runtime/Engine/Private/SkeletalRenderGPUSkin.cpp
+++ b/Engine/Source/Runtime/Engine/Private/SkeletalRenderGPUSkin.cpp
@@ -753,7 +753,7 @@ void FSkeletalMeshObjectGPUSkin::ProcessUpdatedDynamicData(EGPUSkinCacheEntryMod
 
 #if RHI_RAYTRACING
 
-void FSkeletalMeshObjectGPUSkin::UpdateRayTracingGeometry(FRHICommandList& RHICmdList, FSkeletalMeshLODRenderData& LODModel, uint32 LODIndex, TArray<FBufferRHIRef>& VertexBuffers)
+void FSkeletalMeshObjectGPUSkin::UpdateRayTracingGeometry(FRHICommandListBase& RHICmdList, FSkeletalMeshLODRenderData& LODModel, uint32 LODIndex, TArray<FBufferRHIRef>& VertexBuffers)
 {
 	if (IsRayTracingEnabled() && bSupportRayTracing)
 	{
diff --git a/Engine/Source/Runtime/Engine/Private/SkeletalRenderGPUSkin.h b/Engine/Source/Runtime/Engine/Private/SkeletalRenderGPUSkin.h
--- a/Engine/Source/Runtime/Engine/Private/SkeletalRenderGPUSkin.h
+++ b/Engine/Source/Runtime/Engine/Private/SkeletalRenderGPUSkin.h
@@ -395,7 +395,7 @@ public:
 	 * VSinCS path is still required for world position offset materials but this can still use 
 	 * the updated vertex buffers from here with a passthrough vertex factory.
 	 */
-	ENGINE_API void UpdateRayTracingGeometry(FRHICommandList& RHICmdList, FSkeletalMeshLODRenderData& LODModel, uint32 LODIndex, TArray<FBufferRHIRef>& VertexBuffers);
+	ENGINE_API void UpdateRayTracingGeometry(FRHICommandListBase& RHICmdList, FSkeletalMeshLODRenderData& LODModel, uint32 LODIndex, TArray<FBufferRHIRef>& VertexBuffers);
 
 #endif // RHI_RAYTRACING
 
diff --git a/Engine/Source/Runtime/Engine/Private/Streaming/StaticMeshUpdate.cpp b/Engine/Source/Runtime/Engine/Private/Streaming/StaticMeshUpdate.cpp
--- a/Engine/Source/Runtime/Engine/Private/Streaming/StaticMeshUpdate.cpp
+++ b/Engine/Source/Runtime/Engine/Private/Streaming/StaticMeshUpdate.cpp
@@ -368,7 +368,7 @@ void FStaticMeshStreamIn::DoFinishUpdate(const FContext& Context)
 					// Under very rare circumstances that we switch ray tracing on/off right in the middle of streaming RayTracingGeometryRHI might not be valid.
 					if (IsRayTracingEnabled() && ensure(LODResource.RayTracingGeometry.RayTracingGeometryRHI.IsValid()))
 					{
-						LODResource.RayTracingGeometry.RequestBuildIfNeeded(FRHICommandListImmediate::Get(), ERTAccelerationStructureBuildPriority::Normal);
+						LODResource.RayTracingGeometry.RequestBuildIfNeeded(ERTAccelerationStructureBuildPriority::Normal);
 					}
 				}
 			}
diff --git a/Engine/Source/Runtime/Engine/Public/Rendering/RayTracingGeometryManager.h b/Engine/Source/Runtime/Engine/Public/Rendering/RayTracingGeometryManager.h
--- a/Engine/Source/Runtime/Engine/Public/Rendering/RayTracingGeometryManager.h
+++ b/Engine/Source/Runtime/Engine/Public/Rendering/RayTracingGeometryManager.h
@@ -17,7 +17,7 @@ public:
 
 	ENGINE_API virtual ~FRayTracingGeometryManager();
 
-	ENGINE_API virtual BuildRequestIndex RequestBuildAccelerationStructure(FRHICommandList& RHICmdList, FRayTracingGeometry* InGeometry, ERTAccelerationStructureBuildPriority InPriority, EAccelerationStructureBuildMode InBuildMode) override;
+	ENGINE_API virtual BuildRequestIndex RequestBuildAccelerationStructure(FRayTracingGeometry* InGeometry, ERTAccelerationStructureBuildPriority InPriority, EAccelerationStructureBuildMode InBuildMode) override;
 
 	ENGINE_API virtual void RemoveBuildRequest(BuildRequestIndex InRequestIndex) override;
 	ENGINE_API virtual void BoostPriority(BuildRequestIndex InRequestIndex, float InBoostValue) override;
diff --git a/Engine/Source/Runtime/RenderCore/Private/RayTracingGeometry.cpp b/Engine/Source/Runtime/RenderCore/Private/RayTracingGeometry.cpp
--- a/Engine/Source/Runtime/RenderCore/Private/RayTracingGeometry.cpp
+++ b/Engine/Source/Runtime/RenderCore/Private/RayTracingGeometry.cpp
@@ -83,22 +83,17 @@ void FRayTracingGeometry::CreateRayTracingGeometryFromCPUData(TResourceArray<uin
 	SetRequiresBuild(Initializer.OfflineData == nullptr || RayTracingGeometryRHI->IsCompressed());
 }
 
-void FRayTracingGeometry::RequestBuildIfNeeded(FRHICommandList& RHICmdList, ERTAccelerationStructureBuildPriority InBuildPriority)
+void FRayTracingGeometry::RequestBuildIfNeeded(ERTAccelerationStructureBuildPriority InBuildPriority)
 {
 	RayTracingGeometryRHI->SetInitializer(Initializer);
 
 	if (GetRequiresBuild())
 	{
-		RayTracingBuildRequestIndex = GRayTracingGeometryManager->RequestBuildAccelerationStructure(RHICmdList, this, InBuildPriority);
+		RayTracingBuildRequestIndex = GRayTracingGeometryManager->RequestBuildAccelerationStructure(this, InBuildPriority);
 		SetRequiresBuild(false);
 	}
 }
 
-void FRayTracingGeometry::RequestBuildIfNeeded(ERTAccelerationStructureBuildPriority InBuildPriority)
-{
-	RequestBuildIfNeeded(FRHICommandListImmediate::Get(), InBuildPriority);
-}
-
 void FRayTracingGeometry::MakeResident(FRHICommandList& RHICmdList)
 {
 	check(EnumHasAllFlags(GeometryState, EGeometryStateFlags::Evicted) && RayTracingGeometryRHI == nullptr);
@@ -133,7 +128,7 @@ void FRayTracingGeometry::MakeResident(FRHICommandList& RHICmdList)
 			// We need to do it before we build the current geometry (also on RHI thread).
 		}
 
-		RequestBuildIfNeeded(RHICmdList, ERTAccelerationStructureBuildPriority::Normal);
+		RequestBuildIfNeeded(ERTAccelerationStructureBuildPriority::Normal);
 	}
 	else
 	{
@@ -156,7 +151,7 @@ void FRayTracingGeometry::Evict()
 	EnumAddFlags(GeometryState, EGeometryStateFlags::Evicted);
 }
 
-void FRayTracingGeometry::CreateRayTracingGeometry(FRHICommandList& RHICmdList, ERTAccelerationStructureBuildPriority InBuildPriority)
+void FRayTracingGeometry::CreateRayTracingGeometry(FRHICommandListBase& RHICmdList, ERTAccelerationStructureBuildPriority InBuildPriority)
 {
 	// Release previous RHI object if any
 	ReleaseRHI();
@@ -208,7 +203,7 @@ void FRayTracingGeometry::CreateRayTracingGeometry(FRHICommandList& RHICmdList,
 			{
 				if (RayTracingGeometryRHI)
 				{
-					RayTracingBuildRequestIndex = GRayTracingGeometryManager->RequestBuildAccelerationStructure(RHICmdList, this, InBuildPriority);
+					RayTracingBuildRequestIndex = GRayTracingGeometryManager->RequestBuildAccelerationStructure(this, InBuildPriority);
 				}
 				SetRequiresBuild(false);
 			}
@@ -221,7 +216,7 @@ void FRayTracingGeometry::CreateRayTracingGeometry(FRHICommandList& RHICmdList,
 		{
 			if (RayTracingGeometryRHI && RayTracingGeometryRHI->IsCompressed())
 			{
-				RayTracingBuildRequestIndex = GRayTracingGeometryManager->RequestBuildAccelerationStructure(RHICmdList, this, InBuildPriority);
+				RayTracingBuildRequestIndex = GRayTracingGeometryManager->RequestBuildAccelerationStructure(this, InBuildPriority);
 			}
 
 			SetRequiresBuild(false);
@@ -276,7 +271,7 @@ void FRayTracingGeometry::InitRHI(FRHICommandListBase& RHICmdList)
 	ERTAccelerationStructureBuildPriority BuildPriority = Initializer.Type != ERayTracingGeometryInitializerType::Rendering
 		? ERTAccelerationStructureBuildPriority::Skip
 		: ERTAccelerationStructureBuildPriority::Normal;
-	CreateRayTracingGeometry(FRHICommandList::Get(RHICmdList), BuildPriority);
+	CreateRayTracingGeometry(RHICmdList, BuildPriority);
 }
 
 void FRayTracingGeometry::ReleaseRHI()
diff --git a/Engine/Source/Runtime/RenderCore/Public/RayTracingGeometry.h b/Engine/Source/Runtime/RenderCore/Public/RayTracingGeometry.h
--- a/Engine/Source/Runtime/RenderCore/Public/RayTracingGeometry.h
+++ b/Engine/Source/Runtime/RenderCore/Public/RayTracingGeometry.h
@@ -109,15 +109,12 @@ public:
 	UE_DEPRECATED(5.4, "Use FStaticMeshStreamIn::FIntermediateRayTracingGeometry instead.")
 	RENDERCORE_API void CreateRayTracingGeometryFromCPUData(TResourceArray<uint8>& OfflineData);
 
-	RENDERCORE_API void RequestBuildIfNeeded(FRHICommandList& RHICmdList, ERTAccelerationStructureBuildPriority InBuildPriority);
-
-	UE_DEPRECATED(5.4, "RequestBuildIfNeeded now requires a command list.")
 	RENDERCORE_API void RequestBuildIfNeeded(ERTAccelerationStructureBuildPriority InBuildPriority);
 
 	UE_DEPRECATED(5.4, "InitRHIForDynamicRayTracing now requires a command list and was renamed to MakeResident().")
 	RENDERCORE_API void InitRHIForDynamicRayTracing();
 
-	RENDERCORE_API void CreateRayTracingGeometry(FRHICommandList& RHICmdList, ERTAccelerationStructureBuildPriority InBuildPriority);
+	RENDERCORE_API void CreateRayTracingGeometry(FRHICommandListBase& RHICmdList, ERTAccelerationStructureBuildPriority InBuildPriority);
 
 	UE_DEPRECATED(5.4, "CreateRayTracingGeometry now requires a command list.")
 	RENDERCORE_API void CreateRayTracingGeometry(ERTAccelerationStructureBuildPriority InBuildPriority);
diff --git a/Engine/Source/Runtime/RenderCore/Public/RayTracingGeometryManagerInterface.h b/Engine/Source/Runtime/RenderCore/Public/RayTracingGeometryManagerInterface.h
--- a/Engine/Source/Runtime/RenderCore/Public/RayTracingGeometryManagerInterface.h
+++ b/Engine/Source/Runtime/RenderCore/Public/RayTracingGeometryManagerInterface.h
@@ -23,23 +23,11 @@ public:
 
 	virtual ~IRayTracingGeometryManager() = default;
 
-	RENDERCORE_API virtual BuildRequestIndex RequestBuildAccelerationStructure(FRHICommandList& RHICmdList, FRayTracingGeometry* InGeometry, ERTAccelerationStructureBuildPriority InPriority, EAccelerationStructureBuildMode InBuildMode) = 0;
+	RENDERCORE_API virtual BuildRequestIndex RequestBuildAccelerationStructure(FRayTracingGeometry* InGeometry, ERTAccelerationStructureBuildPriority InPriority, EAccelerationStructureBuildMode InBuildMode) = 0;
 
-	BuildRequestIndex RequestBuildAccelerationStructure(FRHICommandList& RHICmdList, FRayTracingGeometry* InGeometry, ERTAccelerationStructureBuildPriority InPriority)
-	{
-		return RequestBuildAccelerationStructure(RHICmdList, InGeometry, InPriority, EAccelerationStructureBuildMode::Build);
-	}
-
-	UE_DEPRECATED(5.4, "RequestBuildAccelerationStructure requires a command list.")
 	BuildRequestIndex RequestBuildAccelerationStructure(FRayTracingGeometry* InGeometry, ERTAccelerationStructureBuildPriority InPriority)
 	{
-		return RequestBuildAccelerationStructure(FRHICommandListImmediate::Get(), InGeometry, InPriority, EAccelerationStructureBuildMode::Build);
-	}
-
-	UE_DEPRECATED(5.4, "RequestBuildAccelerationStructure requires a command list.")
-	BuildRequestIndex RequestBuildAccelerationStructure(FRayTracingGeometry* InGeometry, ERTAccelerationStructureBuildPriority InPriority, EAccelerationStructureBuildMode InBuildMode)
-	{
-		return RequestBuildAccelerationStructure(FRHICommandListImmediate::Get(), InGeometry, InPriority, InBuildMode);
+		return RequestBuildAccelerationStructure(InGeometry, InPriority, EAccelerationStructureBuildMode::Build);
 	}
 
 	RENDERCORE_API virtual void RemoveBuildRequest(BuildRequestIndex InRequestIndex) = 0;
