Fixed case where FReadOnlyCVARCache was used too early before RHI init

diff --git a/Engine/Source/Runtime/RenderCore/Private/ReadOnlyCVARCache.cpp b/Engine/Source/Runtime/RenderCore/Private/ReadOnlyCVARCache.cpp
--- a/Engine/Source/Runtime/RenderCore/Private/ReadOnlyCVARCache.cpp
+++ b/Engine/Source/Runtime/RenderCore/Private/ReadOnlyCVARCache.cpp
@@ -9,19 +9,19 @@
 #include "Interfaces/ITargetPlatformManagerModule.h"
 #endif
 
-static int32 MobileEarlyZPassIniValue(const FStaticShaderPlatform Platform)
+static int32 MobileEarlyZPassIniValue(EShaderPlatform Platform)
 {
 	static FShaderPlatformCachedIniValue<int32> CVar(TEXT("r.Mobile.EarlyZPass"));
 	return CVar.Get(Platform);
 }
 
-static int32 MobileForwardLocalLightsIniValue(const FStaticShaderPlatform Platform)
+static int32 MobileForwardLocalLightsIniValue(EShaderPlatform Platform)
 {
 	static FShaderPlatformCachedIniValue<int32> CVar(TEXT("r.Mobile.Forward.EnableLocalLights"));
 	return CVar.Get(Platform);
 }
 
-static bool MobileDeferredShadingIniValue(const FStaticShaderPlatform Platform)
+static bool MobileDeferredShadingIniValue(EShaderPlatform Platform)
 {
 	static FShaderPlatformCachedIniValue<bool> MobileShadingPathIniValue(TEXT("r.Mobile.ShadingPath"));
 	static TConsoleVariableData<int32>* MobileAllowDeferredShadingOpenGL = IConsoleManager::Get().FindTConsoleVariableDataInt(TEXT("r.Mobile.AllowDeferredShadingOpenGL"));
@@ -30,7 +30,7 @@ static bool MobileDeferredShadingIniValue(const FStaticShaderPlatform Platform)
 	return MobileShadingPathIniValue.Get(Platform) && bSupportedPlatform;
 }
 
-static bool MobileEnableMovableSpotlightsShadowIniValue(const FStaticShaderPlatform Platform)
+static bool MobileEnableMovableSpotlightsShadowIniValue(EShaderPlatform Platform)
 {
 	static FShaderPlatformCachedIniValue<bool> CVar(TEXT("r.Mobile.EnableMovableSpotlightsShadow"));
 	return CVar.Get(Platform);
@@ -172,7 +172,7 @@ bool FReadOnlyCVARCache::MobileAllowDistanceFieldShadows()
 	return bMobileAllowDistanceFieldShadows;
 }
 
-int32 FReadOnlyCVARCache::MobileEarlyZPass(const FStaticShaderPlatform Platform)
+int32 FReadOnlyCVARCache::MobileEarlyZPass(EShaderPlatform Platform)
 {
 #if WITH_EDITOR
 	return MobileEarlyZPassIniValue(Platform);
@@ -181,7 +181,7 @@ int32 FReadOnlyCVARCache::MobileEarlyZPass(const FStaticShaderPlatform Platform)
 #endif
 }
 
-int32 FReadOnlyCVARCache::MobileForwardLocalLights(const FStaticShaderPlatform Platform)
+int32 FReadOnlyCVARCache::MobileForwardLocalLights(EShaderPlatform Platform)
 {
 #if WITH_EDITOR
 	return MobileForwardLocalLightsIniValue(Platform);
@@ -190,7 +190,7 @@ int32 FReadOnlyCVARCache::MobileForwardLocalLights(const FStaticShaderPlatform P
 #endif
 }
 
-bool FReadOnlyCVARCache::MobileDeferredShading(const FStaticShaderPlatform Platform)
+bool FReadOnlyCVARCache::MobileDeferredShading(EShaderPlatform Platform)
 {
 #if WITH_EDITOR
 	return MobileDeferredShadingIniValue(Platform);
@@ -199,7 +199,7 @@ bool FReadOnlyCVARCache::MobileDeferredShading(const FStaticShaderPlatform Platf
 #endif
 }
 
-bool FReadOnlyCVARCache::MobileEnableMovableSpotlightsShadow(const FStaticShaderPlatform Platform)
+bool FReadOnlyCVARCache::MobileEnableMovableSpotlightsShadow(EShaderPlatform Platform)
 {
 #if WITH_EDITOR
 	return MobileEnableMovableSpotlightsShadowIniValue(Platform);
diff --git a/Engine/Source/Runtime/RenderCore/Public/ReadOnlyCVARCache.h b/Engine/Source/Runtime/RenderCore/Public/ReadOnlyCVARCache.h
--- a/Engine/Source/Runtime/RenderCore/Public/ReadOnlyCVARCache.h
+++ b/Engine/Source/Runtime/RenderCore/Public/ReadOnlyCVARCache.h
@@ -20,8 +20,8 @@ struct FReadOnlyCVARCache
 	// Mobile specific
 	static bool MobileHDR();
 	static bool MobileSupportsGPUScene();
-	static bool MobileDeferredShading(const FStaticShaderPlatform Platform);
-	static bool MobileEnableMovableSpotlightsShadow(const FStaticShaderPlatform Platform);
+	static bool MobileDeferredShading(EShaderPlatform Platform);
+	static bool MobileEnableMovableSpotlightsShadow(EShaderPlatform Platform);
 
 	RENDERCORE_API static bool MobileAllowMovableDirectionalLights();
 	RENDERCORE_API static bool MobileAllowDistanceFieldShadows();
@@ -29,8 +29,8 @@ struct FReadOnlyCVARCache
 	RENDERCORE_API static bool MobileEnableMovableLightCSMShaderCulling();
 	RENDERCORE_API static int32 MobileSkyLightPermutation();
 	RENDERCORE_API static bool MobileEnableNoPrecomputedLightingCSMShader();
-	RENDERCORE_API static int32 MobileEarlyZPass(const FStaticShaderPlatform Platform);
-	RENDERCORE_API static int32 MobileForwardLocalLights(const FStaticShaderPlatform Platform);
+	RENDERCORE_API static int32 MobileEarlyZPass(EShaderPlatform Platform);
+	RENDERCORE_API static int32 MobileForwardLocalLights(EShaderPlatform Platform);
 
 private:
 	static bool bInitialized;
diff --git a/Engine/Source/Runtime/Renderer/Private/RendererScene.cpp b/Engine/Source/Runtime/Renderer/Private/RendererScene.cpp
--- a/Engine/Source/Runtime/Renderer/Private/RendererScene.cpp
+++ b/Engine/Source/Runtime/Renderer/Private/RendererScene.cpp
@@ -1358,20 +1358,17 @@ static void UpdateEarlyZPassModeCVarSinkFunction()
 	static int32 CachedMSAACount = CVarMSAACount->GetValueOnGameThread();
 	static int32 CachedEarlyZPass = CVarEarlyZPass.GetValueOnGameThread();
 	static int32 CachedBasePassWriteDepthEvenWithFullPrepass = CVarBasePassWriteDepthEvenWithFullPrepass.GetValueOnGameThread();
-	static int32 CachedMobileEarlyZPass = FReadOnlyCVARCache::MobileEarlyZPass(GMaxRHIShaderPlatform);
 
 	const int32 AntiAliasingMethod = CVarAntiAliasingMethod->GetValueOnGameThread();
 	const int32 MSAACount = CVarMSAACount->GetValueOnGameThread();
 	const int32 EarlyZPass = CVarEarlyZPass.GetValueOnGameThread();
 	const int32 BasePassWriteDepthEvenWithFullPrepass = CVarBasePassWriteDepthEvenWithFullPrepass.GetValueOnGameThread();
-	const int32 MobileEarlyZPass = FReadOnlyCVARCache::MobileEarlyZPass(GMaxRHIShaderPlatform);
 
 	// Switching between MSAA and another AA in forward shading mode requires EarlyZPassMode to update.
 	if (AntiAliasingMethod != CachedAntiAliasingMethod
 		|| MSAACount != CachedMSAACount
 		|| EarlyZPass != CachedEarlyZPass
-		|| BasePassWriteDepthEvenWithFullPrepass != CachedBasePassWriteDepthEvenWithFullPrepass
-		|| MobileEarlyZPass != CachedMobileEarlyZPass)
+		|| BasePassWriteDepthEvenWithFullPrepass != CachedBasePassWriteDepthEvenWithFullPrepass)
 	{
 		for (TObjectIterator<UWorld> It; It; ++It)
 		{
@@ -1387,7 +1384,6 @@ static void UpdateEarlyZPassModeCVarSinkFunction()
 		CachedMSAACount = MSAACount;
 		CachedEarlyZPass = EarlyZPass;
 		CachedBasePassWriteDepthEvenWithFullPrepass = BasePassWriteDepthEvenWithFullPrepass;
-		CachedMobileEarlyZPass = MobileEarlyZPass;
 	}
 }
 
