Added lock to render target pool to protect FindFreeElement.

diff --git a/Engine/Source/Runtime/RenderCore/Private/RenderTargetPool.cpp b/Engine/Source/Runtime/RenderCore/Private/RenderTargetPool.cpp
--- a/Engine/Source/Runtime/RenderCore/Private/RenderTargetPool.cpp
+++ b/Engine/Source/Runtime/RenderCore/Private/RenderTargetPool.cpp
@@ -82,6 +82,8 @@ TRefCountPtr<IPooledRenderTarget> FRenderTargetPool::FindFreeElement(FRHICommand
 
 	const uint32 DescHash = GetTypeHash(Desc);
 
+	UE::TScopeLock Lock(Mutex);
+
 	for (uint32 Index = 0, Num = (uint32)PooledRenderTargets.Num(); Index < Num; ++Index)
 	{
 		if (PooledRenderTargetHashes[Index] == DescHash)
@@ -175,10 +177,13 @@ bool FRenderTargetPool::FindFreeElement(FRHICommandListBase& RHICmdList, const F
 	// Querying a render target that have no mip levels makes no sens.
 	check(Desc.NumMips > 0);
 
+	FPooledRenderTarget* Current = nullptr;
+	bool bFreeCurrent = false;
+
 	// if we can keep the current one, do that
 	if (Out)
 	{
-		FPooledRenderTarget* Current = (FPooledRenderTarget*)Out.GetReference();
+		Current = (FPooledRenderTarget*)Out.GetReference();
 
 		if (Translate(Out->GetDesc()) == Desc)
 		{
@@ -197,21 +202,22 @@ bool FRenderTargetPool::FindFreeElement(FRHICommandListBase& RHICmdList, const F
 		{
 			// release old reference, it might free a RT we can use
 			Out = 0;
-
-			if (Current->IsFree())
-			{
-				UE::TScopeLock Lock(Mutex);
-
-				AllocationLevelInKB -= ComputeSizeInKB(*Current);
-				TRACE_COUNTER_SUBTRACT(RenderTargetPoolCount, 1);
-				TRACE_COUNTER_SET(RenderTargetPoolSize, (int64)AllocationLevelInKB * 1024);
-				int32 Index = FindIndex(Current);
-				check(Index >= 0);
-				FreeElementAtIndex(Index);
-			}
+			bFreeCurrent = Current->IsFree();
 		}
 	}
 
+	UE::TScopeLock Lock(Mutex);
+
+	if (bFreeCurrent)
+	{
+		AllocationLevelInKB -= ComputeSizeInKB(*Current);
+		TRACE_COUNTER_SUBTRACT(RenderTargetPoolCount, 1);
+		TRACE_COUNTER_SET(RenderTargetPoolSize, (int64)AllocationLevelInKB * 1024);
+		int32 Index = FindIndex(Current);
+		check(Index >= 0);
+		FreeElementAtIndex(Index);
+	}
+
 	Out = FindFreeElement(RHICmdList, Desc, Name);
 	return false;
 }
