Added ability to turn some of FReadOnlyCVARCache properties into compile time constants

diff --git a/Engine/Source/Runtime/RenderCore/Private/ReadOnlyCVARCache.cpp b/Engine/Source/Runtime/RenderCore/Private/ReadOnlyCVARCache.cpp
--- a/Engine/Source/Runtime/RenderCore/Private/ReadOnlyCVARCache.cpp
+++ b/Engine/Source/Runtime/RenderCore/Private/ReadOnlyCVARCache.cpp
@@ -9,19 +9,41 @@
 #include "Interfaces/ITargetPlatformManagerModule.h"
 #endif
 
-static int32 MobileEarlyZPassIniValue(EShaderPlatform Platform)
+bool FReadOnlyCVARCache::bInitialized = false;
+
+bool FReadOnlyCVARCache::bEnablePointLightShadows = true;
+bool FReadOnlyCVARCache::bEnableStationarySkylight = true;
+bool FReadOnlyCVARCache::bEnableLowQualityLightmaps = true;
+bool FReadOnlyCVARCache::bAllowStaticLighting = true;
+bool FReadOnlyCVARCache::bSupportSkyAtmosphere = true;
+
+// Mobile specific
+bool FReadOnlyCVARCache::bMobileHDR = true;
+bool FReadOnlyCVARCache::bMobileAllowMovableDirectionalLights = true;
+bool FReadOnlyCVARCache::bMobileAllowDistanceFieldShadows = true;
+bool FReadOnlyCVARCache::bMobileEnableStaticAndCSMShadowReceivers = true;
+bool FReadOnlyCVARCache::bMobileEnableMovableLightCSMShaderCulling = true;
+bool FReadOnlyCVARCache::bMobileSupportsGPUScene = false;
+int32 FReadOnlyCVARCache::MobileSkyLightPermutationValue = 0;
+int32 FReadOnlyCVARCache::MobileEarlyZPassValue = 0;
+int32 FReadOnlyCVARCache::MobileForwardLocalLightsValue = 1;
+bool FReadOnlyCVARCache::bMobileEnableNoPrecomputedLightingCSMShader = false;
+bool FReadOnlyCVARCache::bMobileDeferredShadingValue = false;
+bool FReadOnlyCVARCache::bMobileEnableMovableSpotlightsShadowValue = false;
+
+int32 FReadOnlyCVARCache::MobileEarlyZPassIniValue(EShaderPlatform Platform)
 {
 	static FShaderPlatformCachedIniValue<int32> CVar(TEXT("r.Mobile.EarlyZPass"));
 	return CVar.Get(Platform);
 }
 
-static int32 MobileForwardLocalLightsIniValue(EShaderPlatform Platform)
+int32 FReadOnlyCVARCache::MobileForwardLocalLightsIniValue(EShaderPlatform Platform)
 {
 	static FShaderPlatformCachedIniValue<int32> CVar(TEXT("r.Mobile.Forward.EnableLocalLights"));
 	return CVar.Get(Platform);
 }
 
-static bool MobileDeferredShadingIniValue(EShaderPlatform Platform)
+bool FReadOnlyCVARCache::MobileDeferredShadingIniValue(EShaderPlatform Platform)
 {
 	static FShaderPlatformCachedIniValue<bool> MobileShadingPathIniValue(TEXT("r.Mobile.ShadingPath"));
 	static TConsoleVariableData<int32>* MobileAllowDeferredShadingOpenGL = IConsoleManager::Get().FindTConsoleVariableDataInt(TEXT("r.Mobile.AllowDeferredShadingOpenGL"));
@@ -30,34 +52,12 @@ static bool MobileDeferredShadingIniValue(EShaderPlatform Platform)
 	return MobileShadingPathIniValue.Get(Platform) && bSupportedPlatform;
 }
 
-static bool MobileEnableMovableSpotlightsShadowIniValue(EShaderPlatform Platform)
+bool FReadOnlyCVARCache::MobileEnableMovableSpotlightsShadowIniValue(EShaderPlatform Platform)
 {
 	static FShaderPlatformCachedIniValue<bool> CVar(TEXT("r.Mobile.EnableMovableSpotlightsShadow"));
 	return CVar.Get(Platform);
 }
 
-bool FReadOnlyCVARCache::bInitialized = false;
-
-bool FReadOnlyCVARCache::bEnablePointLightShadows = true;
-bool FReadOnlyCVARCache::bEnableStationarySkylight = true;
-bool FReadOnlyCVARCache::bEnableLowQualityLightmaps = true;
-bool FReadOnlyCVARCache::bAllowStaticLighting = true;
-bool FReadOnlyCVARCache::bSupportSkyAtmosphere = true;
-
-// Mobile specific
-bool FReadOnlyCVARCache::bMobileHDR = true;
-bool FReadOnlyCVARCache::bMobileAllowMovableDirectionalLights = true;
-bool FReadOnlyCVARCache::bMobileAllowDistanceFieldShadows = true;
-bool FReadOnlyCVARCache::bMobileEnableStaticAndCSMShadowReceivers = true;
-bool FReadOnlyCVARCache::bMobileEnableMovableLightCSMShaderCulling = true;
-bool FReadOnlyCVARCache::bMobileSupportsGPUScene = false;
-int32 FReadOnlyCVARCache::MobileSkyLightPermutationValue = 0;
-int32 FReadOnlyCVARCache::MobileEarlyZPassValue = 0;
-int32 FReadOnlyCVARCache::MobileForwardLocalLightsValue = 1;
-bool FReadOnlyCVARCache::bMobileEnableNoPrecomputedLightingCSMShader = false;
-bool FReadOnlyCVARCache::bMobileDeferredShadingValue = false;
-bool FReadOnlyCVARCache::bMobileEnableMovableSpotlightsShadowValue = false;
-
 void FReadOnlyCVARCache::Initialize()
 {
 	UE_LOG(LogInit, Log, TEXT("Initializing FReadOnlyCVARCache"));
@@ -100,110 +100,26 @@ void FReadOnlyCVARCache::Initialize()
 	bMobileDeferredShadingValue = MobileDeferredShadingIniValue(GMaxRHIShaderPlatform);
 	bMobileEnableMovableSpotlightsShadowValue = MobileEnableMovableSpotlightsShadowIniValue(GMaxRHIShaderPlatform);
 	bMobileSupportsGPUScene = CVarMobileSupportGPUScene->GetValueOnAnyThread() != 0;
-	
-	bInitialized = true;
-}
-
-bool FReadOnlyCVARCache::AllowStaticLighting()
-{
-	checkSlow(bInitialized);
-	return bAllowStaticLighting;
-}
 
-bool FReadOnlyCVARCache::EnableLowQualityLightmaps()
-{
-	return bEnableLowQualityLightmaps;
-}
-
-bool FReadOnlyCVARCache::SupportSkyAtmosphere()
-{
-	return bSupportSkyAtmosphere;
-}
-
-bool FReadOnlyCVARCache::EnablePointLightShadows()
-{
-	return bEnablePointLightShadows;
-}
-
-bool FReadOnlyCVARCache::EnableStationarySkylight()
-{
-	return bEnableStationarySkylight;
-}
-
-bool FReadOnlyCVARCache::MobileHDR()
-{
-	checkSlow(bInitialized);
-	return bMobileHDR;
-}
-
-int32 FReadOnlyCVARCache::MobileSkyLightPermutation()
-{
-	return MobileSkyLightPermutationValue;
-}
-
-bool FReadOnlyCVARCache::MobileEnableNoPrecomputedLightingCSMShader()
-{
-	return bMobileEnableNoPrecomputedLightingCSMShader;
-}
-
-bool FReadOnlyCVARCache::MobileEnableStaticAndCSMShadowReceivers()
-{
-	return bMobileEnableStaticAndCSMShadowReceivers;
-}
-
-bool FReadOnlyCVARCache::MobileEnableMovableLightCSMShaderCulling()
-{
-	return bMobileEnableMovableLightCSMShaderCulling;
-}
-
-bool FReadOnlyCVARCache::MobileSupportsGPUScene()
-{
-	checkSlow(bInitialized);
-	return bMobileSupportsGPUScene;
-}
-
-bool FReadOnlyCVARCache::MobileAllowMovableDirectionalLights()
-{
-	return bMobileAllowMovableDirectionalLights;
-}
-
-bool FReadOnlyCVARCache::MobileAllowDistanceFieldShadows()
-{
-	return bMobileAllowDistanceFieldShadows;
-}
+#ifdef PROJECT_CVAR_ALLOW_STATIC_LIGHTING
+	check(!!PROJECT_CVAR_ALLOW_STATIC_LIGHTING == bAllowStaticLighting);
+#endif
 
-int32 FReadOnlyCVARCache::MobileEarlyZPass(EShaderPlatform Platform)
-{
-#if WITH_EDITOR
-	return MobileEarlyZPassIniValue(Platform);
-#else
-	return MobileEarlyZPassValue;
+#ifdef PROJECT_CVAR_MOBILE_HDR
+	check(!!PROJECT_CVAR_MOBILE_HDR == bMobileHDR);
 #endif
-}
 
-int32 FReadOnlyCVARCache::MobileForwardLocalLights(EShaderPlatform Platform)
-{
-#if WITH_EDITOR
-	return MobileForwardLocalLightsIniValue(Platform);
-#else
-	return MobileForwardLocalLightsValue;
+#ifdef PROJECT_CVAR_MOBILE_SUPPORTS_GPUSCENE
+	check(!!PROJECT_CVAR_MOBILE_SUPPORTS_GPUSCENE == bMobileSupportsGPUScene);
 #endif
-}
 
-bool FReadOnlyCVARCache::MobileDeferredShading(EShaderPlatform Platform)
-{
-#if WITH_EDITOR
-	return MobileDeferredShadingIniValue(Platform);
-#else
-	return bMobileDeferredShadingValue;
+#ifdef PROJECT_CVAR_MOBILE_FORWARD_LOCALLIGHTS
+	check(PROJECT_CVAR_MOBILE_FORWARD_LOCALLIGHTS == MobileForwardLocalLightsValue);
 #endif
-}
 
-bool FReadOnlyCVARCache::MobileEnableMovableSpotlightsShadow(EShaderPlatform Platform)
-{
-#if WITH_EDITOR
-	return MobileEnableMovableSpotlightsShadowIniValue(Platform);
-#else
-	return bMobileEnableMovableSpotlightsShadowValue;
+#ifdef PROJECT_CVAR_MOBILE_DEFERRED_SHADING
+	check(!!PROJECT_CVAR_MOBILE_DEFERRED_SHADING == bMobileDeferredShadingValue);
 #endif
-}
\ No newline at end of file
+	
+	bInitialized = true;
+}
diff --git a/Engine/Source/Runtime/RenderCore/Private/RenderUtils.cpp b/Engine/Source/Runtime/RenderCore/Private/RenderUtils.cpp
--- a/Engine/Source/Runtime/RenderCore/Private/RenderUtils.cpp
+++ b/Engine/Source/Runtime/RenderCore/Private/RenderUtils.cpp
@@ -18,7 +18,6 @@
 #include "SubstrateDefinitions.h"
 #include "Animation/MeshDeformerProvider.h"
 #include "Interfaces/ITargetPlatform.h"
-#include "ReadOnlyCVARCache.h"
 
 #if WITH_EDITOR
 #include "Interfaces/ITargetPlatformManagerModule.h"
@@ -474,26 +473,11 @@ RENDERCORE_API FVertexDeclarationRHIRef& GetVertexDeclarationFVector2()
 	return GVector2VertexDeclaration.VertexDeclarationRHI;
 }
 
-RENDERCORE_API bool IsMobileHDR()
-{
-	return FReadOnlyCVARCache::MobileHDR();
-}
-
-RENDERCORE_API bool MobileSupportsGPUScene()
-{
-	return FReadOnlyCVARCache::MobileSupportsGPUScene();
-}
-
 RENDERCORE_API bool PlatformGPUSceneUsesUniformBufferView(const FStaticShaderPlatform Platform)
 {
 	return IsMobilePlatform(Platform) && FDataDrivenShaderPlatformInfo::GetSupportsUniformBufferObjects(Platform);
 }
 
-RENDERCORE_API bool IsMobileDeferredShadingEnabled(const FStaticShaderPlatform Platform)
-{
-	return FReadOnlyCVARCache::MobileDeferredShading(Platform) && IsMobileHDR();
-}
-
 RENDERCORE_API bool MobileRequiresSceneDepthAux(const FStaticShaderPlatform Platform)
 {
 	const bool bMobileHDR = IsMobileHDR();
@@ -584,11 +568,6 @@ RENDERCORE_API bool IsMobileMovableSpotlightShadowsEnabled(const FStaticShaderPl
 	return FReadOnlyCVARCache::MobileEnableMovableSpotlightsShadow(Platform);
 }
 
-RENDERCORE_API bool MobileForwardEnableLocalLights(const FStaticShaderPlatform Platform)
-{
-	return FReadOnlyCVARCache::MobileForwardLocalLights(Platform) > 0;
-}
-
 RENDERCORE_API bool MobileForwardEnableClusteredReflections(const FStaticShaderPlatform Platform)
 {
 	static FShaderPlatformCachedIniValue<bool> MobileForwardEnableClusteredReflectionsIniValue(TEXT("r.Mobile.Forward.EnableClusteredReflections"));
@@ -1644,11 +1623,6 @@ ERayTracingMode GetRayTracingMode()
 	return IsRayTracingAllowed() ? GRayTracingMode : ERayTracingMode::Disabled;
 }
 
-bool IsStaticLightingAllowed()
-{
-	return FReadOnlyCVARCache::AllowStaticLighting();
-}
-
 bool UseSplineMeshSceneResources(const FStaticShaderPlatform Platform)
 {
 	// This feature currently requires GPU Scene to be supported and enabled.
diff --git a/Engine/Source/Runtime/RenderCore/Public/ReadOnlyCVARCache.h b/Engine/Source/Runtime/RenderCore/Public/ReadOnlyCVARCache.h
--- a/Engine/Source/Runtime/RenderCore/Public/ReadOnlyCVARCache.h
+++ b/Engine/Source/Runtime/RenderCore/Public/ReadOnlyCVARCache.h
@@ -11,48 +11,154 @@ struct FReadOnlyCVARCache
 {
 	static void Initialize();
 		
-	static bool AllowStaticLighting();
-	RENDERCORE_API static bool EnablePointLightShadows();
-	RENDERCORE_API static bool EnableStationarySkylight();
-	RENDERCORE_API static bool EnableLowQualityLightmaps();
-	RENDERCORE_API static bool SupportSkyAtmosphere();
+	static inline bool AllowStaticLighting()
+	{
+		checkSlow(bInitialized);
+		#ifdef PROJECT_CVAR_ALLOW_STATIC_LIGHTING
+			return PROJECT_CVAR_ALLOW_STATIC_LIGHTING;
+		#else	
+			return bAllowStaticLighting;
+		#endif	
+	}
+
+	static inline bool EnablePointLightShadows()
+	{
+		return bEnablePointLightShadows;
+	}
+	
+	static inline bool EnableStationarySkylight()
+	{
+		return bEnableStationarySkylight;
+	}
+	
+	static inline bool EnableLowQualityLightmaps()
+	{
+		return bEnableLowQualityLightmaps;
+	}
+
+	static inline bool SupportSkyAtmosphere()
+	{
+		return bSupportSkyAtmosphere;
+	}
 
 	// Mobile specific
-	static bool MobileHDR();
-	static bool MobileSupportsGPUScene();
-	static bool MobileDeferredShading(EShaderPlatform Platform);
-	static bool MobileEnableMovableSpotlightsShadow(EShaderPlatform Platform);
-
-	RENDERCORE_API static bool MobileAllowMovableDirectionalLights();
-	RENDERCORE_API static bool MobileAllowDistanceFieldShadows();
-	RENDERCORE_API static bool MobileEnableStaticAndCSMShadowReceivers();
-	RENDERCORE_API static bool MobileEnableMovableLightCSMShaderCulling();
-	RENDERCORE_API static int32 MobileSkyLightPermutation();
-	RENDERCORE_API static bool MobileEnableNoPrecomputedLightingCSMShader();
-	RENDERCORE_API static int32 MobileEarlyZPass(EShaderPlatform Platform);
-	RENDERCORE_API static int32 MobileForwardLocalLights(EShaderPlatform Platform);
+	static inline bool MobileHDR()
+	{
+		checkSlow(bInitialized);
+		#ifdef PROJECT_CVAR_MOBILE_HDR
+			return PROJECT_CVAR_MOBILE_HDR;
+		#else
+			return bMobileHDR;
+		#endif	
+	}
+
+	static inline bool MobileSupportsGPUScene()
+	{
+		checkSlow(bInitialized);
+		#ifdef PROJECT_CVAR_MOBILE_SUPPORTS_GPUSCENE
+			return PROJECT_CVAR_MOBILE_SUPPORTS_GPUSCENE;
+		#else
+			return bMobileSupportsGPUScene;
+		#endif
+	}
+
+	static inline bool MobileAllowMovableDirectionalLights()
+	{
+		return bMobileAllowMovableDirectionalLights;
+	}
+
+	static inline bool MobileAllowDistanceFieldShadows()
+	{
+		return bMobileAllowDistanceFieldShadows;
+	}
+	
+	static inline bool MobileEnableStaticAndCSMShadowReceivers()
+	{
+		return bMobileEnableStaticAndCSMShadowReceivers;
+	}
+
+	static inline bool MobileEnableMovableLightCSMShaderCulling()
+	{
+		return bMobileEnableMovableLightCSMShaderCulling;
+	}
+	
+	static inline int32 MobileSkyLightPermutation()
+	{
+		return MobileSkyLightPermutationValue;
+	}
+
+	static inline bool MobileEnableNoPrecomputedLightingCSMShader()
+	{
+		return bMobileEnableNoPrecomputedLightingCSMShader;
+	}
+
+	static inline int32 MobileEarlyZPass(EShaderPlatform Platform)
+	{
+		#if WITH_EDITOR
+			return MobileEarlyZPassIniValue(Platform);
+		#else
+			return MobileEarlyZPassValue;
+		#endif
+	}
+	
+	static inline int32 MobileForwardLocalLights(EShaderPlatform Platform)
+	{
+		#if WITH_EDITOR
+			return MobileForwardLocalLightsIniValue(Platform);
+		#elif defined PROJECT_CVAR_MOBILE_FORWARD_LOCALLIGHTS
+			return PROJECT_CVAR_MOBILE_FORWARD_LOCALLIGHTS;
+		#else
+			return MobileForwardLocalLightsValue;
+		#endif
+	}
+	
+	static inline bool MobileDeferredShading(EShaderPlatform Platform)
+	{
+		#if WITH_EDITOR
+			return MobileDeferredShadingIniValue(Platform);
+		#elif defined PROJECT_CVAR_MOBILE_DEFERRED_SHADING
+			return PROJECT_CVAR_MOBILE_DEFERRED_SHADING;
+		#else
+			return bMobileDeferredShadingValue;
+		#endif
+	}
+
+	static inline bool MobileEnableMovableSpotlightsShadow(EShaderPlatform Platform)
+	{
+		#if WITH_EDITOR
+			return MobileEnableMovableSpotlightsShadowIniValue(Platform);
+		#else
+			return bMobileEnableMovableSpotlightsShadowValue;
+		#endif
+	}
 
 private:
 	static bool bInitialized;
 
-	static bool bEnablePointLightShadows;
-	static bool bEnableStationarySkylight;
-	static bool bEnableLowQualityLightmaps;
-	static bool bAllowStaticLighting;
-	static bool bSupportSkyAtmosphere;
+	RENDERCORE_API static bool bAllowStaticLighting;
+	RENDERCORE_API static bool bEnablePointLightShadows;
+	RENDERCORE_API static bool bEnableStationarySkylight;
+	RENDERCORE_API static bool bEnableLowQualityLightmaps;
+	RENDERCORE_API static bool bSupportSkyAtmosphere;
 
 	// Mobile specific
-	static bool bMobileHDR;
-	static bool bMobileAllowMovableDirectionalLights;
-	static bool bMobileAllowDistanceFieldShadows;
-	static bool bMobileEnableStaticAndCSMShadowReceivers;
-	static bool bMobileEnableMovableLightCSMShaderCulling;
-	static bool bMobileSupportsGPUScene;
-	static int32 MobileSkyLightPermutationValue;
-	static int32 MobileEarlyZPassValue;
-	static int32 MobileForwardLocalLightsValue;
-	static bool bMobileEnableNoPrecomputedLightingCSMShader;
-	static bool bMobileDeferredShadingValue;
-	static bool bMobileEnableMovableSpotlightsShadowValue;
+	RENDERCORE_API static bool bMobileHDR;
+	RENDERCORE_API static bool bMobileAllowMovableDirectionalLights;
+	RENDERCORE_API static bool bMobileAllowDistanceFieldShadows;
+	RENDERCORE_API static bool bMobileEnableStaticAndCSMShadowReceivers;
+	RENDERCORE_API static bool bMobileEnableMovableLightCSMShaderCulling;
+	RENDERCORE_API static bool bMobileSupportsGPUScene;
+	RENDERCORE_API static int32 MobileSkyLightPermutationValue;
+	RENDERCORE_API static int32 MobileEarlyZPassValue;
+	RENDERCORE_API static int32 MobileForwardLocalLightsValue;
+	RENDERCORE_API static bool bMobileEnableNoPrecomputedLightingCSMShader;
+	RENDERCORE_API static bool bMobileDeferredShadingValue;
+	RENDERCORE_API static bool bMobileEnableMovableSpotlightsShadowValue;
+
+private:
+	RENDERCORE_API static int32 MobileEarlyZPassIniValue(EShaderPlatform Platform);
+	RENDERCORE_API static int32 MobileForwardLocalLightsIniValue(EShaderPlatform Platform);
+	RENDERCORE_API static bool MobileDeferredShadingIniValue(EShaderPlatform Platform);
+	RENDERCORE_API static bool MobileEnableMovableSpotlightsShadowIniValue(EShaderPlatform Platform);
 };
 
diff --git a/Engine/Source/Runtime/RenderCore/Public/RenderUtils.h b/Engine/Source/Runtime/RenderCore/Public/RenderUtils.h
--- a/Engine/Source/Runtime/RenderCore/Public/RenderUtils.h
+++ b/Engine/Source/Runtime/RenderCore/Public/RenderUtils.h
@@ -6,6 +6,7 @@
 #include "RHIFwd.h"
 #include "RHIShaderPlatform.h"
 #include "RHIFeatureLevel.h"
+#include "ReadOnlyCVARCache.h"
 
 #if UE_ENABLE_INCLUDE_ORDER_DEPRECATED_IN_5_2
 #include "RHI.h"
@@ -270,13 +271,27 @@ RENDERCORE_API FVertexDeclarationRHIRef& GetVertexDeclarationFVector3();
 RENDERCORE_API FVertexDeclarationRHIRef& GetVertexDeclarationFVector2();
 
 /** True if HDR is enabled for the mobile renderer. */
-RENDERCORE_API bool IsMobileHDR();
+inline bool IsMobileHDR()
+{
+	return FReadOnlyCVARCache::MobileHDR();
+}
 
-RENDERCORE_API bool MobileSupportsGPUScene();
+inline bool MobileSupportsGPUScene()
+{
+	return FReadOnlyCVARCache::MobileSupportsGPUScene();
+}
 
-RENDERCORE_API bool PlatformGPUSceneUsesUniformBufferView(const FStaticShaderPlatform Platform);
+inline bool IsMobileDeferredShadingEnabled(const FStaticShaderPlatform Platform)
+{
+	return FReadOnlyCVARCache::MobileDeferredShading(Platform) && IsMobileHDR();
+}
 
-RENDERCORE_API bool IsMobileDeferredShadingEnabled(const FStaticShaderPlatform Platform);
+inline bool MobileForwardEnableLocalLights(const FStaticShaderPlatform Platform)
+{
+	return FReadOnlyCVARCache::MobileForwardLocalLights(Platform) > 0;
+}
+
+RENDERCORE_API bool PlatformGPUSceneUsesUniformBufferView(const FStaticShaderPlatform Platform);
 
 RENDERCORE_API bool MobileRequiresSceneDepthAux(const FStaticShaderPlatform Platform);
 
@@ -296,8 +311,6 @@ RENDERCORE_API bool IsMobileDistanceFieldEnabled(const FStaticShaderPlatform Pla
 
 RENDERCORE_API bool IsMobileMovableSpotlightShadowsEnabled(const FStaticShaderPlatform Platform);
 
-RENDERCORE_API bool MobileForwardEnableLocalLights(const FStaticShaderPlatform Platform);
-
 RENDERCORE_API bool MobileForwardEnableClusteredReflections(const FStaticShaderPlatform Platform);
 
 RENDERCORE_API bool MobileUsesShadowMaskTexture(const FStaticShaderPlatform Platform);
@@ -510,7 +523,10 @@ extern RENDERCORE_API bool IsRayTracingAllowed();
 extern RENDERCORE_API ERayTracingMode GetRayTracingMode();
 
 // Returns 'true' when static lighting is enabled for the project
-extern RENDERCORE_API bool IsStaticLightingAllowed();
+inline bool IsStaticLightingAllowed()
+{
+	return FReadOnlyCVARCache::AllowStaticLighting();
+}
 
 extern RENDERCORE_API bool DoesPlatformSupportLumenGI(EShaderPlatform Platform, bool bSkipProjectCheck = false);
 
diff --git a/Engine/Source/Runtime/Renderer/Private/LightMapRendering.cpp b/Engine/Source/Runtime/Renderer/Private/LightMapRendering.cpp
--- a/Engine/Source/Runtime/Renderer/Private/LightMapRendering.cpp
+++ b/Engine/Source/Runtime/Renderer/Private/LightMapRendering.cpp
@@ -10,7 +10,6 @@ LightMapRendering.cpp: Light map rendering implementations.
 #include "PrecomputedVolumetricLightmap.h"
 #include "RenderCore.h"
 #include "DataDrivenShaderPlatformInfo.h"
-#include "ReadOnlyCVARCache.h"
 
 IMPLEMENT_TYPE_LAYOUT(FUniformLightMapPolicyShaderParametersType);
 
diff --git a/Engine/Source/Runtime/Renderer/Private/LightSceneInfo.cpp b/Engine/Source/Runtime/Renderer/Private/LightSceneInfo.cpp
--- a/Engine/Source/Runtime/Renderer/Private/LightSceneInfo.cpp
+++ b/Engine/Source/Runtime/Renderer/Private/LightSceneInfo.cpp
@@ -12,7 +12,6 @@
 #include "Misc/LargeWorldRenderPosition.h"
 #include "LocalLightSceneProxy.h"
 #include "ShadowRendering.h"
-#include "ReadOnlyCVARCache.h"
 
 int32 GWholeSceneShadowUnbuiltInteractionThreshold = 500;
 static FAutoConsoleVariableRef CVarWholeSceneShadowUnbuiltInteractionThreshold(
diff --git a/Engine/Source/Runtime/Renderer/Private/MobileBasePass.cpp b/Engine/Source/Runtime/Renderer/Private/MobileBasePass.cpp
--- a/Engine/Source/Runtime/Renderer/Private/MobileBasePass.cpp
+++ b/Engine/Source/Runtime/Renderer/Private/MobileBasePass.cpp
@@ -14,7 +14,6 @@
 #include "MeshPassProcessor.inl"
 #include "Engine/TextureCube.h"
 #include "ShaderPlatformCachedIniValue.h"
-#include "ReadOnlyCVARCache.h"
 
 bool MobileLocalLightsBufferEnabled(const FStaticShaderPlatform Platform)
 {
diff --git a/Engine/Source/Runtime/Renderer/Private/MobileBasePassRendering.h b/Engine/Source/Runtime/Renderer/Private/MobileBasePassRendering.h
--- a/Engine/Source/Runtime/Renderer/Private/MobileBasePassRendering.h
+++ b/Engine/Source/Runtime/Renderer/Private/MobileBasePassRendering.h
@@ -28,7 +28,6 @@
 #include "LocalFogVolumeRendering.h"
 #include "DBufferTextures.h"
 #include "CompositionLighting/PostProcessDeferredDecals.h"
-#include "ReadOnlyCVARCache.h"
 
 bool MobileLocalLightsBufferEnabled(const FStaticShaderPlatform Platform);
 bool MobileMergeLocalLightsInPrepassEnabled(const FStaticShaderPlatform Platform);
diff --git a/Engine/Source/Runtime/Renderer/Private/RendererScene.cpp b/Engine/Source/Runtime/Renderer/Private/RendererScene.cpp
--- a/Engine/Source/Runtime/Renderer/Private/RendererScene.cpp
+++ b/Engine/Source/Runtime/Renderer/Private/RendererScene.cpp
@@ -81,7 +81,6 @@
 #include "SceneCulling/SceneCulling.h"
 #include "InstanceCulling/InstanceCullingOcclusionQuery.h"
 #include "ComputeWorkerInterface.h"
-#include "ReadOnlyCVARCache.h"
 
 #if RHI_RAYTRACING
 #include "Nanite/NaniteRayTracing.h"
diff --git a/Engine/Source/Runtime/Renderer/Private/SceneRendering.cpp b/Engine/Source/Runtime/Renderer/Private/SceneRendering.cpp
--- a/Engine/Source/Runtime/Renderer/Private/SceneRendering.cpp
+++ b/Engine/Source/Runtime/Renderer/Private/SceneRendering.cpp
@@ -66,7 +66,6 @@
 #include "Misc/AutomationTest.h"
 #include "Engine/TextureCube.h"
 #include "GPUSkinCacheVisualizationData.h"
-#include "ReadOnlyCVARCache.h"
 #if WITH_EDITOR
 #include "Rendering/StaticLightingSystemInterface.h"
 #endif
diff --git a/Engine/Source/Runtime/Renderer/Private/ShadowSetup.cpp b/Engine/Source/Runtime/Renderer/Private/ShadowSetup.cpp
--- a/Engine/Source/Runtime/Renderer/Private/ShadowSetup.cpp
+++ b/Engine/Source/Runtime/Renderer/Private/ShadowSetup.cpp
@@ -41,7 +41,6 @@
 #include "Shadows/ShadowScene.h"
 #include "LineTypes.h"
 #include "SceneCulling/SceneCulling.h"
-#include "ReadOnlyCVARCache.h"
 
 using namespace UE::Geometry;
 
diff --git a/Engine/Source/Runtime/Renderer/Private/ShadowSetupMobile.cpp b/Engine/Source/Runtime/Renderer/Private/ShadowSetupMobile.cpp
--- a/Engine/Source/Runtime/Renderer/Private/ShadowSetupMobile.cpp
+++ b/Engine/Source/Runtime/Renderer/Private/ShadowSetupMobile.cpp
@@ -21,7 +21,6 @@ ShadowSetupMobile.cpp: Shadow setup implementation for mobile specific features.
 #include "RenderCore.h"
 #include "ShadowRendering.h"
 #include "StaticMeshBatch.h"
-#include "ReadOnlyCVARCache.h"
 
 static TAutoConsoleVariable<int32> CVarCsmShaderCullingDebugGfx(
 	TEXT("r.Mobile.Shadow.CSMShaderCullingDebugGfx"),
diff --git a/Engine/Source/Runtime/Renderer/Private/SkyAtmosphereRendering.cpp b/Engine/Source/Runtime/Renderer/Private/SkyAtmosphereRendering.cpp
--- a/Engine/Source/Runtime/Renderer/Private/SkyAtmosphereRendering.cpp
+++ b/Engine/Source/Runtime/Renderer/Private/SkyAtmosphereRendering.cpp
@@ -19,7 +19,6 @@
 #include "RendererUtils.h"
 #include "ScreenPass.h"
 #include "UnrealEngine.h"
-#include "ReadOnlyCVARCache.h"
 
 
 //PRAGMA_DISABLE_OPTIMIZATION
