Fixed linking error in a debug configurations

diff --git a/Engine/Source/Runtime/RenderCore/Public/ReadOnlyCVARCache.h b/Engine/Source/Runtime/RenderCore/Public/ReadOnlyCVARCache.h
--- a/Engine/Source/Runtime/RenderCore/Public/ReadOnlyCVARCache.h
+++ b/Engine/Source/Runtime/RenderCore/Public/ReadOnlyCVARCache.h
@@ -133,7 +133,7 @@ struct FReadOnlyCVARCache
 	}
 
 private:
-	static bool bInitialized;
+	RENDERCORE_API static bool bInitialized;
 
 	RENDERCORE_API static bool bAllowStaticLighting;
 	RENDERCORE_API static bool bEnablePointLightShadows;
