[shaders] fix verbose logging of shader hashes to be deterministic - parts of the FSHA1 struct are seemingly intentionally uninitialized, but further hashing this struct itself doesn't really make sense. instead copy the incremental state and generate a "final" hash after each file; this isn't the hash of just that file but rather the hash of all files up to the point of that file, which is not perfect but good enough to track down real nondeterminism issues.

diff --git a/Engine/Source/Runtime/RenderCore/Private/ShaderCore.cpp b/Engine/Source/Runtime/RenderCore/Private/ShaderCore.cpp
--- a/Engine/Source/Runtime/RenderCore/Private/ShaderCore.cpp
+++ b/Engine/Source/Runtime/RenderCore/Private/ShaderCore.cpp
@@ -2643,10 +2643,10 @@ static void UpdateSingleShaderFilehash(FSHA1& InOutHashState, const TCHAR* Virtu
 #if WITH_EDITOR &&  !(UE_BUILD_SHIPPING || UE_BUILD_TEST)
 		if (UE_LOG_ACTIVE(LogShaders, Verbose))
 		{
-			uint8 HashBytes[20];
-			FSHA1::HashBuffer(&InOutHashState, sizeof(FSHA1), HashBytes);
+			FSHA1 HashStateCopy = InOutHashState;
+			FSHAHash IncrementalHash = HashStateCopy.Finalize();
 			
-			UE_LOG(LogShaders, Verbose, TEXT("Processing include file for %s, %s, %s"), VirtualFilePath, *IncludeVirtualFilePaths[IncludeIndex], *BytesToHex(HashBytes, 20));
+			UE_LOG(LogShaders, Verbose, TEXT("Processing include file for %s, %s, %s"), VirtualFilePath, *IncludeVirtualFilePaths[IncludeIndex], *BytesToHex(IncrementalHash.Hash, 20));
 		}
 #endif
 	}
