Add some Nanite read-only CVars to the shader map DDC key to prevent caching issues.

diff --git a/Engine/Source/Runtime/RenderCore/Private/Shader.cpp b/Engine/Source/Runtime/RenderCore/Private/Shader.cpp
--- a/Engine/Source/Runtime/RenderCore/Private/Shader.cpp
+++ b/Engine/Source/Runtime/RenderCore/Private/Shader.cpp
@@ -2317,6 +2317,20 @@ void ShaderMapAppendKeyString(EShaderPlatform Platform, FString& KeyString)
 	{
 		KeyString += TEXT("_SKPassThrough0");
 	}
+
+	if (UseNanite(Platform))
+	{
+		static const auto CVarAllowTess = IConsoleManager::Get().FindConsoleVariable(TEXT("r.Nanite.AllowTessellation"));
+		static const auto CVarAllowCSMat = IConsoleManager::Get().FindConsoleVariable(TEXT("r.Nanite.AllowComputeMaterials"));
+		static const auto CVarAllowPSMat = IConsoleManager::Get().FindConsoleVariable(TEXT("r.Nanite.AllowLegacyMaterials"));
+
+		KeyString.Appendf(
+			TEXT("_Nanite-Tess%dCSMat%dPSMat%d"),
+			CVarAllowTess ? CVarAllowTess->GetInt() : 0,
+			CVarAllowCSMat ? CVarAllowCSMat->GetInt() : 0,
+			CVarAllowPSMat ? CVarAllowPSMat->GetInt() : 0
+		);
+	}
 }
 
 EShaderPermutationFlags GetShaderPermutationFlags(const FPlatformTypeLayoutParameters& LayoutParams)
