Remove IsInParallelRenderingThread when creating RHI shaders because this call can happen during PSO precaching as well from worker threads - check was only around stat tracking updates

diff --git a/Engine/Source/Runtime/RenderCore/Private/ShaderCodeLibrary.cpp b/Engine/Source/Runtime/RenderCore/Private/ShaderCodeLibrary.cpp
--- a/Engine/Source/Runtime/RenderCore/Private/ShaderCodeLibrary.cpp
+++ b/Engine/Source/Runtime/RenderCore/Private/ShaderCodeLibrary.cpp
@@ -1203,11 +1203,8 @@ FRHIShader* FShaderMapResource_SharedCode::CreateRHIShaderOrCrash(int32 ShaderIn
 	double TimeFunctionEntered = FPlatformTime::Seconds();
 	ON_SCOPE_EXIT
 	{
-		if (IsInRenderingThread())
-		{
-			double ShaderCreationTime = FPlatformTime::Seconds() - TimeFunctionEntered;
-			INC_FLOAT_STAT_BY(STAT_Shaders_TotalRTShaderInitForRenderingTime, ShaderCreationTime);
-		}
+		double ShaderCreationTime = FPlatformTime::Seconds() - TimeFunctionEntered;
+		INC_FLOAT_STAT_BY(STAT_Shaders_TotalRTShaderInitForRenderingTime, ShaderCreationTime);
 	};
 #endif
 
diff --git a/Engine/Source/Runtime/RenderCore/Private/ShaderResource.cpp b/Engine/Source/Runtime/RenderCore/Private/ShaderResource.cpp
--- a/Engine/Source/Runtime/RenderCore/Private/ShaderResource.cpp
+++ b/Engine/Source/Runtime/RenderCore/Private/ShaderResource.cpp
@@ -600,11 +600,8 @@ FRHIShader* FShaderMapResource_InlineCode::CreateRHIShaderOrCrash(int32 ShaderIn
 	double TimeFunctionEntered = FPlatformTime::Seconds();
 	ON_SCOPE_EXIT
 	{
-		if (IsInParallelRenderingThread())
-		{
-			double ShaderCreationTime = FPlatformTime::Seconds() - TimeFunctionEntered;
-			INC_FLOAT_STAT_BY(STAT_Shaders_TotalRTShaderInitForRenderingTime, ShaderCreationTime);
-		}
+		double ShaderCreationTime = FPlatformTime::Seconds() - TimeFunctionEntered;
+		INC_FLOAT_STAT_BY(STAT_Shaders_TotalRTShaderInitForRenderingTime, ShaderCreationTime);
 	};
 #endif
 
