Add IES Atlas support for translucent surface.

diff --git a/Engine/Shaders/Private/BasePassPixelShader.usf b/Engine/Shaders/Private/BasePassPixelShader.usf
--- a/Engine/Shaders/Private/BasePassPixelShader.usf
+++ b/Engine/Shaders/Private/BasePassPixelShader.usf
@@ -57,6 +57,11 @@
 #define USE_DBUFFER 0
 #endif
 
+// Enable IES profile evaluation based on project settings
+#if SUPPORT_IESPROFILE_ON_FORWARD_LIT_TRANSLUCENT
+#define USE_IES_PROFILE 1
+#endif
+
 #ifndef COMPUTE_SHADED
 #define COMPUTE_SHADED 0
 #endif
diff --git a/Engine/Shaders/Private/ForwardLightingCommon.ush b/Engine/Shaders/Private/ForwardLightingCommon.ush
--- a/Engine/Shaders/Private/ForwardLightingCommon.ush
+++ b/Engine/Shaders/Private/ForwardLightingCommon.ush
@@ -284,6 +284,9 @@ FDeferredLightingSplit GetForwardDirectLightingSplit(
 			LightData.HairTransmittance = EvaluateDualScattering(GBufferData, -CameraVector, LightData.Direction);
 		}
 		#endif
+		#if USE_IES_PROFILE
+		LightData.Color *= ComputeLightProfileMultiplier(TranslatedWorldPosition, LightData.TranslatedWorldPosition, -LightData.Direction, LightData.Tangent, LightData.IESAtlasIndex);
+		#endif
 
 		float DynamicShadowing = dot(PreviewShadowMapChannelMask, DynamicShadowFactors);
 		float4 LightAttenuation = float4(1, 1, DynamicShadowing.x, DynamicShadowing.x);
diff --git a/Engine/Shaders/Private/Substrate/SubstrateForwardLighting.ush b/Engine/Shaders/Private/Substrate/SubstrateForwardLighting.ush
--- a/Engine/Shaders/Private/Substrate/SubstrateForwardLighting.ush
+++ b/Engine/Shaders/Private/Substrate/SubstrateForwardLighting.ush
@@ -479,7 +479,10 @@ float3 SubstrateForwardLighting(
 					half4 PreviewShadowMapChannelMask = 1;
 					uint  LocalLightChannelMask = LIGHTING_CHANNEL_MASK;
 					FDeferredLightData LightData = ConvertToDeferredLight(LocalLight, 1.0f /*SpecularScale*/, PreviewShadowMapChannelMask, LocalLightChannelMask);
-					
+					#if USE_IES_PROFILE
+					LightData.Color *= ComputeLightProfileMultiplier(TranslatedWorldPosition, LightData.TranslatedWorldPosition, -LightData.Direction, LightData.Tangent, LightData.IESAtlasIndex);
+					#endif
+
 					// Rect light optionally supported in forward.
 					// * This is not done explicitely in the legacy path (resulting in no lighting) but we do it for Substrate in order to avoid visual artefact.
 					// * This can be removed when LocalLight properly supports rect lights.
diff --git a/Engine/Source/Runtime/Engine/Classes/Engine/RendererSettings.h b/Engine/Source/Runtime/Engine/Classes/Engine/RendererSettings.h
--- a/Engine/Source/Runtime/Engine/Classes/Engine/RendererSettings.h
+++ b/Engine/Source/Runtime/Engine/Classes/Engine/RendererSettings.h
@@ -927,11 +927,20 @@ class URendererSettings : public UDeveloperSettings
 		ConfigRestartRequired = true))
 		uint32 bTranslucentUsesLightFunctionAtlas : 1;
 
+	/**
+		"Enable IES profile evaluation on translucent materials when using the Forward Shading mode."
+		*/
+	UPROPERTY(config, EditAnywhere, Category = Optimizations, meta = (
+		ConsoleVariable = "r.Translucent.UsesIESProfiles", DisplayName = "Support IES profiles On Translucent Materials (When Using ForwardShading)",
+		ToolTip = "Enable IES profile evaluation on translucent materials when using the Forward Shading mode.",
+		ConfigRestartRequired = true))
+	uint32 bSupportIESProfileOnTranslucent : 1;
+
 	/**
 	"Enable rect light evaluation on translucent materials when using the Forward Shading mode."
 	*/
 	UPROPERTY(config, EditAnywhere, Category = Optimizations, meta = (
-		ConsoleVariable = "r.RectLightAtlas.Translucent", DisplayName = "Support Rect Light On Translucent Materials (When Using ForwardShading)",
+		ConsoleVariable = "r.Translucent.UsesRectLights", DisplayName = "Support Rect Light On Translucent Materials (When Using ForwardShading)",
 		ToolTip = "Enable rect light evaluation on translucent materials when using the Forward Shading mode.",
 		ConfigRestartRequired = true))
 		uint32 bSupportRectLightOnTranslucent : 1;
diff --git a/Engine/Source/Runtime/Engine/Private/ShaderCompiler/ShaderCompiler.cpp b/Engine/Source/Runtime/Engine/Private/ShaderCompiler/ShaderCompiler.cpp
--- a/Engine/Source/Runtime/Engine/Private/ShaderCompiler/ShaderCompiler.cpp
+++ b/Engine/Source/Runtime/Engine/Private/ShaderCompiler/ShaderCompiler.cpp
@@ -8218,9 +8218,13 @@ void GlobalBeginCompileShader(
 	}
 
 	{
-		static IConsoleVariable* CVar = IConsoleManager::Get().FindConsoleVariable(TEXT("r.RectLightAtlas.Translucent"));
-		const bool bSupportRectLlightOnTranslucentSurface = CVar && CVar->GetInt() > 0;
-		SET_SHADER_DEFINE(Input.Environment, SUPPORT_RECTLIGHT_ON_FORWARD_LIT_TRANSLUCENT, bSupportRectLlightOnTranslucentSurface ? 1 : 0);
+		const bool bTranslucentUsesLightRectLights = GetTranslucentUsesLightRectLights();
+		SET_SHADER_DEFINE(Input.Environment, SUPPORT_RECTLIGHT_ON_FORWARD_LIT_TRANSLUCENT, bTranslucentUsesLightRectLights ? 1 : 0);
+	}
+
+	{
+		const bool bTranslucentUsesLightIESProfiles = GetTranslucentUsesLightIESProfiles();
+		SET_SHADER_DEFINE(Input.Environment, SUPPORT_IESPROFILE_ON_FORWARD_LIT_TRANSLUCENT, bTranslucentUsesLightIESProfiles ? 1 : 0);
 	}
 
 	{
diff --git a/Engine/Source/Runtime/RenderCore/Private/RenderUtils.cpp b/Engine/Source/Runtime/RenderCore/Private/RenderUtils.cpp
--- a/Engine/Source/Runtime/RenderCore/Private/RenderUtils.cpp
+++ b/Engine/Source/Runtime/RenderCore/Private/RenderUtils.cpp
@@ -1655,6 +1655,18 @@ static TAutoConsoleVariable<int32> CVarTranslucentUsesLightFunctionAtlas(
 	TEXT("Enable sampling of the light function atlas on translucent materials using forward shading."),
 	ECVF_ReadOnly | ECVF_RenderThreadSafe);
 
+static TAutoConsoleVariable<int32> CVarRectLighTranslucent(
+	TEXT("r.Translucent.UsesRectLights"),
+	0,
+	TEXT("Enable rect light support for translucent surfaces. When enabled, it will add an extrat sampler to the pixel shader (limited to 16 on dx11 based system)"),
+	ECVF_ReadOnly | ECVF_RenderThreadSafe);
+
+static TAutoConsoleVariable<int32> CVarIESProfileTranslucent(
+	TEXT("r.Translucent.UsesIESProfiles"),
+	0,
+	TEXT("Enable IES profiles support for translucent surfaces. When enabled, it will add an extrat sampler to the pixel shader (limited to 16 on dx11 based system)"),
+	ECVF_ReadOnly | ECVF_RenderThreadSafe);
+
 int32 GetLightFunctionAtlasFormat()
 {
 	return CVarLightFunctionAtlasFormat.GetValueOnAnyThread();
@@ -1670,6 +1682,16 @@ bool GetTranslucentUsesLightFunctionAtlas()
 	return CVarTranslucentUsesLightFunctionAtlas.GetValueOnAnyThread() > 0;
 }
 
+bool GetTranslucentUsesLightRectLights()
+{
+	return CVarRectLighTranslucent.GetValueOnAnyThread() > 0;
+}
+
+bool GetTranslucentUsesLightIESProfiles()
+{
+	return CVarIESProfileTranslucent.GetValueOnAnyThread() > 0;
+}
+
 ///////////////////////////////////////////////////////////////////////////////////////////////////
 // Substrate settings interface
 
diff --git a/Engine/Source/Runtime/RenderCore/Private/Shader.cpp b/Engine/Source/Runtime/RenderCore/Private/Shader.cpp
--- a/Engine/Source/Runtime/RenderCore/Private/Shader.cpp
+++ b/Engine/Source/Runtime/RenderCore/Private/Shader.cpp
@@ -1992,13 +1992,21 @@ void ShaderMapAppendKeyString(EShaderPlatform Platform, FString& KeyString)
 	}
 
 	{
-		static const auto CVar = IConsoleManager::Get().FindTConsoleVariableDataInt(TEXT("r.RectLightAtlas.Translucent"));
-		if (CVar && CVar->GetValueOnAnyThread() > 0)
+		const bool bTranslucentUsesLightRectLights = GetTranslucentUsesLightRectLights();
+		if (bTranslucentUsesLightRectLights)
 		{
 			KeyString += TEXT("_RECTTRANS");
 		}
 	}
 
+	{
+		const bool bTranslucentUsesLightIESProfiles = GetTranslucentUsesLightIESProfiles();
+		if (bTranslucentUsesLightIESProfiles)
+		{
+			KeyString += TEXT("_IESTRANS");
+		}
+	}
+
 	{
 		static const auto CVar = IConsoleManager::Get().FindTConsoleVariableDataInt(TEXT("r.Shadow.Virtual.TranslucentQuality"));
 		if (CVar && CVar->GetValueOnAnyThread() > 0)
diff --git a/Engine/Source/Runtime/RenderCore/Public/RenderUtils.h b/Engine/Source/Runtime/RenderCore/Public/RenderUtils.h
--- a/Engine/Source/Runtime/RenderCore/Public/RenderUtils.h
+++ b/Engine/Source/Runtime/RenderCore/Public/RenderUtils.h
@@ -570,6 +570,8 @@ RENDERCORE_API int32 GetLightFunctionAtlasFormat();
 // Light function atlas project settings triggering shader compilation
 RENDERCORE_API bool GetSingleLayerWaterUsesLightFunctionAtlas();
 RENDERCORE_API bool GetTranslucentUsesLightFunctionAtlas();
+RENDERCORE_API bool GetTranslucentUsesLightRectLights();
+RENDERCORE_API bool GetTranslucentUsesLightIESProfiles();
 
 // LuminanceMax is the amount of light that will cause the sensor to saturate at EV100.
 //  See also https://en.wikipedia.org/wiki/Film_speed and https://en.wikipedia.org/wiki/Exposure_value for more info.
diff --git a/Engine/Source/Runtime/Renderer/Private/RectLightTextureManager.cpp b/Engine/Source/Runtime/Renderer/Private/RectLightTextureManager.cpp
--- a/Engine/Source/Runtime/Renderer/Private/RectLightTextureManager.cpp
+++ b/Engine/Source/Runtime/Renderer/Private/RectLightTextureManager.cpp
@@ -52,12 +52,6 @@ static TAutoConsoleVariable<int32> CVarRectLighForceUpdate(
 	TEXT("Force rect. light atlas update very frame."),
 	ECVF_RenderThreadSafe);
 
-static TAutoConsoleVariable<int32> CVarRectLighTranslucent(
-	TEXT("r.RectLightAtlas.Translucent"),
-	0,
-	TEXT("Enable rect light support for translucent surfaces. When enabled, it will add an extrat sampler to the pixel shader (limited to 16 on dx11 based system)"),
-	ECVF_ReadOnly | ECVF_RenderThreadSafe);
-
 static TAutoConsoleVariable<int32> CVarRectLighFilterQuality(
 	TEXT("r.RectLightAtlas.FilterQuality"),
 	1,
