Skip RDG tracking for a baked skylight textures

diff --git a/Engine/Source/Runtime/RenderCore/Public/RenderGraphUtils.h b/Engine/Source/Runtime/RenderCore/Public/RenderGraphUtils.h
--- a/Engine/Source/Runtime/RenderCore/Public/RenderGraphUtils.h
+++ b/Engine/Source/Runtime/RenderCore/Public/RenderGraphUtils.h
@@ -268,14 +268,19 @@ inline FRDGBufferRef TryRegisterExternalBuffer(
 	return ExternalPooledBuffer ? GraphBuilder.RegisterExternalBuffer(ExternalPooledBuffer, Flags) : nullptr;
 }
 
-inline FRDGTextureRef RegisterExternalTexture(FRDGBuilder& GraphBuilder, FRHITexture* Texture, const TCHAR* NameIfUnregistered)
+inline FRDGTextureRef RegisterExternalTexture(FRDGBuilder& GraphBuilder, FRHITexture* Texture, const TCHAR* NameIfUnregistered, ERDGTextureFlags Flags)
 {
 	if (FRDGTextureRef FoundTexture = GraphBuilder.FindExternalTexture(Texture))
 	{
 		return FoundTexture;
 	}
 
-	return GraphBuilder.RegisterExternalTexture(CreateRenderTarget(Texture, NameIfUnregistered));
+	return GraphBuilder.RegisterExternalTexture(CreateRenderTarget(Texture, NameIfUnregistered), Flags);
+}
+
+inline FRDGTextureRef RegisterExternalTexture(FRDGBuilder& GraphBuilder, FRHITexture* Texture, const TCHAR* NameIfUnregistered)
+{
+	return RegisterExternalTexture(GraphBuilder, Texture, NameIfUnregistered, ERDGTextureFlags::None);
 }
 
 /** Simple pair of RDG textures used for MSAA. */
diff --git a/Engine/Source/Runtime/Renderer/Private/IndirectLightRendering.cpp b/Engine/Source/Runtime/Renderer/Private/IndirectLightRendering.cpp
--- a/Engine/Source/Runtime/Renderer/Private/IndirectLightRendering.cpp
+++ b/Engine/Source/Runtime/Renderer/Private/IndirectLightRendering.cpp
@@ -545,6 +545,7 @@ void SetupReflectionUniformParameters(FRDGBuilder& GraphBuilder, const FViewInfo
 
 	const bool bApplySkyLight = View.Family->EngineShowFlags.SkyLighting;
 	const FScene* Scene = (const FScene*)View.Family->Scene;
+	ERDGTextureFlags SkyLightTextureFlags = ERDGTextureFlags::None;
 
 	if (Scene
 		&& Scene->SkyLight
@@ -577,6 +578,8 @@ void SetupReflectionUniformParameters(FRDGBuilder& GraphBuilder, const FViewInfo
 					BlendFraction = 0;
 				}
 			}
+			
+			SkyLightTextureFlags = ERDGTextureFlags::SkipTracking;
 		}
 
 		ApplySkyLightMask = 1;
@@ -587,7 +590,7 @@ void SetupReflectionUniformParameters(FRDGBuilder& GraphBuilder, const FViewInfo
 	FRDGTextureRef SkyLightTexture = nullptr;
 	if (SkyLightTextureResource)
 	{
-		 SkyLightTexture = RegisterExternalTexture(GraphBuilder, SkyLightTextureResource, TEXT("SkyLightTexture"));
+		SkyLightTexture = RegisterExternalTexture(GraphBuilder, SkyLightTextureResource, TEXT("SkyLightTexture"), SkyLightTextureFlags);
 	}
 	else
 	{
