Make sure memoryless textures are discarded by RDG at the end of the pass

diff --git a/Engine/Source/Runtime/RenderCore/Private/RenderGraphBuilder.cpp b/Engine/Source/Runtime/RenderCore/Private/RenderGraphBuilder.cpp
--- a/Engine/Source/Runtime/RenderCore/Private/RenderGraphBuilder.cpp
+++ b/Engine/Source/Runtime/RenderCore/Private/RenderGraphBuilder.cpp
@@ -3097,6 +3097,12 @@ void FRDGBuilder::AddEpilogueTransition(FRDGTextureRef Texture)
 			Texture->EpilogueAccess = ERHIAccess::Discard;
 		}
 
+		// Memoryless textures should be discarded at the end of the pass
+		if (EnumHasAnyFlags(Texture->Desc.Flags, TexCreate_Memoryless))
+		{
+			Texture->EpilogueAccess = ERHIAccess::Discard;
+		}
+
 		SubresourceState.Access = Texture->EpilogueAccess;
 
 		InitTextureSubresources(ScratchTextureState, Texture->Layout, &SubresourceState);
