FrameSync: Implemented a workaround for FTask::Wait creating multiple wait task and event object instances when called repeatedly while waiting for a long task to finish.

In particular, it would happen when capturing GPU profiles on some platforms as it leads to the sync between the game and render threads to be very long.

Resolved by creating a custom wait task manually which ensures that we only have one instance.

#rnx

diff --git a/Engine/Source/Runtime/RenderCore/Private/RenderingThread.cpp b/Engine/Source/Runtime/RenderCore/Private/RenderingThread.cpp
--- a/Engine/Source/Runtime/RenderCore/Private/RenderingThread.cpp
+++ b/Engine/Source/Runtime/RenderCore/Private/RenderingThread.cpp
@@ -1358,6 +1358,21 @@ static void GameThreadWaitForTask(const UE::Tasks::FTask& Task, bool bEmptyGameT
 
 			static bool bDisabled = FParse::Param(FCommandLine::Get(), TEXT("nothreadtimeout"));
 
+			// Creating the wait task manually is a workaround for the problem of FTast::Wait creating
+			// a separate wait task and event object on each call. It's a problem because we may call
+			// Wait it in the loop below many times during long frame syncs (e.g. when using GPU profilers)
+			// which would create thousands of such objects and run out of system resources.
+			FSharedEventRef CompletionEvent;
+
+			UE::Tasks::Launch(
+				TEXT("Waiting Task (FrameSync)"),
+				[CompletionEvent] { CompletionEvent->Trigger(); },
+				Task,
+				LowLevelTasks::ETaskPriority::Default,
+				UE::Tasks::EExtendedTaskPriority::Inline,
+				UE::Tasks::ETaskFlags::None
+			);
+
 			do
 			{
 				CheckRenderingThreadHealth();
@@ -1366,7 +1381,7 @@ static void GameThreadWaitForTask(const UE::Tasks::FTask& Task, bool bEmptyGameT
 					// process gamethread tasks if there are any
 					FTaskGraphInterface::Get().ProcessThreadUntilIdle(ENamedThreads::GameThread);
 				}
-				bDone = Task.Wait(FTimespan::FromMilliseconds(WaitTime));
+				bDone = CompletionEvent->Wait(FTimespan::FromMilliseconds(WaitTime));
 
 				RenderThreadTimeoutClock.Tick();
 
