ShaderCodeArchive detect likely 900K decompression failures and point message to web site with details

diff --git a/Engine/Source/Runtime/RenderCore/Private/ShaderCodeArchive.cpp b/Engine/Source/Runtime/RenderCore/Private/ShaderCodeArchive.cpp
--- a/Engine/Source/Runtime/RenderCore/Private/ShaderCodeArchive.cpp
+++ b/Engine/Source/Runtime/RenderCore/Private/ShaderCodeArchive.cpp
@@ -358,30 +358,58 @@ void FShaderUsageVisualizer::SaveShaderUsageBitmap(const FString& Name, EShaderP
 }
 #endif
 
+
+static FString GetDecompressionFailureExtraMessage()
+{
+	FString Brand = FPlatformMisc::GetCPUBrand();
+
+	// most common offender is 13900K
+	//	also seems to be a problem on 14900K and maybe 12900K
+	//	not sure if 700K is affected
+	// don't try to filter the brand string, just always log it if we hit a fatal error
+	#if 0
+	bool Is900K =
+	 ( Brand.Find(TEXT("900K")) != INDEX_NONE ) ||
+	 ( Brand.Find(TEXT("900-K")) != INDEX_NONE ) ||
+	 ( Brand.Find(TEXT("900 K")) != INDEX_NONE );
+
+	if ( ! Is900K )
+	{
+		return FString();
+	}
+	#endif
+
+	return FString::Printf(TEXT("The CPU (%s) may be unstable; for details see http://www.radgametools.com/oodleintel.htm"),*Brand);
+}
+
 void ShaderCodeArchive::DecompressShaderWithOodle(uint8* OutDecompressedShader, int64 UncompressedSize, const uint8* CompressedShaderCode, int64 CompressedSize)
 {
 	// Iostore always compresses with Oodle.
 	bool bSucceed = FCompression::UncompressMemory(NAME_Oodle, OutDecompressedShader, UncompressedSize, CompressedShaderCode, CompressedSize);
 	if (!bSucceed)
 	{
-		UE_LOG(LogShaderLibrary, Fatal, TEXT("ShaderCodeArchive::DecompressShader(): Could not decompress shader with Oodle"));
+		UE_LOG(LogShaderLibrary, Fatal, TEXT("ShaderCodeArchive::DecompressShader(): Could not decompress shader with Oodle. %s"),
+			*GetDecompressionFailureExtraMessage());
 	}
 }
 
 namespace
 {
-	void DecompressShadergroupWithOodleAndExtraLogging(const int32 GroupIndex, const FIoChunkId& GroupHash, FIoStoreShaderGroupEntry& Entry, const int32 ShaderIndex, const uint64 ShaderInGroupIndex, const FSHAHash& ShaderHash, uint8* OutDecompressedShader, int64 UncompressedSize, const uint8* CompressedShaderCode, int64 CompressedSize)
+	void DecompressShadergroupWithOodleAndExtraLogging(
+			const int32 GroupIndex, const FIoChunkId& GroupHash, FIoStoreShaderGroupEntry& Entry, const int32 ShaderIndex, const uint64 ShaderInGroupIndex, const FSHAHash& ShaderHash, 
+			uint8* OutDecompressedShader, int64 UncompressedSize, const uint8* CompressedShaderCode, int64 CompressedSize)
 	{
 		bool bSucceed = FCompression::UncompressMemory(NAME_Oodle, OutDecompressedShader, UncompressedSize, CompressedShaderCode, CompressedSize);
 		if (!bSucceed)
 		{
-			UE_LOG(LogShaderLibrary, Fatal, TEXT("DecompressShaderWithOodleAndExtraLogging(): Could not decompress shader group with Oodle. Group Index: %d Group IoStoreHash:%s Group NumShaders: %d Shader Index: %d Shader In-group Index: %d Shader Hash: %s"),
+			UE_LOG(LogShaderLibrary, Fatal, TEXT("DecompressShaderWithOodleAndExtraLogging(): Could not decompress shader group with Oodle. Group Index: %d Group IoStoreHash:%s Group NumShaders: %d Shader Index: %d Shader In-group Index: %d Shader Hash: %s. %s"),
 				GroupIndex,
 				*LexToString(GroupHash),
 				Entry.NumShaders,
 				ShaderIndex,
 				ShaderInGroupIndex,
-				*LexToString(ShaderHash)
+				*LexToString(ShaderHash),
+				*GetDecompressionFailureExtraMessage()
 				);
 		}
 	}
