Added new compiler flag to support compiling shaders with support for certain (rare) indirect draw setups on certain platforms. The flag is a generic "IndirectDraw" flag so that we can use it on all indirect draw shaders in the future.

#rnx

diff --git a/Engine/Source/Runtime/RenderCore/Public/ShaderCompilerCore.h b/Engine/Source/Runtime/RenderCore/Public/ShaderCompilerCore.h
--- a/Engine/Source/Runtime/RenderCore/Public/ShaderCompilerCore.h
+++ b/Engine/Source/Runtime/RenderCore/Public/ShaderCompilerCore.h
@@ -160,6 +160,9 @@ enum ECompilerFlags
 	// auto derivatives. If none are found, the shader will be marked with EShaderResourceUsageFlags::NoDerivativeOps, meaning that
 	// calling code can safely assume only provided derivatives are used.
 	CFLAG_CheckForDerivativeOps,
+	// Shader is used with indirect draws. This flag is currently used to fix a platform specific problem with certain (rare) indirect draw setups, but it is intended to be set for all indirect draw shaders in the future.
+	// Must not be used on shaders that are used with direct draws. Doing so might cause crashes or visual corruption on certain platforms.
+	CFLAG_IndirectDraw,
 	CFLAG_Max,
 };
 static_assert(CFLAG_Max < 64, "Out of bitfield space! Modify FShaderCompilerFlags");
