Fixed UseSplineMeshSceneResources to read cvar data from platform ini file when cooking

#rnx

diff --git a/Engine/Shaders/Private/SplineMeshCommon.ush b/Engine/Shaders/Private/SplineMeshCommon.ush
--- a/Engine/Shaders/Private/SplineMeshCommon.ush
+++ b/Engine/Shaders/Private/SplineMeshCommon.ush
@@ -51,8 +51,8 @@ half SplineMeshCalcSplineDistance(FSplineMeshShaderParams Params, float3 LocalPo
 float3 SplineMeshEvalSplinePos(FSplineMeshShaderParams Params, half SplineDist)
 {
 	float A = SplineDist;
-	float A2 = A  * A;
-	float A3 = A2 * A;
+	float A2 = A * A;
+	float A3 = A * A2;
 
 	return (((2*A3)-(3*A2)+1) * Params.StartPos) +
 		   ((A3-(2*A2)+A) * Params.StartTangent) +
diff --git a/Engine/Source/Runtime/RenderCore/Private/RenderUtils.cpp b/Engine/Source/Runtime/RenderCore/Private/RenderUtils.cpp
--- a/Engine/Source/Runtime/RenderCore/Private/RenderUtils.cpp
+++ b/Engine/Source/Runtime/RenderCore/Private/RenderUtils.cpp
@@ -1636,9 +1636,8 @@ bool UseSplineMeshSceneResources(const FStaticShaderPlatform Platform)
 	// NOTE: Mobile GPU can't currently support this feature because it doesn't include the payload extension in its instance data
 	if (UseGPUScene(Platform) && !IsMobilePlatform(Platform))
 	{
-		static const auto AllowSceneTexture = IConsoleManager::Get().FindTConsoleVariableDataInt(TEXT("r.SplineMesh.SceneTextures"));
-		static const bool bAllowSceneTexture = (AllowSceneTexture && AllowSceneTexture->GetValueOnAnyThread() != 0);
-		return bAllowSceneTexture;
+		static FShaderPlatformCachedIniValue<int32> CVar(TEXT("r.SplineMesh.SceneTextures"));
+		return CVar.Get(Platform) > 0;
 	}
 
 	return false;
