[shaders] versioning/ddc key fixes
- include full hash of uniform buffer metadata in shadermap key instead of just hashing rows/columns/type (the latter doesn't properly account for nested struct layout changes)
- include shader format version in per-shader key (this is already included in the shadermap key, not including it in the per-shader key as well is a dangerous oversight)

diff --git a/Engine/Source/Runtime/RenderCore/Private/ShaderCore.cpp b/Engine/Source/Runtime/RenderCore/Private/ShaderCore.cpp
--- a/Engine/Source/Runtime/RenderCore/Private/ShaderCore.cpp
+++ b/Engine/Source/Runtime/RenderCore/Private/ShaderCore.cpp
@@ -3068,93 +3068,6 @@ private:
 	TArray<FPlatformCache, TInlineAllocator<8>> Platforms;
 };
 
-/** Efficient lookup to find FShaderParametersMetadata members by name pointer */
-class FShaderParameterMemberLookup 
-{
-public:
-	FShaderParameterMemberLookup(TLinkedList<FShaderParametersMetadata*>& ShaderParameters)
-	{
-		for (FShaderParametersMetadata* Struct : ShaderParameters)
-		{
-			Map.Add(Struct->GetShaderVariableName(), Struct->GetMembers());
-		}
-	}
-
-	const TConstArrayView<FShaderParametersMetadata::FMember>* FindMembersByPointer(const TCHAR* ShaderVariableName) const
-	{
-		return Map.Find(ShaderVariableName);
-	}
-
-private:
-	TMap<const void*, TConstArrayView<FShaderParametersMetadata::FMember>> Map;
-};
-
-/** Cache providing a FShaderParameterMemberLookup for the current FShaderParametersMetadata::GetStructList */
-class FFShaderParameterPointerLookupCache
-{
-public:
-	TSharedPtr<const FShaderParameterMemberLookup> Get()
-	{
-		TLinkedList<FShaderParametersMetadata*>* CurrentHead = FShaderParametersMetadata::GetStructList();
-		check(CurrentHead);
-
-		{
-			FReadScopeLock ReadScope(Lock);
-			if (CurrentHead == CachedHead)
-			{
-				return CachedLookup;
-			}
-		}
-
-		TSharedPtr<FShaderParameterMemberLookup> NewLookup = MakeShared<FShaderParameterMemberLookup>(*CurrentHead);
-
-		FWriteScopeLock WriteScope(Lock);
-		CachedHead = CurrentHead;
-		CachedLookup = NewLookup;
-
-		return NewLookup;
-	}
-
-private:
-	FRWLock Lock;
-	TLinkedList<FShaderParametersMetadata*>* CachedHead = nullptr;
-	TSharedPtr<const FShaderParameterMemberLookup> CachedLookup;
-};
-
-static FFShaderParameterPointerLookupCache GShaderParameterMemberLookupCache;
-
-
-// This copy is only used internally once - it could be inlined once the public API version is removed
-void SerializeUniformBufferInfo_Internal(FShaderSaveArchive& Ar, const TArray<const TCHAR*>& UniformBufferNames)
-{
-	if (UniformBufferNames.IsEmpty())
-	{
-		return;
-	}
-
-	TSharedPtr<const FShaderParameterMemberLookup> ShaderParameterMembers = GShaderParameterMemberLookupCache.Get();
-
-	for (const TCHAR* UniformBufferName : UniformBufferNames)
-	{
-		if (const TConstArrayView<FShaderParametersMetadata::FMember>* Members = ShaderParameterMembers->FindMembersByPointer(UniformBufferName))
-		{
-			// Serialize information about the struct layout so we can detect when it changes
-			int32 NumMembers = Members->Num();
-			// Serializing with NULL so that FShaderSaveArchive will record the length without causing an actual data serialization
-			Ar.Serialize(nullptr, NumMembers);
-
-			for (const FShaderParametersMetadata::FMember& Member : *Members)
-			{
-				// Note: Only comparing number of floats used by each member and type, so this can be tricked (eg. swapping two equal size and type members)
-				int32 MemberSize = Member.GetNumColumns() * Member.GetNumRows();
-				Ar.Serialize(nullptr, MemberSize);
-				int32 MemberType = (int32)Member.GetBaseType();
-				Ar.Serialize(nullptr, MemberType);
-			}
-		}
-	}
-}
-
 } // anonymous namespace
 
 
@@ -3291,7 +3204,11 @@ void AppendKeyStringShaderDependencies(
 		Algo::Sort(SortedUniformBufferNames, FUniformBufferNameSortOrder());
 
 		// Save uniform buffer member info so we can detect when layout has changed
-		SerializeUniformBufferInfo_Internal(SaveArchive, SortedUniformBufferNames);
+		for (const TCHAR* UniformBufferName : SortedUniformBufferNames)
+		{
+			FShaderParametersMetadata* UniformBufferMetadata = FindUniformBufferStructByName(UniformBufferName);
+			UniformBufferMetadata->SerializeLayout(SaveArchive);
+		}
 
 		SerializationHistory.AppendKeyString(OutKeyString);
 	}
@@ -3660,6 +3577,9 @@ FShaderCommonCompileJob::FInputHash FShaderCompileJob::GetInputHash()
 		int32 FShaderCompilerOutputStructVersionLocal = FShaderCompilerOutputStructVersion;
 		Hasher << FShaderCompilerOutputStructVersionLocal;
 
+		uint32 FormatVersion = GetTargetPlatformManagerRef().ShaderFormatVersion(Input.ShaderFormat);
+		Hasher << FormatVersion;
+		
 		FShaderTarget Target = Input.Target;
 		Hasher << Target;
 		Hasher << Input.EntryPointName;
@@ -3705,6 +3625,9 @@ FShaderCommonCompileJob::FInputHash FShaderCompileJob::GetInputHash()
 			int32 FShaderCompilerOutputStructVersionLocal = FShaderCompilerOutputStructVersion;
 			Archive << FShaderCompilerOutputStructVersionLocal;
 
+			uint32 FormatVersion = GetTargetPlatformManagerRef().ShaderFormatVersion(Input.ShaderFormat);
+			Archive << FormatVersion;
+
 			// Don't include debug group name in the hashing; this drastically worsens our cache hit rate
 			FString DebugGroupNameTmp(MoveTemp(Input.DebugGroupName));
 			Archive << Input;
