Add strongly typed EShaderOptionalDataKey to access shader optional data instead of adhoc characters

diff --git a/Engine/Source/Developer/Apple/MetalShaderFormat/Private/MetalShaderCompiler.cpp b/Engine/Source/Developer/Apple/MetalShaderFormat/Private/MetalShaderCompiler.cpp
--- a/Engine/Source/Developer/Apple/MetalShaderFormat/Private/MetalShaderCompiler.cpp
+++ b/Engine/Source/Developer/Apple/MetalShaderFormat/Private/MetalShaderCompiler.cpp
@@ -712,14 +712,14 @@ void BuildMetalShaderOutput(
 
 			if (ShaderInput.Environment.CompilerFlags.Contains(CFLAG_Archive))
 			{
-				ShaderOutput.ShaderCode.AddOptionalData('o', Bytecode.ObjectFile.GetData(), Bytecode.ObjectFile.Num());
+				ShaderOutput.ShaderCode.AddOptionalData(EShaderOptionalDataKey::ObjectFile, Bytecode.ObjectFile.GetData(), Bytecode.ObjectFile.Num());
 			}
 			
 			if (bDebugInfoSucceded && !ShaderInput.Environment.CompilerFlags.Contains(CFLAG_Archive) && DebugCode.CompressedData.Num())
 			{
-				ShaderOutput.ShaderCode.AddOptionalData('z', DebugCode.CompressedData.GetData(), DebugCode.CompressedData.Num());
-				ShaderOutput.ShaderCode.AddOptionalData('p', TCHAR_TO_UTF8(*Bytecode.NativePath));
-				ShaderOutput.ShaderCode.AddOptionalData('u', (const uint8*)&DebugCode.UncompressedSize, sizeof(DebugCode.UncompressedSize));
+				ShaderOutput.ShaderCode.AddOptionalData(EShaderOptionalDataKey::CompressedDebugCode, DebugCode.CompressedData.GetData(), DebugCode.CompressedData.Num());
+				ShaderOutput.ShaderCode.AddOptionalData(EShaderOptionalDataKey::NativePath, TCHAR_TO_UTF8(*Bytecode.NativePath));
+				ShaderOutput.ShaderCode.AddOptionalData(EShaderOptionalDataKey::UncompressedSize, (const uint8*)&DebugCode.UncompressedSize, sizeof(DebugCode.UncompressedSize));
 			}
 			
 			if (ShaderInput.Environment.CompilerFlags.Contains(CFLAG_ExtraShaderData))
@@ -728,14 +728,14 @@ void BuildMetalShaderOutput(
 				ShaderOutput.ShaderCode.AddOptionalData(FShaderCodeName::Key, TCHAR_TO_UTF8(*ShaderInput.GenerateShaderName()));
 				if (DebugCode.CompressedData.Num() == 0)
 				{
-					ShaderOutput.ShaderCode.AddOptionalData('c', TCHAR_TO_UTF8(*MetalCode));
-					ShaderOutput.ShaderCode.AddOptionalData('p', TCHAR_TO_UTF8(*Bytecode.NativePath));
+					ShaderOutput.ShaderCode.AddOptionalData(EShaderOptionalDataKey::SourceCode, TCHAR_TO_UTF8(*MetalCode));
+					ShaderOutput.ShaderCode.AddOptionalData(EShaderOptionalDataKey::NativePath, TCHAR_TO_UTF8(*Bytecode.NativePath));
 				}
 			}
 			else if (ShaderInput.Environment.CompilerFlags.Contains(CFLAG_Archive))
 			{
-				ShaderOutput.ShaderCode.AddOptionalData('c', TCHAR_TO_UTF8(*MetalCode));
-				ShaderOutput.ShaderCode.AddOptionalData('p', TCHAR_TO_UTF8(*Bytecode.NativePath));
+				ShaderOutput.ShaderCode.AddOptionalData(EShaderOptionalDataKey::SourceCode, TCHAR_TO_UTF8(*MetalCode));
+				ShaderOutput.ShaderCode.AddOptionalData(EShaderOptionalDataKey::NativePath, TCHAR_TO_UTF8(*Bytecode.NativePath));
 			}
 			
 			ShaderOutput.NumTextureSamplers = Header.Bindings.NumSamplers;
@@ -1034,11 +1034,11 @@ bool StripShader_Metal(TArray<uint8>& Code, class FString const& DebugPath, bool
 			TArray<uint8> SourceCode;
 			SourceCode.Append(SourceCodePtr, ShaderCode.GetActualShaderCodeSize() - CodeOffset);
 			
-			const ANSICHAR* ShaderSource = ShaderCode.FindOptionalData('c');
+			const ANSICHAR* ShaderSource = ShaderCode.FindOptionalData(EShaderOptionalDataKey::SourceCode);
 			const size_t ShaderSourceLength = ShaderSource ? FCStringAnsi::Strlen(ShaderSource) : 0;
 			bool const bHasShaderSource = ShaderSourceLength > 0;
 			
-			const ANSICHAR* ShaderPath = ShaderCode.FindOptionalData('p');
+			const ANSICHAR* ShaderPath = ShaderCode.FindOptionalData(EShaderOptionalDataKey::NativePath);
 			bool const bHasShaderPath = (ShaderPath && FCStringAnsi::Strlen(ShaderPath) > 0);
 			
 			if (bHasShaderSource && bHasShaderPath)
@@ -1068,7 +1068,7 @@ bool StripShader_Metal(TArray<uint8>& Code, class FString const& DebugPath, bool
 			if (bNative)
 			{
 				int32 ObjectSize = 0;
-				const uint8* ShaderObject = ShaderCode.FindOptionalDataAndSize('o', ObjectSize);
+				const uint8* ShaderObject = ShaderCode.FindOptionalDataAndSize(EShaderOptionalDataKey::ObjectFile, ObjectSize);
 				
 				// If ShaderObject and ObjectSize is zero then the code has already been stripped - source code should be the byte code
 				if(ShaderObject && ObjectSize)
@@ -1140,7 +1140,7 @@ uint64 AppendShader_Metal(FString const& WorkingDir, const FSHAHash& Hash, TArra
 				
 				// Copy the non-optional shader bytecode
 				int32 ObjectCodeDataSize = 0;
-				uint8 const* Object = ShaderCode.FindOptionalDataAndSize('o', ObjectCodeDataSize);
+				uint8 const* Object = ShaderCode.FindOptionalDataAndSize(EShaderOptionalDataKey::ObjectFile, ObjectCodeDataSize);
 
 				// 'o' segment missing this is a pre stripped shader
 				if(!Object)
diff --git a/Engine/Source/Developer/Apple/MetalShaderFormat/Private/MetalShaderFormat.cpp b/Engine/Source/Developer/Apple/MetalShaderFormat/Private/MetalShaderFormat.cpp
--- a/Engine/Source/Developer/Apple/MetalShaderFormat/Private/MetalShaderFormat.cpp
+++ b/Engine/Source/Developer/Apple/MetalShaderFormat/Private/MetalShaderFormat.cpp
@@ -37,7 +37,7 @@ extern uint64 AppendShader_Metal(class FString const& ArchivePath, const FSHAHas
 extern bool FinalizeLibrary_Metal(class FName const& Format, class FString const& ArchivePath, class FString const& LibraryPath, TSet<uint64> const& Shaders, class FString const& DebugOutputDir);
 
 /** Version for shader format, this becomes part of the DDC key. */
-static const FGuid UE_SHADER_METAL_VER = FGuid("222AACA0-E61A-490A-A5A3-0F3957CC248F");
+static const FGuid UE_SHADER_METAL_VER = FGuid("282D7BF5-C5F8-47B7-B065-E165D60FC17B");
 
 class FMetalShaderFormat : public IShaderFormat
 {
diff --git a/Engine/Source/Developer/Windows/ShaderFormatD3D/Private/D3DShaderCompilerDXC.cpp b/Engine/Source/Developer/Windows/ShaderFormatD3D/Private/D3DShaderCompilerDXC.cpp
--- a/Engine/Source/Developer/Windows/ShaderFormatD3D/Private/D3DShaderCompilerDXC.cpp
+++ b/Engine/Source/Developer/Windows/ShaderFormatD3D/Private/D3DShaderCompilerDXC.cpp
@@ -1250,7 +1250,7 @@ bool CompileAndProcessD3DShaderDXC(
 				if (ShaderModel >= ED3DShaderModel::SM6_0)
 				{
 					uint8 IsSM6 = 1;
-					ShaderCode.AddOptionalData('6', &IsSM6, 1);
+					ShaderCode.AddOptionalData(EShaderOptionalDataKey::ShaderModel6, &IsSM6, 1);
 				}
 			};
 
diff --git a/Engine/Source/Runtime/Apple/MetalRHI/Private/Shaders/Types/Templates/MetalBaseShader.h b/Engine/Source/Runtime/Apple/MetalRHI/Private/Shaders/Types/Templates/MetalBaseShader.h
--- a/Engine/Source/Runtime/Apple/MetalRHI/Private/Shaders/Types/Templates/MetalBaseShader.h
+++ b/Engine/Source/Runtime/Apple/MetalRHI/Private/Shaders/Types/Templates/MetalBaseShader.h
@@ -167,7 +167,7 @@ void TMetalBaseShader<BaseResourceType, ShaderType>::Init(TArrayView<const uint8
 
 	bool bOfflineCompile = (OfflineCompiledFlag > 0);
 
-	const ANSICHAR* ShaderSource = ShaderCode.FindOptionalData('c');
+	const ANSICHAR* ShaderSource = ShaderCode.FindOptionalData(EShaderOptionalDataKey::SourceCode);
 	bool bHasShaderSource = (ShaderSource && FCStringAnsi::Strlen(ShaderSource) > 0);
 
 	static bool bForceTextShaders = FMetalCommandQueue::SupportsFeature(EMetalFeaturesGPUTrace);
@@ -175,8 +175,8 @@ void TMetalBaseShader<BaseResourceType, ShaderType>::Init(TArrayView<const uint8
 	{
 		int32 LZMASourceSize = 0;
 		int32 SourceSize = 0;
-		const uint8* LZMASource = ShaderCode.FindOptionalDataAndSize('z', LZMASourceSize);
-		const uint8* UnSourceLen = ShaderCode.FindOptionalDataAndSize('u', SourceSize);
+		const uint8* LZMASource = ShaderCode.FindOptionalDataAndSize(EShaderOptionalDataKey::CompressedDebugCode, LZMASourceSize);
+		const uint8* UnSourceLen = ShaderCode.FindOptionalDataAndSize(EShaderOptionalDataKey::UncompressedSize, SourceSize);
 		if (LZMASource && LZMASourceSize > 0 && UnSourceLen && SourceSize == sizeof(uint32))
 		{
 			CompressedSource.Append(LZMASource, LZMASourceSize);
diff --git a/Engine/Source/Runtime/RenderCore/Public/ShaderCore.h b/Engine/Source/Runtime/RenderCore/Public/ShaderCore.h
--- a/Engine/Source/Runtime/RenderCore/Public/ShaderCore.h
+++ b/Engine/Source/Runtime/RenderCore/Public/ShaderCore.h
@@ -704,6 +704,26 @@ struct FSharedShaderCompilerEnvironment final : public FShaderCompilerEnvironmen
 	virtual ~FSharedShaderCompilerEnvironment() = default;
 };
 
+enum class EShaderOptionalDataKey : uint8
+{
+	AttributeInputs      = uint8('i'),
+	AttributeOutputs     = uint8('o'),
+	CompressedDebugCode  = uint8('z'),
+	Diagnostic           = uint8('D'),
+	Features             = uint8('x'),
+	Name                 = uint8('n'),
+	NativePath           = uint8('P'),
+	ObjectFile           = uint8('O'),
+	PackedResourceCounts = uint8('p'),
+	ResourceMasks        = uint8('m'),
+	ShaderModel6         = uint8('6'),
+	SourceCode           = uint8('c'),
+	UncompressedSize     = uint8('U'),
+	UniformBuffers       = uint8('u'),
+	Validation           = uint8('V'),
+	VendorExtension      = uint8('v'),
+};
+
 enum class EShaderResourceUsageFlags : uint8
 {
 	GlobalUniformBuffer   = 1 << 0,
@@ -718,7 +738,7 @@ ENUM_CLASS_FLAGS(EShaderResourceUsageFlags)
 struct FShaderCodePackedResourceCounts
 {
 	// for FindOptionalData() and AddOptionalData()
-	static const uint8 Key = 'p';
+	static const EShaderOptionalDataKey Key = EShaderOptionalDataKey::PackedResourceCounts;
 
 	EShaderResourceUsageFlags UsageFlags;
 	uint8 NumSamplers;
@@ -730,7 +750,7 @@ struct FShaderCodePackedResourceCounts
 struct FShaderCodeResourceMasks
 {
 	// for FindOptionalData() and AddOptionalData()
-	static const uint8 Key = 'm';
+	static const EShaderOptionalDataKey Key = EShaderOptionalDataKey::ResourceMasks;
 
 	uint32 UAVMask; // Mask of UAVs bound
 };
@@ -754,7 +774,7 @@ ENUM_CLASS_FLAGS(EShaderCodeFeatures);
 struct FShaderCodeFeatures
 {
 	// for FindOptionalData() and AddOptionalData()
-	static const uint8 Key = 'x';
+	static const EShaderOptionalDataKey Key = EShaderOptionalDataKey::Features;
 
 	EShaderCodeFeatures CodeFeatures = EShaderCodeFeatures::None;
 };
@@ -762,14 +782,14 @@ struct FShaderCodeFeatures
 // if this changes you need to make sure all shaders get invalidated
 struct FShaderCodeName
 {
-	static const uint8 Key = 'n';
+	static const EShaderOptionalDataKey Key = EShaderOptionalDataKey::Name;
 
 	// We store the straight ANSICHAR zero-terminated string
 };
 
 struct FShaderCodeUniformBuffers
 {
-	static const uint8 Key = 'u';
+	static const EShaderOptionalDataKey Key = EShaderOptionalDataKey::UniformBuffers;
 	// We store an array of FString objects
 };
 
@@ -777,7 +797,7 @@ struct FShaderCodeUniformBuffers
 struct FShaderCodeVendorExtension
 {
 	// for FindOptionalData() and AddOptionalData()
-	static const uint8 Key = 'v';
+	static const EShaderOptionalDataKey Key = EShaderOptionalDataKey::VendorExtension;
 
 	EGpuVendorId VendorId = EGpuVendorId::NotQueried;
 	FParameterAllocation Parameter;
@@ -825,7 +845,7 @@ inline FArchive& operator<<(FArchive& Ar, FShaderCodeValidationUBSize& ShaderCod
 struct FShaderCodeValidationExtension
 {
 	// for FindOptionalData() and AddOptionalData()
-	static constexpr uint8 Key = 'V';
+	static constexpr EShaderOptionalDataKey Key = EShaderOptionalDataKey::Validation;
 	static constexpr uint16 StaticVersion = 0;
 
 	TArray<FShaderCodeValidationStride> ShaderCodeValidationStride;
@@ -859,7 +879,7 @@ inline FArchive& operator<<(FArchive& Ar, FShaderDiagnosticData& ShaderCodeDiagn
 struct FShaderDiagnosticExtension
 {
 	// for FindOptionalData() and AddOptionalData()
-	static constexpr uint8 Key = 'D';
+	static constexpr EShaderOptionalDataKey Key = EShaderOptionalDataKey::Diagnostic;
 	static constexpr uint16 StaticVersion = 0;
 
 	TArray<FShaderDiagnosticData> ShaderDiagnosticDatas;
@@ -912,7 +932,7 @@ public:
 
 	// @param InKey e.g. FShaderCodePackedResourceCounts::Key
 	// @return 0 if not found
-	const uint8* FindOptionalData(uint8 InKey, uint8 ValueSize) const
+	const uint8* FindOptionalData(EShaderOptionalDataKey InKey, uint8 ValueSize) const
 	{
 		check(ValueSize);
 
@@ -927,7 +947,7 @@ public:
 
 		while(Current < End)
 		{
-			uint8 Key = *Current++;
+			EShaderOptionalDataKey Key = EShaderOptionalDataKey(*Current++);
 			uint32 Size = *((const unaligned_uint32*)Current);
 			Current += sizeof(Size);
 
@@ -942,7 +962,7 @@ public:
 		return 0;
 	}
 
-	const ANSICHAR* FindOptionalData(uint8 InKey) const
+	const ANSICHAR* FindOptionalData(EShaderOptionalDataKey InKey) const
 	{
 		check(ShaderCode.Num() >= 4);
 
@@ -957,7 +977,7 @@ public:
 
 		while(Current < End)
 		{
-			uint8 Key = *Current++;
+			EShaderOptionalDataKey Key = EShaderOptionalDataKey(*Current++);
 			uint32 Size = *((const unaligned_uint32*)Current);
 			Current += sizeof(Size);
 
@@ -973,7 +993,7 @@ public:
 	}
 
 	// Returns nullptr and Size -1 if key was not found
-	const uint8* FindOptionalDataAndSize(uint8 InKey, int32& OutSize) const
+	const uint8* FindOptionalDataAndSize(EShaderOptionalDataKey InKey, int32& OutSize) const
 	{
 		check(ShaderCode.Num() >= 4);
 
@@ -988,7 +1008,7 @@ public:
 
 		while (Current < End)
 		{
-			uint8 Key = *Current++;
+			EShaderOptionalDataKey Key = EShaderOptionalDataKey(*Current++);
 			uint32 Size = *((const unaligned_uint32*)Current);
 			Current += sizeof(Size);
 
@@ -1144,14 +1164,14 @@ public:
 	// can be called after the non optional data was stored in ShaderData
 	// @param Key uint8 to save memory so max 255, e.g. FShaderCodePackedResourceCounts::Key
 	// @param Size >0, only restriction is that sum of all optional data values must be < 4GB
-	void AddOptionalData(uint8 Key, const uint8* ValuePtr, uint32 ValueSize)
+	void AddOptionalData(EShaderOptionalDataKey Key, const uint8* ValuePtr, uint32 ValueSize)
 	{
 		check(ValuePtr);
 
 		// don't add after Finalize happened
 		check(OptionalDataSize >= 0);
 
-		ShaderCodeWithOptionalData.Add(Key);
+		ShaderCodeWithOptionalData.Add(uint8(Key));
 		ShaderCodeWithOptionalData.Append((const uint8*)&ValueSize, sizeof(ValueSize));
 		ShaderCodeWithOptionalData.Append(ValuePtr, ValueSize);
 		OptionalDataSize += sizeof(uint8) + sizeof(ValueSize) + (uint32)ValueSize;
@@ -1160,7 +1180,7 @@ public:
 	// Note: we don't hash the optional attachments in GenerateOutputHash() as they would prevent sharing (e.g. many material share the save VS)
 	// convenience, silently drops the data if string is too long
 	// @param e.g. 'n' for the ShaderSourceFileName
-	void AddOptionalData(uint8 Key, const ANSICHAR* InString)
+	void AddOptionalData(EShaderOptionalDataKey Key, const ANSICHAR* InString)
 	{
 		uint32 Size = FCStringAnsi::Strlen(InString) + 1;
 		AddOptionalData(Key, (uint8*)InString, Size);
diff --git a/Engine/Source/Runtime/Windows/D3D11RHI/Private/D3D11Shaders.cpp b/Engine/Source/Runtime/Windows/D3D11RHI/Private/D3D11Shaders.cpp
--- a/Engine/Source/Runtime/Windows/D3D11RHI/Private/D3D11Shaders.cpp
+++ b/Engine/Source/Runtime/Windows/D3D11RHI/Private/D3D11Shaders.cpp
@@ -52,7 +52,7 @@ static inline void ReadShaderOptionalData(FShaderCodeReader& InShaderCode, TShad
 	OutShader.bShaderNeedsGlobalConstantBuffer = EnumHasAnyFlags(PackedResourceCounts->UsageFlags, EShaderResourceUsageFlags::GlobalUniformBuffer);
 
 	int32 IsSm6ShaderSize = 1;
-	const uint8* IsSm6Shader = InShaderCode.FindOptionalData('6', IsSm6ShaderSize);
+	const uint8* IsSm6Shader = InShaderCode.FindOptionalData(EShaderOptionalDataKey::ShaderModel6, IsSm6ShaderSize);
 	OutShader.bIsSm6Shader = IsSm6Shader && IsSm6ShaderSize && *IsSm6Shader;
 
 	UE::RHICore::SetupShaderCodeValidationData(&OutShader, InShaderCode);
