Implemented recursive render command pipe command support. Render commands that are enqueued while replaying another command execute immediately. This mirrors the behavior of render thread commands. Also added an ensure to catch when a command is enqueued to pipe B while replaying pipe A, as this is not a supported scenario.

diff --git a/Engine/Source/Runtime/RenderCore/Private/RenderingThread.cpp b/Engine/Source/Runtime/RenderCore/Private/RenderingThread.cpp
--- a/Engine/Source/Runtime/RenderCore/Private/RenderingThread.cpp
+++ b/Engine/Source/Runtime/RenderCore/Private/RenderingThread.cpp
@@ -2031,6 +2031,8 @@ inline bool HasBitsSet(const FRenderCommandPipeBitArray& Bits)
 
 namespace UE::RenderCommandPipe
 {
+	static thread_local FRenderCommandPipe* ReplayingPipe = nullptr;
+
 	void Initialize()
 	{
 		GRenderCommandPipeRegistry.Initialize();
@@ -2046,6 +2048,11 @@ namespace UE::RenderCommandPipe
 		return GRenderCommandPipeRegistry.IsReplaying();
 	}
 
+	bool IsReplaying(const FRenderCommandPipe& Pipe)
+	{
+		return ReplayingPipe == &Pipe;
+	}
+
 	void StartRecording()
 	{
 		GRenderCommandPipeRegistry.StartRecording();
@@ -2129,55 +2136,63 @@ FRenderCommandPipe::~FRenderCommandPipe()
 	Frame_RenderThread = nullptr;
 }
 
+void FRenderCommandPipe::ExecuteCommand(FFunctionVariant&& FunctionVariant, const TCHAR* CommandName, uint32& CommandSpecId)
+{
+	TRACE_CPUPROFILER_EVENT_SCOPE_USE_ON_CHANNEL(CommandSpecId, CommandName, CommandEventScope, RenderCommandsChannel, true);
+	if (FCommandListFunction* Function = FunctionVariant.TryGet<FCommandListFunction>())
+	{
+		if (!Frame_RenderThread->RHICmdList)
+		{
+			FRHICommandList* RHICmdList = new FRHICommandList(FRHIGPUMask::All());
+			RHICmdList->SwitchPipeline(ERHIPipeline::Graphics);
+			Frame_RenderThread->RHICmdList = RHICmdList;
+		}
+
+		(*Function)(*Frame_RenderThread->RHICmdList);
+	}
+	else
+	{
+		FunctionVariant.Get<FEmptyFunction>()();
+	}
+}
+
 void FRenderCommandPipe::EnqueueAndLaunch(FFunctionVariant&& FunctionVariant, const TCHAR* CommandName, uint32& CommandSpecId)
 {
+	ensureMsgf(!UE::RenderCommandPipe::ReplayingPipe, TEXT("Attempting to launch render command to render command pipe %s from another pipe %s"), Name, UE::RenderCommandPipe::ReplayingPipe->Name);
+
 	bool bWasEmpty = Frame_GameThread->Queue.IsEmpty();
 	Frame_GameThread->Queue.Emplace(MoveTemp(FunctionVariant), CommandName, CommandSpecId);
 	NumInFlightCommands.fetch_add(1, std::memory_order_relaxed);
-
+	
 	if (bWasEmpty)
 	{
 		TRACE_CPUPROFILER_EVENT_SCOPE_ON_CHANNEL_STR("RenderCommandPipe LaunchTask", RenderCommandsChannel)
-
+	
 		Frame_GameThread->Pipe.Launch(Name, [this]
 		{
 			check(Frame_RenderThread);
 			TRACE_CPUPROFILER_EVENT_SCOPE_ON_CHANNEL_STR("RenderCommandPipe ReplayCommands", RenderCommandsChannel)
 			SCOPED_NAMED_EVENT_TCHAR(Name, FColor::Magenta);
 			FOptionalTaskTagScope Scope(ETaskTag::EParallelRenderingThread);
-
+	
 			TArray<FCommand> PoppedQueue;
-
+	
 			Mutex.Lock();
 			PoppedQueue = MoveTemp(Frame_RenderThread->Queue);
 			Frame_RenderThread->Queue.Reserve(128);
 			Mutex.Unlock();
-
+	
+			FRenderCommandPipe* const PreviousReplayingPipe = UE::RenderCommandPipe::ReplayingPipe;
+			UE::RenderCommandPipe::ReplayingPipe = this;
+	
 			for (FCommand& Command : PoppedQueue)
 			{
-				TRACE_CPUPROFILER_EVENT_SCOPE_USE_ON_CHANNEL(*Command.SpecId, Command.Name, CommandEventScope, RenderCommandsChannel, true);
-
-				if (FCommandListFunction* Function = Command.Function.TryGet<FCommandListFunction>())
-				{
-					if (!Frame_RenderThread->RHICmdList)
-					{
-						FRHICommandList* RHICmdList = new FRHICommandList(FRHIGPUMask::All());
-						RHICmdList->SwitchPipeline(ERHIPipeline::Graphics);
-						Frame_RenderThread->RHICmdList = RHICmdList;
-					}
-
-					(*Function)(*Frame_RenderThread->RHICmdList);
-				}
-				else
-				{
-					Command.Function.Get<FEmptyFunction>()();
-				}
-
-				Command.Function = {};
+				ExecuteCommand(MoveTemp(Command.Function), Command.Name, *Command.SpecId);
 			}
-
+	
+			UE::RenderCommandPipe::ReplayingPipe = PreviousReplayingPipe;
 			NumInFlightCommands.fetch_sub(PoppedQueue.Num(), std::memory_order_release);
-
+	
 		}, Frame_GameThread->TaskEvent);
 	}
 }
\ No newline at end of file
diff --git a/Engine/Source/Runtime/RenderCore/Public/RenderingThread.h b/Engine/Source/Runtime/RenderCore/Public/RenderingThread.h
--- a/Engine/Source/Runtime/RenderCore/Public/RenderingThread.h
+++ b/Engine/Source/Runtime/RenderCore/Public/RenderingThread.h
@@ -311,6 +311,9 @@ namespace UE::RenderCommandPipe
 	// [Render Thread (Parallel)] Returns whether any render command pipes are currently replaying commands on the render thread timeline.
 	extern RENDERCORE_API bool IsReplaying();
 
+	// [Render Thread (Parallel)] Returns whether the specific render command pipe is replaying.
+	extern RENDERCORE_API bool IsReplaying(const FRenderCommandPipe& Pipe);
+
 	// [Game Thread] Starts recording render commands into pipes. Returns whether the operation succeeded.
 	extern RENDERCORE_API void StartRecording();
 	extern RENDERCORE_API void StartRecording(const FRenderCommandPipeBitArray& PipeBits);
@@ -447,6 +450,14 @@ public:
 		if (GRenderCommandPipeMode == ERenderCommandPipeMode::All && Pipe)
 		{
 			UE::TScopeLock Lock(Pipe->Mutex);
+
+			// Execute the function directly if this is being called recursively from within another pipe command.
+			if (UE::RenderCommandPipe::IsReplaying(*Pipe))
+			{
+				Pipe->ExecuteCommand(MoveTemp(Function), RenderCommandTag::GetName(), RenderCommandTag::GetSpecId());
+				return;
+			}
+
 			if (Pipe->Frame_GameThread)
 			{
 				Pipe->EnqueueAndLaunch(MoveTemp(Function), RenderCommandTag::GetName(), RenderCommandTag::GetSpecId());
@@ -469,6 +480,14 @@ public:
 		if (GRenderCommandPipeMode == ERenderCommandPipeMode::All && Pipe)
 		{
 			UE::TScopeLock Lock(Pipe->Mutex);
+
+			// Execute the function directly if this is being called recursively from within another pipe command.
+			if (UE::RenderCommandPipe::IsReplaying(*Pipe))
+			{
+				Pipe->ExecuteCommand(MoveTemp(Function), RenderCommandTag::GetName(), RenderCommandTag::GetSpecId());
+				return;
+			}
+
 			if (Pipe->Frame_GameThread)
 			{
 				Pipe->EnqueueAndLaunch(MoveTemp(Function), RenderCommandTag::GetName(), RenderCommandTag::GetSpecId());
@@ -508,6 +527,18 @@ private:
 		EnqueueAndLaunch(FFunctionVariant(TInPlaceType<FEmptyFunction>(), MoveTemp(Function)), CommandName, CommandSpecId);
 	}
 
+	RENDERCORE_API void ExecuteCommand(FFunctionVariant&& FunctionVariant, const TCHAR* CommandName, uint32& CommandSpecId);
+
+	void ExecuteCommand(FCommandListFunction&& Function, const TCHAR* CommandName, uint32& CommandSpecId)
+	{
+		ExecuteCommand(FFunctionVariant(TInPlaceType<FCommandListFunction>(), MoveTemp(Function)), CommandName, CommandSpecId);
+	}
+
+	void ExecuteCommand(FEmptyFunction&& Function, const TCHAR* CommandName, uint32& CommandSpecId)
+	{
+		ExecuteCommand(FFunctionVariant(TInPlaceType<FEmptyFunction>(), MoveTemp(Function)), CommandName, CommandSpecId);
+	}
+
 	struct FCommand
 	{
 		FCommand(FFunctionVariant&& InFunction, const TCHAR* InName, uint32& InOutSpecId)
