RDG Refactor to Improve Parallelism
 - Subresource state tracking manipulates less data by always using pointers instead of copying full states.
 - Offloaded pooled texture / buffer allocations to an async task.
 - Removed dependency between pooled resource allocations and collecting resource barriers. The 'finalize' phase now gathers both first and last barriers for a resource.
 - Refactors RDG allocator implementation to make things simpler.
 - Moved async tasks within RDG to use the setup task API.
 - Moved culling into the async setup queue.

diff --git a/Engine/Source/Runtime/Core/Private/Misc/MemStack.cpp b/Engine/Source/Runtime/Core/Private/Misc/MemStack.cpp
--- a/Engine/Source/Runtime/Core/Private/Misc/MemStack.cpp
+++ b/Engine/Source/Runtime/Core/Private/Misc/MemStack.cpp
@@ -309,12 +309,13 @@ void FPageAllocator::UpdateStats()
 	FMemStack implementation.
 -----------------------------------------------------------------------------*/
  
-FMemStackBase::FMemStackBase()
+FMemStackBase::FMemStackBase(EPageSize InPageSize)
 : Top(nullptr)
 , End(nullptr)
 , TopChunk(nullptr)
 , TopMark(nullptr)
 , NumMarks(0)
+, PageSize(InPageSize)
 , bShouldEnforceAllocMarks(false)
 {
 	// By fetching the FPageAllocator singleton's address here we can guarantee
@@ -345,7 +346,7 @@ void FMemStackBase::AllocateNewChunk(int32 MinSize)
 	// Create new chunk.
 	int32 TotalSize = MinSize + (int32)sizeof(FTaggedMemory);
 	uint32 AllocSize;
-	if (TopChunk || TotalSize > FPageAllocator::SmallPageSize)
+	if (TopChunk || TotalSize > FPageAllocator::SmallPageSize || PageSize == EPageSize::Large)
 	{
 		AllocSize = AlignArbitrary<int32>(TotalSize, FPageAllocator::PageSize);
 		if (AllocSize == FPageAllocator::PageSize)
diff --git a/Engine/Source/Runtime/Core/Public/Misc/MemStack.h b/Engine/Source/Runtime/Core/Public/Misc/MemStack.h
--- a/Engine/Source/Runtime/Core/Public/Misc/MemStack.h
+++ b/Engine/Source/Runtime/Core/Public/Misc/MemStack.h
@@ -77,8 +77,16 @@ private:
 class FMemStackBase //-V1062
 {
 public:
+	enum class EPageSize : uint8
+	{
+		// Small pages are allocated unless the allocation requires a larger page.
+		Small,
+
+		// Large pages are always allocated.
+		Large
+	};
 
-	CORE_API FMemStackBase();
+	CORE_API FMemStackBase(EPageSize PageSize = EPageSize::Small);
 
 	FMemStackBase(const FMemStackBase&) = delete;
 	FMemStackBase(FMemStackBase&& Other)
@@ -211,6 +219,8 @@ private:
 	/** The number of marks on this stack. */
 	int32 NumMarks;
 
+	/** The page size to use when allocating. */
+	EPageSize PageSize;
 
 protected:
 	bool bShouldEnforceAllocMarks;	
diff --git a/Engine/Source/Runtime/RHI/Private/RHI.cpp b/Engine/Source/Runtime/RHI/Private/RHI.cpp
--- a/Engine/Source/Runtime/RHI/Private/RHI.cpp
+++ b/Engine/Source/Runtime/RHI/Private/RHI.cpp
@@ -2116,38 +2116,38 @@ FRHIUnorderedAccessView* FRHIBufferViewCache::GetOrCreateUAV(FRHICommandListBase
 
 #if !(UE_BUILD_SHIPPING || UE_BUILD_TEST)
 
-void FRHITextureViewCache::SetDebugName(const TCHAR* DebugName)
+void FRHITextureViewCache::SetDebugName(FRHICommandListBase& RHICmdList, const TCHAR* DebugName)
 {
 	for (const auto& KeyValue : UAVs)
 	{
-		RHIBindDebugLabelName(KeyValue.Value, DebugName);
+		RHICmdList.BindDebugLabelName(KeyValue.Value, DebugName);
 	}
 }
 
-void FRHIBufferViewCache::SetDebugName(const TCHAR* DebugName)
+void FRHIBufferViewCache::SetDebugName(FRHICommandListBase& RHICmdList, const TCHAR* DebugName)
 {
 	for (const auto& KeyValue : UAVs)
 	{
-		RHIBindDebugLabelName(KeyValue.Value, DebugName);
+		RHICmdList.BindDebugLabelName(KeyValue.Value, DebugName);
 	}
 }
 
 #endif
 
-void FRHITransientTexture::Acquire(const TCHAR* InName, uint32 InPassIndex, uint64 InAcquireCycle)
+void FRHITransientTexture::Acquire(FRHICommandListBase& RHICmdList, const TCHAR* InName, uint32 InPassIndex, uint64 InAcquireCycle)
 {
-	FRHITransientResource::Acquire(InName, InPassIndex, InAcquireCycle);
-	ViewCache.SetDebugName(InName);
+	FRHITransientResource::Acquire(RHICmdList, InName, InPassIndex, InAcquireCycle);
+	ViewCache.SetDebugName(RHICmdList, InName);
 
 #if !(UE_BUILD_SHIPPING || UE_BUILD_TEST)
-	RHIBindDebugLabelName(GetRHI(), InName);
+	RHICmdList.BindDebugLabelName(GetRHI(), InName);
 #endif
 }
 
-void FRHITransientBuffer::Acquire(const TCHAR* InName, uint32 InPassIndex, uint64 InAcquireCycle)
+void FRHITransientBuffer::Acquire(FRHICommandListBase& RHICmdList, const TCHAR* InName, uint32 InPassIndex, uint64 InAcquireCycle)
 {
-	FRHITransientResource::Acquire(InName, InPassIndex, InAcquireCycle);
-	ViewCache.SetDebugName(InName);
+	FRHITransientResource::Acquire(RHICmdList, InName, InPassIndex, InAcquireCycle);
+	ViewCache.SetDebugName(RHICmdList, InName);
 
 	// TODO: Add method to rename a buffer.
 }
diff --git a/Engine/Source/Runtime/RHI/Public/RHIResources.h b/Engine/Source/Runtime/RHI/Public/RHIResources.h
--- a/Engine/Source/Runtime/RHI/Public/RHIResources.h
+++ b/Engine/Source/Runtime/RHI/Public/RHIResources.h
@@ -4761,9 +4761,9 @@ public:
 
 	// Sets the debug name of the RHI view resources.
 #if !(UE_BUILD_SHIPPING || UE_BUILD_TEST)
-	RHI_API void SetDebugName(const TCHAR* DebugName);
+	RHI_API void SetDebugName(FRHICommandListBase& RHICmdList, const TCHAR* DebugName);
 #else
-	void SetDebugName(const TCHAR* DebugName) {}
+	void SetDebugName(FRHICommandListBase& RHICmdList, const TCHAR* DebugName) {}
 #endif
 
 private:
@@ -4787,9 +4787,9 @@ public:
 
 	// Sets the debug name of the RHI view resources.
 #if !(UE_BUILD_SHIPPING || UE_BUILD_TEST)
-	RHI_API void SetDebugName(const TCHAR* DebugName);
+	RHI_API void SetDebugName(FRHICommandListBase& RHICmdList, const TCHAR* DebugName);
 #else
-	void SetDebugName(const TCHAR* DebugName) {}
+	void SetDebugName(FRHICommandListBase& RHICmdList, const TCHAR* DebugName) {}
 #endif
 
 private:
diff --git a/Engine/Source/Runtime/RHI/Public/RHITransientResourceAllocator.h b/Engine/Source/Runtime/RHI/Public/RHITransientResourceAllocator.h
--- a/Engine/Source/Runtime/RHI/Public/RHITransientResourceAllocator.h
+++ b/Engine/Source/Runtime/RHI/Public/RHITransientResourceAllocator.h
@@ -5,6 +5,7 @@
 #include "RHI.h"
 #include "RHIResources.h"
 
+class FRHICommandListBase;
 class FRHICommandListImmediate;
 class FRHITransientHeap;
 class FRHITransientPagePool;
@@ -102,7 +103,7 @@ public:
 	//////////////////////////////////////////////////////////////////////////////////////////////////
 	//! Internal Allocator API
 
-	virtual void Acquire(const TCHAR* InName, uint32 InAcquirePassIndex, uint64 InAllocatorCycle)
+	virtual void Acquire(FRHICommandListBase& RHICmdList, const TCHAR* InName, uint32 InAcquirePassIndex, uint64 InAllocatorCycle)
 	{
 		Name = InName;
 		AcquirePasses = TInterval<uint32>(0, InAcquirePassIndex);
@@ -251,7 +252,7 @@ public:
 
 	//////////////////////////////////////////////////////////////////////////////////////////////////
 	//! Internal Allocator API
-	RHI_API void Acquire(const TCHAR* InName, uint32 InAcquirePassIndex, uint64 InInitCycle) override;
+	RHI_API void Acquire(FRHICommandListBase& RHICmdList, const TCHAR* InName, uint32 InAcquirePassIndex, uint64 InInitCycle) override;
 	//////////////////////////////////////////////////////////////////////////////////////////////////
 
 	// Returns the underlying RHI texture.
@@ -289,7 +290,7 @@ public:
 
 	//////////////////////////////////////////////////////////////////////////////////////////////////
 	//! Internal Allocator API
-	RHI_API void Acquire(const TCHAR* InName, uint32 InAcquirePassIndex, uint64 InInitCycle) override;
+	RHI_API void Acquire(FRHICommandListBase& RHICmdList, const TCHAR* InName, uint32 InAcquirePassIndex, uint64 InInitCycle) override;
 	//////////////////////////////////////////////////////////////////////////////////////////////////
 
 	// Returns the underlying RHI buffer.
diff --git a/Engine/Source/Runtime/RHICore/Private/RHICoreTransientResourceAllocator.cpp b/Engine/Source/Runtime/RHICore/Private/RHICoreTransientResourceAllocator.cpp
--- a/Engine/Source/Runtime/RHICore/Private/RHICoreTransientResourceAllocator.cpp
+++ b/Engine/Source/Runtime/RHICore/Private/RHICoreTransientResourceAllocator.cpp
@@ -15,13 +15,6 @@ static FAutoConsoleVariableRef CVarRHITransientAllocatorMinimumHeapSize(
 	TEXT("Minimum size of an RHI transient heap in MB. Heaps will default to this size and grow to the maximum based on the first allocation (Default 128)."),
 	ECVF_ReadOnly);
 
-static int32 GRHITransientAllocatorMaximumHeapSize = 512;
-static FAutoConsoleVariableRef CVarRHITransientAllocatorMaximumHeapSize(
-	TEXT("RHI.TransientAllocator.MaximumHeapSize"),
-	GRHITransientAllocatorMaximumHeapSize,
-	TEXT("Maximum size of an RHI transient allocation in MB. Allocations larger than this will fail the transient allocator (Default 512)."),
-	ECVF_ReadOnly);
-
 static int32 GRHITransientAllocatorBufferCacheSize = 64;
 static FAutoConsoleVariableRef CVarRHITransientAllocatorBufferCacheSize(
 	TEXT("RHI.TransientAllocator.BufferCacheSize"),
@@ -440,7 +433,7 @@ FRHITransientTexture* FRHITransientHeap::CreateTexture(
 	});
 
 	check(Texture);
-	Texture->Acquire(DebugName, PassIndex, CurrentAllocatorCycle);
+	Texture->Acquire(FRHICommandListImmediate::Get(), DebugName, PassIndex, CurrentAllocatorCycle);
 	AllocateMemoryInternal(Texture, DebugName, PassIndex, CurrentAllocatorCycle, Allocation);
 	Stats.AllocateTexture(Allocation.Size);
 	return Texture;
@@ -476,7 +469,7 @@ FRHITransientBuffer* FRHITransientHeap::CreateBuffer(
 	});
 
 	check(Buffer);
-	Buffer->Acquire(DebugName, PassIndex, CurrentAllocatorCycle);
+	Buffer->Acquire(FRHICommandListImmediate::Get(), DebugName, PassIndex, CurrentAllocatorCycle);
 	AllocateMemoryInternal(Buffer, DebugName, PassIndex, CurrentAllocatorCycle, Allocation);
 	Stats.AllocateBuffer(Allocation.Size);
 	return Buffer;
@@ -493,10 +486,10 @@ void FRHITransientHeap::AllocateMemoryInternal(FRHITransientResource* Resource,
 	Resource->GetHeapAllocation() = Allocation;
 
 	check(Allocation.Offset % Initializer.Alignment == 0);
-	const uint64 AlignedSize   = Align(Allocation.Size, Initializer.Alignment);
-	const uint32 PageOffsetMin = Allocation.Offset >> AlignmentLog2;
-	const uint32 PageOffsetMax = (Allocation.Offset + AlignedSize) >> AlignmentLog2;
-	OverlapTracker.Track(Resource, PageOffsetMin, PageOffsetMax);
+		const uint64 AlignedSize   = Align(Allocation.Size, Initializer.Alignment);
+		const uint32 PageOffsetMin = Allocation.Offset >> AlignmentLog2;
+		const uint32 PageOffsetMax = (Allocation.Offset + AlignedSize) >> AlignmentLog2;
+		OverlapTracker.Track(Resource, PageOffsetMin, PageOffsetMax);
 
 	CommitSize = FMath::Max(CommitSize, Allocation.Offset + Allocation.Size);
 }
@@ -573,8 +566,8 @@ void FRHITransientHeap::Flush(uint64 AllocatorCycle, FRHITransientMemoryStats& O
 			}
 		}
 
-		OverlapTracker.Reset();
-	}
+			OverlapTracker.Reset();
+		}
 }
 
 //////////////////////////////////////////////////////////////////////////
@@ -583,7 +576,6 @@ FRHITransientHeapCache::FInitializer FRHITransientHeapCache::FInitializer::Creat
 {
 	FInitializer Initializer;
 	Initializer.MinimumHeapSize = GRHITransientAllocatorMinimumHeapSize * 1024 * 1024;
-	Initializer.MaximumHeapSize = GRHITransientAllocatorMaximumHeapSize * 1024 * 1024;
 	Initializer.HeapAlignment = 64 * 1024;
 	Initializer.BufferCacheSize = GRHITransientAllocatorBufferCacheSize;
 	Initializer.TextureCacheSize = GRHITransientAllocatorTextureCacheSize;
@@ -682,11 +674,6 @@ FRHITransientTexture* FRHITransientResourceHeapAllocator::CreateTextureInternal(
 	uint32 TextureAlignment,
 	FRHITransientHeap::FCreateTextureFunction CreateTextureFunction)
 {
-	if (TextureSize > HeapCache.GetInitializer().MaximumHeapSize)
-	{
-		return nullptr;
-	}
-
 	const ERHITransientHeapFlags TextureHeapFlags =
 		EnumHasAnyFlags(CreateInfo.Flags, TexCreate_RenderTargetable | TexCreate_ResolveTargetable | TexCreate_DepthStencilTargetable | TexCreate_DepthStencilResolveTarget)
 		? ERHITransientHeapFlags::AllowRenderTargets
@@ -730,11 +717,6 @@ FRHITransientBuffer* FRHITransientResourceHeapAllocator::CreateBufferInternal(
 	uint32 BufferAlignment,
 	FRHITransientHeap::FCreateBufferFunction CreateBufferFunction)
 {
-	if (BufferSize > HeapCache.GetInitializer().MaximumHeapSize)
-	{
-		return nullptr;
-	}
-
 	FRHITransientBuffer* Buffer = nullptr;
 
 	for (FRHITransientHeap* Heap : Heaps)
@@ -1218,11 +1200,11 @@ void FRHITransientPagePool::Allocate(FAllocationContext& Context)
 		const  int32 AllocationIndex   = Context.Allocations.Num();
 		const uint64 AllocationHash    = CityHash64WithSeed((const char*)&PageSpans[PageSpanOffsetMin], PageSpanCount * sizeof(FRHITransientPageSpan), DestinationGpuVirtualAddress);
 
-		for (uint32 Index = PageSpanOffsetMin; Index < PageSpanOffsetMax; ++Index)
-		{
-			const FRHITransientPageSpan Span = PageSpans[Index];
-			OverlapTracker.Track(&Context.Resource, Span.Offset, Span.Offset + Span.Count);
-		}
+			for (uint32 Index = PageSpanOffsetMin; Index < PageSpanOffsetMax; ++Index)
+			{
+				const FRHITransientPageSpan Span = PageSpans[Index];
+				OverlapTracker.Track(&Context.Resource, Span.Offset, Span.Offset + Span.Count);
+			}
 
 		FRHITransientPagePoolAllocation Allocation;
 		Allocation.Pool = this;
@@ -1273,7 +1255,7 @@ void FRHITransientPagePool::Flush(FRHICommandListImmediate& RHICmdList)
 		PageSpans.Reserve(PageSpanCountMax);
 	}
 
-	OverlapTracker.Reset();
+		OverlapTracker.Reset();
 }
 
 ///////////////////////////////////////////////////////////////////////////////////////////////////
@@ -1420,7 +1402,7 @@ FRHITransientTexture* FRHITransientResourcePageAllocator::CreateTexture(
 	const float FastPoolPercentageRequested = bFastPool ? CreateInfo.FastVRAMPercentage / 255.f : 0.f;
 
 	check(Texture);
-	Texture->Acquire(DebugName, PassIndex, CurrentCycle);
+	Texture->Acquire(FRHICommandListImmediate::Get(), DebugName, PassIndex, CurrentCycle);
 	AllocateMemoryInternal(Texture, DebugName, PassIndex, bFastPool, FastPoolPercentageRequested);
 	Stats.AllocateTexture(Texture->GetSize());
 	IF_RHICORE_TRANSIENT_ALLOCATOR_DEBUG(ActiveResources.Emplace(Texture));
@@ -1439,7 +1421,7 @@ FRHITransientBuffer* FRHITransientResourcePageAllocator::CreateBuffer(
 	});
 
 	check(Buffer);
-	Buffer->Acquire(DebugName, PassIndex, CurrentCycle);
+	Buffer->Acquire(FRHICommandListImmediate::Get(), DebugName, PassIndex, CurrentCycle);
 	AllocateMemoryInternal(Buffer, DebugName, PassIndex, EnumHasAnyFlags(CreateInfo.Usage, EBufferUsageFlags::FastVRAM), false);
 	Stats.AllocateBuffer(Buffer->GetSize());
 	IF_RHICORE_TRANSIENT_ALLOCATOR_DEBUG(ActiveResources.Emplace(Buffer));
diff --git a/Engine/Source/Runtime/RHICore/Public/RHICoreTransientResourceAllocator.h b/Engine/Source/Runtime/RHICore/Public/RHICoreTransientResourceAllocator.h
--- a/Engine/Source/Runtime/RHICore/Public/RHICoreTransientResourceAllocator.h
+++ b/Engine/Source/Runtime/RHICore/Public/RHICoreTransientResourceAllocator.h
@@ -211,6 +211,7 @@ public:
 			}
 		}
 
+		TRACE_CPUPROFILER_EVENT_SCOPE(CreatePlacedResource);
 		TransientResourceType* Resource = CreateFunction(Hash);
 		Allocated.Emplace(Resource);
 		MissCount++;
@@ -596,9 +597,6 @@ public:
 		// The minimum size to use when creating a heap. This is the default but can grow based on allocations.
 		uint64 MinimumHeapSize = 0;
 
-		// The maximum size of a pool. Allocations above this size will fail.
-		uint64 MaximumHeapSize = 0;
-
 		// The minimum alignment for resources in the heap.
 		uint32 HeapAlignment = 0;
 
@@ -633,8 +631,7 @@ public:
 
 	uint64 GetHeapSize(uint64 RequestedHeapSize) const
 	{
-		check(RequestedHeapSize <= Initializer.MaximumHeapSize);
-		return FMath::Clamp(FMath::RoundUpToPowerOfTwo64(RequestedHeapSize), Initializer.MinimumHeapSize, Initializer.MaximumHeapSize);
+		return FMath::Max(FMath::RoundUpToPowerOfTwo64(RequestedHeapSize), Initializer.MinimumHeapSize);
 	}
 
 private:
diff --git a/Engine/Source/Runtime/RenderCore/Private/RenderGraphAllocator.cpp b/Engine/Source/Runtime/RenderCore/Private/RenderGraphAllocator.cpp
--- a/Engine/Source/Runtime/RenderCore/Private/RenderGraphAllocator.cpp
+++ b/Engine/Source/Runtime/RenderCore/Private/RenderGraphAllocator.cpp
@@ -1,24 +1,78 @@
 // Copyright Epic Games, Inc. All Rights Reserved.
 
 #include "RenderGraphAllocator.h"
-#include "Async/TaskGraphInterfaces.h"
 #include "RenderCore.h"
 #include "RenderGraphPrivate.h"
 
-FRDGAllocator& FRDGAllocator::Get()
+#if RDG_ALLOCATOR_DEBUG
+thread_local int32 FRDGAllocator::NumAccessesTLS = 0;
+#endif
+
+thread_local FRDGAllocator* FRDGAllocator::AllocatorTLS = nullptr;
+
+FRDGAllocator::FRDGAllocator()
+#if !RDG_USE_MALLOC
+	: MemStack(FMemStackBase::EPageSize::Large)
+#endif
+{}
+
+FRDGAllocator::FRDGAllocator(FRDGAllocator&& Other)
 {
-	static FRDGAllocator Instance;
-	return Instance;
+	*this = MoveTemp(Other);
+}
+
+FRDGAllocator& FRDGAllocator::operator= (FRDGAllocator&& Other)
+{
+#if RDG_USE_MALLOC
+	Mallocs = MoveTemp(Other.Mallocs);
+	NumMallocBytes = Other.NumMallocBytes;
+	Other.NumMallocBytes = 0;
+	check(Other.Mallocs.IsEmpty());
+#else
+	MemStack = MoveTemp(Other.MemStack);
+	check(Other.MemStack.IsEmpty());
+#endif
+	Objects = MoveTemp(Other.Objects);
+	check(Other.Objects.IsEmpty());
+
+#if RDG_ALLOCATOR_DEBUG
+	NumAccesses = Other.NumAccesses.load(std::memory_order_relaxed);
+	Other.NumAccesses.store(0, std::memory_order_relaxed);
+#endif
+
+	return *this;
 }
 
 FRDGAllocator::~FRDGAllocator()
 {
-	Context.ReleaseAll();
-	ContextForTasks.ReleaseAll();
+	ReleaseAll();
 }
 
-void FRDGAllocator::FContext::ReleaseAll()
+FRDGAllocator& FRDGAllocator::GetTLS()
+{
+	checkf(AllocatorTLS, TEXT("Attempted to access RDG allocator outside of FRDGAllocatorScope"));
+	return *AllocatorTLS;
+}
+
+#if RDG_ALLOCATOR_DEBUG
+
+void FRDGAllocator::AcquireAccess()
 {
+	check(NumAccesses.fetch_add(1, std::memory_order_acquire) == NumAccessesTLS++);
+}
+
+void FRDGAllocator::ReleaseAccess()
+{
+	check(NumAccesses.fetch_sub(1, std::memory_order_release) == NumAccessesTLS--);
+}
+
+#endif
+
+void FRDGAllocator::ReleaseAll()
+{
+	CSV_SCOPED_TIMING_STAT_EXCLUSIVE_CONDITIONAL(RDGAllocator_Clear, GRDGVerboseCSVStats != 0 && IsInRenderingThread());
+	TRACE_CPUPROFILER_EVENT_SCOPE(FRDGAllocator::ReleaseAll);
+
 	for (int32 Index = Objects.Num() - 1; Index >= 0; --Index)
 	{
 #if RDG_USE_MALLOC
@@ -30,41 +84,26 @@ void FRDGAllocator::FContext::ReleaseAll()
 	Objects.Reset();
 
 #if RDG_USE_MALLOC
-	for (void* RawAlloc : RawAllocs)
+	for (void* Malloc : Mallocs)
 	{
-		FMemory::Free(RawAlloc);
+		FMemory::Free(Malloc);
 	}
-	RawAllocs.Reset();
+	Mallocs.Reset();
+	NumMallocBytes = 0;
 #else
 	MemStack.Flush();
 #endif
 }
 
-void FRDGAllocator::ReleaseAll()
+FRDGAllocatorScope::FRDGAllocatorScope(FRDGAllocator& Allocator)
+	: AllocatorToRestore(FRDGAllocator::AllocatorTLS)
 {
-	TRACE_CPUPROFILER_EVENT_SCOPE(FRDGAllocator::ReleaseAll);
-	CSV_SCOPED_TIMING_STAT_EXCLUSIVE_CONDITIONAL(RDGAllocator_Clear, GRDGVerboseCSVStats != 0);
-	Context.ReleaseAll();
-	ContextForTasks.ReleaseAll();
+	FRDGAllocator::AllocatorTLS = &Allocator;
 }
 
 FRDGAllocatorScope::~FRDGAllocatorScope()
 {
-	if (AsyncDeleteFunction)
-	{
-		FFunctionGraphTask::CreateAndDispatchWhenReady(
-			[Allocator = MoveTemp(Allocator), AsyncDeleteFunction = MoveTemp(AsyncDeleteFunction)] () mutable
-		{
-			SCOPED_NAMED_EVENT(FRDGAllocatorScope_AsyncDelete, FColor::Emerald);
-			AsyncDeleteFunction();
-			AsyncDeleteFunction = {};
-
-		}, TStatId(), nullptr, ENamedThreads::AnyThread);
-	}
-	else
-	{
-		Allocator.ReleaseAll();
-	}
+	FRDGAllocator::AllocatorTLS = AllocatorToRestore;
 }
 
 FORCENOINLINE void UE::RenderCore::Private::OnInvalidRDGAllocatorNum(int32 NewNum, SIZE_T NumBytesPerElement)
diff --git a/Engine/Source/Runtime/RenderCore/Private/RenderGraphBuilder.cpp b/Engine/Source/Runtime/RenderCore/Private/RenderGraphBuilder.cpp
--- a/Engine/Source/Runtime/RenderCore/Private/RenderGraphBuilder.cpp
+++ b/Engine/Source/Runtime/RenderCore/Private/RenderGraphBuilder.cpp
@@ -52,7 +52,6 @@ struct FParallelPassSet : public FRHICommandListImmediate::FQueuedCommandList
 	TArray<FRDGPass*, FRDGArrayAllocator> Passes;
 	IF_RHI_WANT_BREADCRUMB_EVENTS(FRDGBreadcrumbState* BreadcrumbStateBegin{});
 	IF_RHI_WANT_BREADCRUMB_EVENTS(FRDGBreadcrumbState* BreadcrumbStateEnd{});
-	int8 bInitialized = 0;
 	bool bDispatchAfterExecute = false;
 	bool bParallelTranslate = false;
 };
@@ -494,6 +493,11 @@ ERDGPassFlags FRDGBuilder::OverridePassFlags(const TCHAR* PassName, ERDGPassFlag
 
 bool FRDGBuilder::IsTransient(FRDGBufferRef Buffer) const
 {
+	if (!bSupportsTransientBuffers || Buffer->bQueuedForUpload)
+	{
+		return false;
+	}
+
 	if (!IsTransientInternal(Buffer, EnumHasAnyFlags(Buffer->Desc.Usage, BUF_FastVRAM)))
 	{
 		return false;
@@ -509,21 +513,21 @@ bool FRDGBuilder::IsTransient(FRDGBufferRef Buffer) const
 
 bool FRDGBuilder::IsTransient(FRDGTextureRef Texture) const
 {
+	if (!bSupportsTransientTextures)
+	{
+		return false;
+	}
+
 	if (EnumHasAnyFlags(Texture->Desc.Flags, ETextureCreateFlags::Shared))
 	{
 		return false;
 	}
+
 	return IsTransientInternal(Texture, EnumHasAnyFlags(Texture->Desc.Flags, ETextureCreateFlags::FastVRAM));
 }
 
 bool FRDGBuilder::IsTransientInternal(FRDGViewableResource* Resource, bool bFastVRAM) const
 {
-	// Immediate mode can't use the transient allocator because we don't know if the user will extract the resource.
-	if (!GRDGTransientAllocator || IsImmediateMode())
-	{
-		return false;
-	}
-
 	// FastVRAM resources are always transient regardless of extraction or other hints, since they are performance critical.
 	if (!bFastVRAM || !FPlatformMemory::SupportsFastVRAMMemory())
 	{
@@ -552,9 +556,19 @@ bool FRDGBuilder::IsTransientInternal(FRDGViewableResource* Resource, bool bFast
 	}
 
 #if RDG_ENABLE_DEBUG
-	if (GRDGDebugDisableTransientResources != 0 && IsDebugAllowedForResource(Resource->Name))
+	if (GRDGDebugDisableTransientResources != 0)
 	{
-		return false;
+		const bool bDebugAllowed = IsDebugAllowedForResource(Resource->Name);
+
+		if (GRDGDebugDisableTransientResources == 2 && Resource->Type == ERDGViewableResourceType::Buffer && bDebugAllowed)
+		{
+			return false;
+		}
+
+		if (GRDGDebugDisableTransientResources == 3 && Resource->Type == ERDGViewableResourceType::Texture && bDebugAllowed)
+		{
+			return false;
+		}
 	}
 #endif
 
@@ -562,27 +576,27 @@ bool FRDGBuilder::IsTransientInternal(FRDGViewableResource* Resource, bool bFast
 }
 
 FRDGBuilder::FRDGBuilder(FRHICommandListImmediate& InRHICmdList, FRDGEventName InName, ERDGBuilderFlags InFlags)
-	: RHICmdList(InRHICmdList)
-	, Blackboard(Allocator)
+	: RootAllocatorScope(Allocators.Root)
+	, RHICmdList(InRHICmdList)
+	, Blackboard(Allocators.Root)
 	, BuilderName(InName)
-	, CompilePipe(TEXT("RDG_CompilePipe"))
-#if RDG_CPU_SCOPES
-	, CPUScopeStacks(Allocator)
-#endif
-	, GPUScopeStacks(Allocator)
-	, bParallelExecuteEnabled(IsParallelExecuteEnabled() && EnumHasAnyFlags(InFlags, ERDGBuilderFlags::AllowParallelExecute))
-	, bParallelSetupEnabled(::IsParallelSetupEnabled() && EnumHasAnyFlags(InFlags, ERDGBuilderFlags::AllowParallelExecute))
+	, TransientResourceAllocator(GRDGTransientAllocator != 0 && !::IsImmediateMode() ? GRDGTransientResourceAllocator.Get() : nullptr)
+	, ExtendResourceLifetimeScope(RHICmdList)
 #if RDG_ENABLE_DEBUG
-	, UserValidation(Allocator, bParallelExecuteEnabled)
+	, UserValidation(Allocators.Root, ParallelExecute.bEnabled)
 	, BarrierValidation(&Passes, BuilderName)
 #endif
-	, TransientResourceAllocator(GRDGTransientResourceAllocator.Get())
-	, ExtendResourceLifetimeScope(RHICmdList)
 {
-	AddProloguePass();
+	ProloguePass = SetupEmptyPass(Passes.Allocate<FRDGSentinelPass>(Allocators.Root, RDG_EVENT_NAME("Graph Prologue (Graphics)")));
 
-	bSupportsTransientTextures = TransientResourceAllocator ? TransientResourceAllocator->SupportsResourceType(ERHITransientResourceType::Texture) : false;
-	bSupportsTransientBuffers = TransientResourceAllocator ? TransientResourceAllocator->SupportsResourceType(ERHITransientResourceType::Buffer) : false;
+	ParallelExecute.bEnabled = ::IsParallelExecuteEnabled() && EnumHasAnyFlags(InFlags, ERDGBuilderFlags::AllowParallelExecute);
+	ParallelSetup.bEnabled   = ::IsParallelSetupEnabled()   && EnumHasAnyFlags(InFlags, ERDGBuilderFlags::AllowParallelExecute);
+
+	if (TransientResourceAllocator)
+	{
+		bSupportsTransientTextures = TransientResourceAllocator->SupportsResourceType(ERHITransientResourceType::Texture);
+		bSupportsTransientBuffers  = TransientResourceAllocator->SupportsResourceType(ERHITransientResourceType::Buffer);
+	}
 
 #if RDG_EVENTS != RDG_EVENTS_NONE
 	// This is polled once as a workaround for a race condition since the underlying global is not always changed on the render thread.
@@ -590,9 +604,9 @@ FRDGBuilder::FRDGBuilder(FRHICommandListImmediate& InRHICmdList, FRDGEventName I
 #endif
 
 #if RHI_WANT_BREADCRUMB_EVENTS
-	if (bParallelExecuteEnabled)
+	if (ParallelExecute.bEnabled)
 	{
-		BreadcrumbState = FRDGBreadcrumbState::Create(Allocator);
+		BreadcrumbState = FRDGBreadcrumbState::Create(Allocators.Root);
 	}
 #endif
 
@@ -603,10 +617,10 @@ FRDGBuilder::FRDGBuilder(FRHICommandListImmediate& InRHICmdList, FRDGEventName I
 
 FRDGBuilder::~FRDGBuilder()
 {
-	if (bParallelExecuteEnabled && GRDGParallelDestruction > 0)
+	if (ParallelExecute.bEnabled && GRDGParallelDestruction > 0)
 	{
-		// Move expensive operations into the async deleter, which will be called in the base class destructor.
-		BeginAsyncDelete([
+		UE::Tasks::Launch(UE_SOURCE_LOCATION, [
+			Allocators				= MoveTemp(Allocators),
 			Passes					= MoveTemp(Passes),
 			Textures				= MoveTemp(Textures),
 			Buffers					= MoveTemp(Buffers),
@@ -628,9 +642,12 @@ const TRefCountPtr<FRDGPooledBuffer>& FRDGBuilder::ConvertToExternalBuffer(FRDGB
 	if (!Buffer->bExternal)
 	{
 		Buffer->bExternal = 1;
-		Buffer->bForceNonTransient = 1;
-		BeginResourceRHI(GetProloguePassHandle(), Buffer);
-		ExternalBuffers.Add(Buffer->GetRHIUnchecked(), Buffer);
+		if (!Buffer->ResourceRHI)
+		{
+			SetPooledBufferRHI(Buffer, AllocatePooledBufferRHI(RHICmdList, Buffer));
+		}
+		ExternalBuffers.FindOrAdd(Buffer->GetRHIUnchecked(), Buffer);
+		AsyncSetupQueue.Push(FAsyncSetupOp::CullRootBuffer(Buffer));
 	}
 	return GetPooledBuffer(Buffer);
 }
@@ -641,9 +658,12 @@ const TRefCountPtr<IPooledRenderTarget>& FRDGBuilder::ConvertToExternalTexture(F
 	if (!Texture->bExternal)
 	{
 		Texture->bExternal = 1;
-		Texture->bForceNonTransient = 1;
-		BeginResourceRHI(GetProloguePassHandle(), Texture);
-		ExternalTextures.Add(Texture->GetRHIUnchecked(), Texture);
+		if (!Texture->ResourceRHI)
+		{
+			SetPooledRenderTargetRHI(Texture, AllocatePooledRenderTargetRHI(RHICmdList, Texture));
+		}
+		ExternalTextures.FindOrAdd(Texture->GetRHIUnchecked(), Texture);
+		AsyncSetupQueue.Push(FAsyncSetupOp::CullRootTexture(Texture));
 	}
 	return GetPooledTexture(Texture);
 }
@@ -654,14 +674,15 @@ FRHIUniformBuffer* FRDGBuilder::ConvertToExternalUniformBuffer(FRDGUniformBuffer
 	{
 		UniformBuffer->GetParameters().Enumerate([this](const FRDGParameter& Param)
 		{
-			auto ConvertTexture = [](FRDGBuilder* Builder, FRDGTextureRef Texture)
+			const auto ConvertTexture = [](FRDGBuilder* Builder, FRDGTextureRef Texture)
 			{
 				if (Texture && !Texture->IsExternal())
 				{
 					Builder->ConvertToExternalTexture(Texture);
 				}
 			};
-			auto ConvertBuffer = [](FRDGBuilder* Builder, FRDGBufferRef Buffer)
+
+			const auto ConvertBuffer = [](FRDGBuilder* Builder, FRDGBufferRef Buffer)
 			{
 				if (Buffer && !Buffer->IsExternal())
 				{
@@ -669,6 +690,14 @@ FRHIUniformBuffer* FRDGBuilder::ConvertToExternalUniformBuffer(FRDGUniformBuffer
 				}
 			};
 
+			const auto ConvertView = [this] (FRDGView* View)
+			{
+				if (!View->ResourceRHI)
+				{
+					InitViewRHI(RHICmdList, View);
+				}
+			};
+
 			switch (Param.GetType())
 			{
 			case UBMT_RDG_TEXTURE:
@@ -693,13 +722,13 @@ FRHIUniformBuffer* FRDGBuilder::ConvertToExternalUniformBuffer(FRDGUniformBuffer
 			case UBMT_RDG_TEXTURE_SRV:
 			{
 				ConvertTexture(this, Param.GetAsTextureSRV()->Desc.Texture);
-				InitRHI(Param.GetAsView());
+				ConvertView(Param.GetAsView());
 			}
 			break;
 			case UBMT_RDG_TEXTURE_UAV:
 			{
 				ConvertTexture(this, Param.GetAsTextureUAV()->Desc.Texture);
-				InitRHI(Param.GetAsView());
+				ConvertView(Param.GetAsView());
 			}
 			break;
 			case UBMT_RDG_BUFFER_ACCESS:
@@ -719,13 +748,13 @@ FRHIUniformBuffer* FRDGBuilder::ConvertToExternalUniformBuffer(FRDGUniformBuffer
 			case UBMT_RDG_BUFFER_SRV:
 			{
 				ConvertBuffer(this, Param.GetAsBufferSRV()->Desc.Buffer);
-				InitRHI(Param.GetAsView());
+				ConvertView(Param.GetAsView());
 			}
 			break;
 			case UBMT_RDG_BUFFER_UAV:
 			{
 				ConvertBuffer(this, Param.GetAsBufferUAV()->Desc.Buffer);
-				InitRHI(Param.GetAsView());
+				ConvertView(Param.GetAsView());
 			}
 			break;
 			case UBMT_RDG_UNIFORM_BUFFER:
@@ -761,7 +790,6 @@ FRHIUniformBuffer* FRDGBuilder::ConvertToExternalUniformBuffer(FRDGUniformBuffer
 	if (!UniformBuffer->bExternal)
 	{
 		UniformBuffer->bExternal = true;
-		UniformBuffer->bQueuedForCreate = true;
 
 		// It's safe to reset the access to false because validation won't allow this call during execution.
 		IF_RDG_ENABLE_DEBUG(GRDGAllowRHIAccess = true);
@@ -866,14 +894,14 @@ void FRDGBuilder::FlushAccessModeQueue()
 	ERHIPipeline ParameterPipelines = ERHIPipeline::None;
 
 	TArray<FRDGPass::FExternalAccessOp, FRDGArrayAllocator> Ops;
-	Ops.Reserve(bParallelSetupEnabled ? AccessModeQueue.Num() : 0);
+	Ops.Reserve(ParallelSetup.bEnabled ? AccessModeQueue.Num() : 0);
 
 	for (FRDGViewableResource* Resource : AccessModeQueue)
 	{
 		const auto& AccessModeState = Resource->AccessModeState;
 		Resource->AccessModeState.bQueued = false;
 
-		if (bParallelSetupEnabled)
+		if (ParallelSetup.bEnabled)
 		{
 			Ops.Emplace(Resource, AccessModeState.Mode);
 		}
@@ -923,7 +951,7 @@ void FRDGBuilder::FlushAccessModeQueue()
 		FAccessModePassParameters* Parameters = ParametersByPipeline[GetRHIPipelineIndex(ERHIPipeline::Graphics)];
 
 		FRDGPass* Pass = Passes.Allocate<LambdaPassType>(
-			Allocator,
+			Allocators.Root,
 			RDG_EVENT_NAME("AccessModePass[Graphics] (Textures: %d, Buffers: %d)", Parameters->Textures.Num(), Parameters->Buffers.Num()),
 			FAccessModePassParameters::FTypeInfo::GetStructMetadata(),
 			Parameters,
@@ -944,7 +972,7 @@ void FRDGBuilder::FlushAccessModeQueue()
 		FAccessModePassParameters* Parameters = ParametersByPipeline[GetRHIPipelineIndex(ERHIPipeline::AsyncCompute)];
 
 		FRDGPass* Pass = Passes.Allocate<LambdaPassType>(
-			Allocator,
+			Allocators.Root,
 			RDG_EVENT_NAME("AccessModePass[AsyncCompute] (Textures: %d, Buffers: %d)", Parameters->Textures.Num(), Parameters->Buffers.Num()),
 			FAccessModePassParameters::FTypeInfo::GetStructMetadata(),
 			Parameters,
@@ -992,17 +1020,10 @@ FRDGTexture* FRDGBuilder::RegisterExternalTexture(
 	} 
 
 	const FRDGTextureDesc Desc = Translate(ExternalPooledTexture->GetDesc());
-	FRDGTexture* Texture = Textures.Allocate(Allocator, Name, Desc, Flags);
-	SetRHI(Texture, ExternalPooledTexture.GetReference(), GetProloguePassHandle());
+	FRDGTexture* Texture = Textures.Allocate(Allocators.Root, Name, Desc, Flags);
+	SetPooledRenderTargetRHI(Texture, ExternalPooledTexture.GetReference());
 	Texture->bExternal = true;
-	ExternalTextures.Add(Texture->GetRHIUnchecked(), Texture);
-
-	if (Texture->bTransient)
-	{
-		FRDGSubresourceState State;
-		State.SetPass(ERHIPipeline::Graphics, GetProloguePassHandle());
-		InitTextureSubresources(*Texture->State, Texture->Layout, State);
-	}
+	ExternalTextures.FindOrAdd(Texture->GetRHIUnchecked(), Texture);
 
 	IF_RDG_ENABLE_DEBUG(UserValidation.ValidateRegisterExternalTexture(Texture));
 	IF_RDG_ENABLE_TRACE(Trace.AddResource(Texture));
@@ -1035,11 +1056,10 @@ FRDGBufferRef FRDGBuilder::RegisterExternalBuffer(
 		return FoundBuffer;
 	}
 
-	FRDGBuffer* Buffer = Buffers.Allocate(Allocator, Name, ExternalPooledBuffer->Desc, Flags);
-	SetRHI(Buffer, ExternalPooledBuffer, GetProloguePassHandle());
+	FRDGBuffer* Buffer = Buffers.Allocate(Allocators.Root, Name, ExternalPooledBuffer->Desc, Flags);
+	SetPooledBufferRHI(Buffer, ExternalPooledBuffer);
 	Buffer->bExternal = true;
-
-	ExternalBuffers.Add(Buffer->GetRHIUnchecked(), Buffer);
+	ExternalBuffers.FindOrAdd(Buffer->GetRHIUnchecked(), Buffer);
 
 	IF_RDG_ENABLE_DEBUG(UserValidation.ValidateRegisterExternalBuffer(Buffer));
 	IF_RDG_ENABLE_TRACE(Trace.AddResource(Buffer));
@@ -1052,7 +1072,7 @@ void FRDGBuilder::AddPassDependency(FRDGPass* Producer, FRDGPass* Consumer)
 {
 	auto& Producers = Consumer->Producers;
 
-	if (Producers.Find(Producer->Handle) == INDEX_NONE)
+	if (Producers.Find(Producer) == INDEX_NONE)
 	{
 #if RDG_STATS
 		GRDGStatPassDependencyCount++;
@@ -1060,7 +1080,7 @@ void FRDGBuilder::AddPassDependency(FRDGPass* Producer, FRDGPass* Consumer)
 
 		if (Producer->Pipeline != Consumer->Pipeline)
 		{
-			const auto BinarySearchOrAdd = [](FRDGPassHandleArray& Range, FRDGPassHandle Handle)
+			const auto BinarySearchOrAdd = [](auto& Range, FRDGPassHandle Handle)
 			{
 				const int32 LowerBoundIndex = Algo::LowerBound(Range, Handle);
 				if (LowerBoundIndex < Range.Num())
@@ -1083,11 +1103,11 @@ void FRDGBuilder::AddPassDependency(FRDGPass* Producer, FRDGPass* Consumer)
 			}
 		}
 
-		Producers.Add(Producer->Handle);
+		Producers.Add(Producer);
 	}
 }
 
-void FRDGBuilder::AddCullingDependency(FRDGProducerStatesByPipeline& LastProducers, const FRDGProducerState& NextState, ERHIPipeline NextPipeline)
+bool FRDGBuilder::AddCullingDependency(FRDGProducerStatesByPipeline& LastProducers, const FRDGProducerState& NextState, ERHIPipeline NextPipeline)
 {
 	for (ERHIPipeline LastPipeline : GetRHIPipelines())
 	{
@@ -1149,110 +1169,60 @@ void FRDGBuilder::AddCullingDependency(FRDGProducerStatesByPipeline& LastProduce
 		LastProducer.Access             = NextState.Access;
 		LastProducer.Pass               = NextState.Pass;
 		LastProducer.NoUAVBarrierHandle = NextState.NoUAVBarrierHandle;
+		return true;
 	}
+	return false;
 }
 
-void FRDGBuilder::CompilePassBarriers()
+void FRDGBuilder::AddCullRootTexture(FRDGTexture* Texture)
 {
-	// Walk the culled graph and compile barriers for each subresource. Certain transitions are redundant; read-to-read, for example.
-	// We can avoid them by traversing and merging compatible states together. The merging states removes a transition, but the merging
-	// heuristic is conservative and choosing not to merge doesn't necessarily mean a transition is performed. They are two distinct steps.
-	// Merged states track the first and last pass interval. Pass references are also accumulated onto each resource. This must happen
-	// after culling since culled passes can't contribute references.
-
-	const FRDGPassHandle ProloguePassHandle = GetProloguePassHandle();
-	const FRDGPassHandle EpiloguePassHandle = GetEpiloguePassHandle();
-
-	SCOPED_NAMED_EVENT(CompileBarriers, FColor::Emerald);
+	check(Texture->IsCullRoot());
 
-	for (FRDGPassHandle PassHandle = ProloguePassHandle + 1; PassHandle < EpiloguePassHandle; ++PassHandle)
+	for (auto& LastProducer : Texture->LastProducers)
 	{
-		FRDGPass* Pass = Passes[PassHandle];
-
-		if (Pass->bCulled || Pass->bEmptyParameters)
-		{
-			continue;
-		}
-
-		const ERHIPipeline PassPipeline = Pass->Pipeline;
-
-		const auto MergeSubresourceStates = [&](ERDGViewableResourceType ResourceType, FRDGSubresourceState*& PassMergeState, FRDGSubresourceState*& ResourceMergeState, const FRDGSubresourceState& PassState)
-		{
-			if (!ResourceMergeState || !FRDGSubresourceState::IsMergeAllowed(ResourceType, *ResourceMergeState, PassState))
-			{
-				// Allocate a new pending merge state and assign it to the pass state.
-				ResourceMergeState = AllocSubresource(PassState);
-			}
-			else
-			{
-				// Merge the pass state into the merged state.
-				ResourceMergeState->Access |= PassState.Access;
+		AddLastProducersToCullStack(LastProducer);
+	}
 
-				FRDGPassHandle& FirstPassHandle = ResourceMergeState->FirstPass[PassPipeline];
+	FlushCullStack();
+}
 
-				if (FirstPassHandle.IsNull())
-				{
-					FirstPassHandle = PassHandle;
-				}
+void FRDGBuilder::AddCullRootBuffer(FRDGBuffer* Buffer)
+{
+	check(Buffer->IsCullRoot());
 
-				ResourceMergeState->LastPass[PassPipeline] = PassHandle;
-			}
+	AddLastProducersToCullStack(Buffer->LastProducer);
 
-			PassMergeState = ResourceMergeState;
-		};
+	FlushCullStack();
+}
 
-		for (auto& PassState : Pass->TextureStates)
+void FRDGBuilder::AddLastProducersToCullStack(const FRDGProducerStatesByPipeline& LastProducers)
+{
+	for (const FRDGProducerState& LastProducer : LastProducers)
+	{
+		if (LastProducer.Pass)
 		{
-			FRDGTextureRef Texture = PassState.Texture;
-
-			if (Texture->FirstBarrier == FRDGTexture::EFirstBarrier::ImmediateRequested)
-			{
-				check(Texture->bExternal);
-				Texture->FirstBarrier = FRDGTexture::EFirstBarrier::ImmediateConfirmed;
-				Texture->FirstPass = PassHandle;
-
-				for (FRDGSubresourceState& SubresourceState : *Texture->State)
-				{
-					SubresourceState.SetPass(ERHIPipeline::Graphics, PassHandle);
-				}
-			}
-
-		#if RDG_STATS
-			GRDGStatTextureReferenceCount += PassState.ReferenceCount;
-		#endif
-
-			for (int32 Index = 0; Index < PassState.State.Num(); ++Index)
-			{
-				if (PassState.State[Index].Access == ERHIAccess::Unknown)
-				{
-					continue;
-				}
-
-				MergeSubresourceStates(ERDGViewableResourceType::Texture, PassState.MergeState[Index], Texture->MergeState[Index], PassState.State[Index]);
-			}
+			CullPassStack.Emplace(LastProducer.Pass);
 		}
+	}
+}
 
-		for (auto& PassState : Pass->BufferStates)
-		{
-			FRDGBufferRef Buffer = PassState.Buffer;
-
-			if (Buffer->FirstBarrier == FRDGBuffer::EFirstBarrier::ImmediateRequested)
-			{
-				check(Buffer->bExternal);
-				Buffer->FirstBarrier = FRDGBuffer::EFirstBarrier::ImmediateConfirmed;
-				Buffer->FirstPass = PassHandle;
-				Buffer->State->SetPass(ERHIPipeline::Graphics, PassHandle);
-			}
+void FRDGBuilder::FlushCullStack()
+{
+	while (CullPassStack.Num())
+	{
+		FRDGPass* Pass = CullPassStack.Pop(false);
 
-		#if RDG_STATS
-			GRDGStatBufferReferenceCount += PassState.ReferenceCount;
-		#endif
+		if (Pass->bCulled)
+		{
+			Pass->bCulled = 0;
 
-			MergeSubresourceStates(ERDGViewableResourceType::Buffer, PassState.MergeState, Buffer->MergeState, PassState.State);
+			CullPassStack.Append(Pass->Producers);
 		}
 	}
 }
 
+///////////////////////////////////////////////////////////////////////////////
+
 void FRDGBuilder::Compile()
 {
 	SCOPE_CYCLE_COUNTER(STAT_RDG_CompileTime);
@@ -1267,76 +1237,19 @@ void FRDGBuilder::Compile()
 
 	const bool bCullPasses = GRDGCullPasses > 0;
 
-	if (bParallelSetupEnabled)
-	{
-		SetupPassQueue.Flush(TEXT("FRDGBuilder::SetupPassResources"), [this](FRDGPass* Pass) { SetupPassResources(Pass); });
-	}
-
-	if (bCullPasses)
-	{
-		CullPassStack.Reserve(CompilePassCount);
-	}
-
 	if (bCullPasses || AsyncComputePassCount > 0)
 	{
 		SCOPED_NAMED_EVENT(PassDependencies, FColor::Emerald);
 
-		if (!bParallelSetupEnabled)
+		if (!ParallelSetup.bEnabled)
 		{
 			for (FRDGPassHandle PassHandle = ProloguePassHandle + 1; PassHandle < EpiloguePassHandle; ++PassHandle)
 			{
 				SetupPassDependencies(Passes[PassHandle]);
 			}
 		}
-
-		const auto AddLastProducersToCullStack = [&](const FRDGProducerStatesByPipeline& LastProducers)
-		{
-			for (const FRDGProducerState& LastProducer : LastProducers)
-			{
-				if (LastProducer.Pass)
-				{
-					CullPassStack.Emplace(LastProducer.Pass->Handle);
-				}
-			}
-		};
-
-		// The last producer of a extracted or external resource is a cull graph root as it's not contained within the graph.
-
-		for (const FExtractedTexture& ExtractedTexture : ExtractedTextures)
-		{
-			FRDGTextureRef Texture = ExtractedTexture.Texture;
-
-			for (auto& LastProducer : Texture->LastProducers)
-			{
-				AddLastProducersToCullStack(LastProducer);
-			}
-		}
-
-		for (const FExtractedBuffer& ExtractedBuffer : ExtractedBuffers)
-		{
-			FRDGBufferRef Buffer = ExtractedBuffer.Buffer;
-
-			AddLastProducersToCullStack(Buffer->LastProducer);
-		}
-
-		for (const auto& Pair : ExternalTextures)
-		{
-			FRDGTexture* Texture = Pair.Value;
-
-			for (auto& LastProducer : Texture->LastProducers)
-			{
-				AddLastProducersToCullStack(LastProducer);
-			}
-		}
-
-		for (const auto& Pair : ExternalBuffers)
-		{
-			FRDGBuffer* Buffer = Pair.Value;
-
-			AddLastProducersToCullStack(Buffer->LastProducer);
-		}
 	}
-	else if (!bParallelSetupEnabled)
+	else if (!ParallelSetup.bEnabled)
 	{
 		for (FRDGPassHandle PassHandle = ProloguePassHandle + 1; PassHandle < EpiloguePassHandle; ++PassHandle)
 		{
@@ -1374,25 +1287,11 @@ void FRDGBuilder::Compile()
 	{
 		SCOPED_NAMED_EVENT(PassCulling, FColor::Emerald);
 
-		CullPassStack.Emplace(EpiloguePassHandle);
-
-		// Mark the epilogue pass as culled so that it is traversed.
-		EpiloguePass->bCulled = 1;
-
-		// Manually mark the prologue passes as not culled.
+		// Manually mark the prologue / epilogue passes as not culled.
+		EpiloguePass->bCulled = 0;
 		ProloguePass->bCulled = 0;
 
-		while (CullPassStack.Num())
-		{
-			FRDGPass* Pass = Passes[CullPassStack.Pop()];
-
-			if (Pass->bCulled)
-			{
-				Pass->bCulled = 0;
-
-				CullPassStack.Append(Pass->Producers);
-			}
-		}
+		check(CullPassStack.IsEmpty());
 
 		for (FRDGPassHandle PassHandle = ProloguePassHandle + 1; PassHandle < EpiloguePassHandle; ++PassHandle)
 		{
@@ -1584,13 +1483,13 @@ void FRDGBuilder::Compile()
 			{
 				CurrentGraphicsForkPassHandle = GraphicsForkPassHandle;
 
-				FRDGBarrierBatchBegin& EpilogueBarriersToBeginForAsyncCompute = GraphicsForkPass->GetEpilogueBarriersToBeginForAsyncCompute(Allocator, TransitionCreateQueue);
+				FRDGBarrierBatchBegin& EpilogueBarriersToBeginForAsyncCompute = GraphicsForkPass->GetEpilogueBarriersToBeginForAsyncCompute(Allocators.Transition, TransitionCreateQueue);
 
 				GraphicsForkPass->bGraphicsFork = 1;
 				EpilogueBarriersToBeginForAsyncCompute.SetUseCrossPipelineFence();
 
 				AsyncComputePass->bAsyncComputeBegin = 1;
-				AsyncComputePass->GetPrologueBarriersToEnd(Allocator).AddDependency(&EpilogueBarriersToBeginForAsyncCompute);
+				AsyncComputePass->GetPrologueBarriersToEnd(Allocators.Transition).AddDependency(&EpilogueBarriersToBeginForAsyncCompute);
 
 				// Since we are fencing the graphics pipe to some new async compute work, make sure to flush any prior work.
 				if (AsyncComputePassBeforeFork)
@@ -1637,14 +1536,14 @@ void FRDGBuilder::Compile()
 			{
 				CurrentGraphicsJoinPassHandle = GraphicsJoinPassHandle;
 
-				FRDGBarrierBatchBegin& EpilogueBarriersToBeginForGraphics = AsyncComputePass->GetEpilogueBarriersToBeginForGraphics(Allocator, TransitionCreateQueue);
+				FRDGBarrierBatchBegin& EpilogueBarriersToBeginForGraphics = AsyncComputePass->GetEpilogueBarriersToBeginForGraphics(Allocators.Transition, TransitionCreateQueue);
 
 				AsyncComputePass->bAsyncComputeEnd = 1;
 				AsyncComputePass->bDispatchAfterExecute = 1;
 				EpilogueBarriersToBeginForGraphics.SetUseCrossPipelineFence();
 
 				GraphicsJoinPass->bGraphicsJoin = 1;
-				GraphicsJoinPass->GetPrologueBarriersToEnd(Allocator).AddDependency(&EpilogueBarriersToBeginForGraphics);
+				GraphicsJoinPass->GetPrologueBarriersToEnd(Allocators.Transition).AddDependency(&EpilogueBarriersToBeginForGraphics);
 			}
 		}
 	}
@@ -1652,60 +1551,69 @@ void FRDGBuilder::Compile()
 
 ///////////////////////////////////////////////////////////////////////////////////////////////////
 
-void FRDGBuilder::FPassQueue::Push(FRDGPass* Pass)
+void FRDGBuilder::LaunchAsyncSetupQueueTask()
 {
-	UE::TScopeLock Lock(Mutex);
-	Queue.Emplace(Pass);
-}
-
-template <typename LambdaType>
-void FRDGBuilder::FPassQueue::Flush(UE::Tasks::FPipe& Pipe, const TCHAR* Name, LambdaType&& Lambda)
-{
-	if (LastTask.IsCompleted())
+	if (AsyncSetupQueue.LastTask.IsCompleted())
 	{
-		LastTask = Pipe.Launch(Name, [this, Lambda = MoveTemp(Lambda), Name] () mutable
+		AsyncSetupQueue.LastTask = AsyncSetupQueue.Pipe.Launch(UE_SOURCE_LOCATION, [this]() mutable
 		{
-			Flush(Name, MoveTemp(Lambda));
+			ProcessAsyncSetupQueue();
 		});
 	}
 }
 
-template <typename LambdaType>
-void FRDGBuilder::FPassQueue::Flush(const TCHAR* Name, LambdaType&& Lambda)
+void FRDGBuilder::ProcessAsyncSetupQueue()
 {
-	SCOPED_NAMED_EVENT_TCHAR(Name, FColor::Magenta);
-
-	Mutex.Lock();
-	TArray<FRDGPass*, FRDGArrayAllocator> PoppedQueue = MoveTemp(Queue);
-	Mutex.Unlock();
+	SCOPED_NAMED_EVENT_TCHAR("FRDGBuilder::ProcessAsyncSetupQueue", FColor::Magenta);
+	FRDGAllocatorScope AllocatorScope(Allocators.Task);
 
-	for (FRDGPass* Pass : PoppedQueue)
+	while (true)
 	{
-		Lambda(Pass);
+		AsyncSetupQueue.Mutex.Lock();
+		TArray<FAsyncSetupOp, FRDGArrayAllocator> PoppedOps = MoveTemp(AsyncSetupQueue.Ops);
+		AsyncSetupQueue.Mutex.Unlock();
+
+		if (PoppedOps.IsEmpty())
+		{
+			break;
+		}
+
+		for (FAsyncSetupOp Op : PoppedOps)
+		{
+			switch (Op.Type)
+			{
+			case FAsyncSetupOp::EType::SetupPassResources:
+				SetupPassResources(Op.Pass);
+				break;
+
+			case FAsyncSetupOp::EType::CullRootBuffer:
+				AddCullRootBuffer(Op.Buffer);
+				break;
+
+			case FAsyncSetupOp::EType::CullRootTexture:
+				AddCullRootTexture(Op.Texture);
+				break;
+			}
+		}
 	}
 }
 
-
 void FRDGBuilder::FlushSetupQueue()
 {
-	if (bParallelSetupEnabled)
+	if (ParallelSetup.bEnabled)
 	{
-		SetupPassQueue.Flush(CompilePipe, TEXT("FRDGBuilder::SetupPassResources"), [this](FRDGPass* Pass) { SetupPassResources(Pass); });
+		LaunchAsyncSetupQueueTask();
 	}
 }
 
-template <typename LambdaType>
-UE::Tasks::FTask FRDGBuilder::LaunchCompileTask(const TCHAR* Name, bool bCondition, LambdaType&& Lambda)
+///////////////////////////////////////////////////////////////////////////////////////////////////
+
+void FRDGBuilder::WaitForParallelSetupTasks()
 {
-	if (bCondition)
-	{
-		SCOPED_NAMED_EVENT_TCHAR(Name, FColor::Magenta);
-		return CompilePipe.Launch(Name, [Lambda] { Lambda(); }, LowLevelTasks::ETaskPriority::High);
-	}
-	else
+	if (!ParallelSetup.Tasks.IsEmpty())
 	{
-		Lambda();
-		return {};
+		UE::Tasks::Wait(ParallelSetup.Tasks);
+		ParallelSetup.Tasks.Reset();
 	}
 }
 
@@ -1728,62 +1636,139 @@ void FRDGBuilder::Execute()
 	}
 
 	// Create the epilogue pass at the end of the graph just prior to compilation.
-	SetupEmptyPass(EpiloguePass = Passes.Allocate<FRDGSentinelPass>(Allocator, RDG_EVENT_NAME("Graph Epilogue")));
+	SetupEmptyPass(EpiloguePass = Passes.Allocate<FRDGSentinelPass>(Allocators.Root, RDG_EVENT_NAME("Graph Epilogue")));
 
 	const FRDGPassHandle ProloguePassHandle = GetProloguePassHandle();
 	const FRDGPassHandle EpiloguePassHandle = GetEpiloguePassHandle();
 
-	UE::Tasks::FTask SubmitBufferUploadsTask;
-	UE::Tasks::FTask CompilePassBarriersTask;
-	UE::Tasks::FTask CreateUniformBuffersTask;
-	UE::Tasks::FTask SetupParallelExecuteTask;
+	UE::Tasks::FTask CollectPassBarriersTask;
+	UE::Tasks::FTask CreateViewsTask;
 
 	IF_RDG_ENABLE_DEBUG(UserValidation.ValidateExecuteBegin());
 	IF_RDG_ENABLE_DEBUG(GRDGAllowRHIAccess = true);
 
+	FCollectResourceContext CollectResourceContext;
+
 	if (!IsImmediateMode())
 	{
-		UE::Tasks::FTask CompileTask = LaunchCompileTask(TEXT("FRDGBuilder::Compile"), bParallelSetupEnabled, [this] { Compile(); });
-
-		CompilePassBarriersTask = LaunchCompileTask(TEXT("FRDGBuilder::CompilePassBarriers"), bParallelSetupEnabled, [this] { CompilePassBarriers(); });
+		WaitForParallelSetupTasks();
 
 		BeginFlushResourcesRHI();
 
-		if (!ParallelSetupEvents.IsEmpty())
+		if (ParallelSetup.bEnabled)
 		{
-			RHICmdList.QueueAsyncCommandListSubmit(ParallelSetupCommandLists);
-			UE::Tasks::Wait(ParallelSetupEvents);
-			ParallelSetupEvents.Empty();
+			AsyncSetupQueue.LastTask.Wait();
+			ProcessAsyncSetupQueue();
 		}
 
-		PrepareBufferUploads();
+		// Pre-allocate containers.
+		{
+			const int32 NumBuffers           = Buffers.Num();
+			const int32 NumTextures          = Textures.Num();
+			const int32 NumExternalBuffers   = ExternalBuffers.Num();
+			const int32 NumExternalTextures  = ExternalTextures.Num();
+			const int32 NumTransientBuffers  = bSupportsTransientBuffers ? (NumBuffers - NumExternalBuffers) : 0;
+			const int32 NumTransientTextures = bSupportsTransientTextures ? (NumTextures - NumExternalTextures) : 0;
+			const int32 NumPooledTextures    = NumTextures - NumTransientTextures;
+			const int32 NumPooledBuffers     = NumBuffers - NumTransientBuffers;
+
+			CollectResourceContext.TransientResources.Reserve(NumTransientBuffers + NumTransientTextures);
+			CollectResourceContext.PooledTextures.Reserve(bSupportsTransientTextures ? NumExternalTextures : NumTextures);
+			CollectResourceContext.PooledBuffers.Reserve(bSupportsTransientBuffers ? NumExternalBuffers : NumBuffers);
+			CollectResourceContext.UniformBuffers.Reserve(UniformBuffers.Num());
+			CollectResourceContext.Views.Reserve(Views.Num());
+			CollectResourceContext.UniformBufferMap.Init(true, UniformBuffers.Num());
+			CollectResourceContext.ViewMap.Init(true, Views.Num());
+
+			PooledBufferOwnershipMap.Reserve(NumPooledBuffers);
+			PooledTextureOwnershipMap.Reserve(NumPooledTextures);
+			ActivePooledTextures.Reserve(NumPooledTextures);
+			ActivePooledBuffers.Reserve(NumPooledBuffers);
+			EpilogueResourceAccesses.Reserve(NumTextures + NumBuffers);
+
+			ProloguePass->EpilogueBarriersToBeginForGraphics.Reserve(NumPooledBuffers + NumPooledTextures);
+		}
+
+		const UE::Tasks::ETaskPriority TaskPriority = UE::Tasks::ETaskPriority::High;
+
+		UE::Tasks::FTask BufferNumElementsCallbacksTask = AddSetupTask([this]
+		{
+			SCOPED_NAMED_EVENT_TEXT("FRDGBuilder::FinalizeDescs", FColor::Magenta);
+
+			for (FRDGBuffer* Buffer : NumElementsCallbackBuffers)
+			{
+				Buffer->FinalizeDesc();
+			}
+			NumElementsCallbackBuffers.Empty();
+
+		}, TaskPriority);
 
-		GPUScopeStacks.ReserveOps(Passes.Num());
-		IF_RDG_CPU_SCOPES(CPUScopeStacks.ReserveOps());
+		UE::Tasks::FTask PrepareCollectResourcesTask = AddSetupTask([this]
+		{
+			SCOPED_NAMED_EVENT_TEXT("FRDGBuilder::PrepareCollectResources", FColor::Magenta);
+
+			Buffers.Enumerate([&] (FRDGBuffer* Buffer)
+			{
+				if (Buffer->ResourceRHI || Buffer->bQueuedForUpload)
+				{
+					Buffer->bCollectForAllocate = false;
+				}
+
+				if (Buffer->TransientBuffer || (!Buffer->ResourceRHI && IsTransient(Buffer)))
+				{
+					Buffer->bTransient = true;
+				}
+			});
+
+			Textures.Enumerate([&] (FRDGTexture* Texture)
+			{
+				if (Texture->ResourceRHI)
+				{
+					Texture->bCollectForAllocate = false;
+				}
+
+				if (Texture->TransientTexture || (!Texture->ResourceRHI && IsTransient(Texture)))
+				{
+					Texture->bTransient = true;
+				}
+			});
+
+		}, TaskPriority);
+
+		UE::Tasks::FTaskEvent AllocateUploadBuffersTask{ UE_SOURCE_LOCATION };
+
+		UE::Tasks::FTask SubmitBufferUploadsTask = AddCommandListSetupTask([this, AllocateUploadBuffersTask] (FRHICommandListBase& RHICmdListTask) mutable
+		{
+			SubmitBufferUploads(RHICmdListTask, &AllocateUploadBuffersTask);
+
+		}, BufferNumElementsCallbacksTask, TaskPriority);
+
+		Compile();
+
+		CollectPassBarriersTask = AddSetupTask([this]
+		{
+			CompilePassBarriers();
+			CollectPassBarriers();
+
+		}, TaskPriority);
 
-		if (bParallelExecuteEnabled)
+		if (ParallelExecute.bEnabled)
 		{
-		#if RHI_WANT_BREADCRUMB_EVENTS
+#if RHI_WANT_BREADCRUMB_EVENTS
 			RHICmdList.ExportBreadcrumbState(*BreadcrumbState);
-		#endif
+#endif
 
-			// Parallel execute setup can be done off the render thread and synced prior to dispatch.
-			SetupParallelExecuteTask = LaunchCompileTask(TEXT("FRDGBuilder::SetupParallelExecute"), bParallelExecuteEnabled, [this] { SetupParallelExecute(); });
+			AddSetupTask([this] { SetupParallelExecute(); });
 		}
 
-		CompileTask.Wait();
+		UE::Tasks::FTask AllocatePooledBuffersTask;
+		UE::Tasks::FTask AllocatePooledTexturesTask;
 
 		{
 			SCOPE_CYCLE_COUNTER(STAT_RDG_CollectResourcesTime);
 			CSV_SCOPED_TIMING_STAT_EXCLUSIVE(RDG_CollectResources);
 			SCOPED_NAMED_EVENT_TEXT("FRDGBuilder::CollectResources", FColor::Magenta);
 
-			// Take locks on the buffer / texture pools as nothing else should be allocating from them right now and granular locking is expensive.
-			UE::TScopeLock RenderTargetPoolLock(GRenderTargetPool.Mutex);
-			UE::TScopeLock BufferPoolLock(GRenderGraphResourcePool.Mutex);
-
-			UniformBuffersToCreate.Reserve(UniformBuffers.Num());
-
 			EnumerateExtendedLifetimeResources(Textures, [](FRDGTexture* Texture)
 			{
 				++Texture->ReferenceCount;
@@ -1794,6 +1779,8 @@ void FRDGBuilder::Execute()
 				++Buffer->ReferenceCount;
 			});
 
+			PrepareCollectResourcesTask.Wait();
+
 			// Null out any culled external resources so that the reference is freed up.
 
 			for (const auto& Pair : ExternalTextures)
@@ -1802,7 +1789,7 @@ void FRDGBuilder::Execute()
 
 				if (Texture->IsCulled())
 				{
-					EndResourceRHI(ProloguePassHandle, Texture, 0);
+					CollectDeallocateTexture(CollectResourceContext, ProloguePassHandle, Texture, 0);
 				}
 			}
 
@@ -1812,17 +1799,7 @@ void FRDGBuilder::Execute()
 
 				if (Buffer->IsCulled())
 				{
-					EndResourceRHI(ProloguePassHandle, Buffer, 0);
-				}
-			}
-
-			for (FRDGPassHandle PassHandle = Passes.Begin(); PassHandle < ProloguePassHandle; ++PassHandle)
-			{
-				FRDGPass* Pass = Passes[PassHandle];
-
-				if (!Pass->bCulled)
-				{
-					EndResourcesRHI(Pass, ProloguePassHandle);
+					CollectDeallocateBuffer(CollectResourceContext, ProloguePassHandle, Buffer, 0);
 				}
 			}
 
@@ -1832,105 +1809,92 @@ void FRDGBuilder::Execute()
 
 				if (!Pass->bCulled)
 				{
-					BeginResourcesRHI(Pass, PassHandle);
-					EndResourcesRHI(Pass, PassHandle);
+					CollectAllocations(CollectResourceContext, Pass);
+					CollectDeallocations(CollectResourceContext, Pass);
 				}
 			}
 
 			EnumerateExtendedLifetimeResources(Textures, [&](FRDGTextureRef Texture)
 			{
-				EndResourceRHI(EpiloguePassHandle, Texture, 1);
+				CollectDeallocateTexture(CollectResourceContext, EpiloguePassHandle, Texture, 1);
 			});
 
 			EnumerateExtendedLifetimeResources(Buffers, [&](FRDGBufferRef Buffer)
 			{
-				EndResourceRHI(EpiloguePassHandle, Buffer, 1);
+				CollectDeallocateBuffer(CollectResourceContext, EpiloguePassHandle, Buffer, 1);
 			});
 
-			if (TransientResourceAllocator)
+			BufferNumElementsCallbacksTask.Wait();
+
+			AllocatePooledBuffersTask = AddCommandListSetupTask([this, PooledBuffers = MoveTemp(CollectResourceContext.PooledBuffers)] (FRHICommandListBase& RHICmdListTask)
 			{
-			#if RDG_ENABLE_TRACE
-				TransientResourceAllocator->Flush(RHICmdList, Trace.IsEnabled() ? &Trace.TransientAllocationStats : nullptr);
-			#else
-				TransientResourceAllocator->Flush(RHICmdList);
-			#endif
-			}
-		}
+				AllocatePooledBuffers(RHICmdListTask, PooledBuffers);
 
-		// We have to wait until after view creation to launch uploads because we can't lock / unlock while creating views simultaneously.
-		SubmitBufferUploadsTask = SubmitBufferUploads();
+			}, AllocateUploadBuffersTask, TaskPriority);
 
-		// Uniform buffer creation depends on view creation.
-		CreateUniformBuffersTask = CreateUniformBuffers();
+			AllocatePooledTexturesTask = AddCommandListSetupTask([this, PooledTextures = MoveTemp(CollectResourceContext.PooledTextures)] (FRHICommandListBase& RHICmdListTask)
+			{
+				AllocatePooledTextures(RHICmdListTask, PooledTextures);
 
-		{
-			SCOPED_NAMED_EVENT_TEXT("FRDGBuilder::CollectBarriers", FColor::Magenta);
-			SCOPE_CYCLE_COUNTER(STAT_RDG_CollectBarriersTime);
-			CSV_SCOPED_TIMING_STAT_EXCLUSIVE_CONDITIONAL(RDG_CollectBarriers, GRDGVerboseCSVStats != 0);
+			}, TaskPriority);
 
-			CompilePassBarriersTask.Wait();
+			AllocateTransientResources(MoveTemp(CollectResourceContext.TransientResources));
 
-			for (FRDGPassHandle PassHandle = ProloguePassHandle + 1; PassHandle < EpiloguePassHandle; ++PassHandle)
+			AddSetupTask([this]
 			{
-				FRDGPass* Pass = Passes[PassHandle];
+				FinalizeResources();
 
-				if (!Pass->bCulled && !Pass->bEmptyParameters)
-				{
-					CollectPassBarriers(Pass);
-				}
-			}
-		}
-	}
-	else
-	{
-		PrepareBufferUploads();
-		SubmitBufferUploads();
-		CreateUniformBuffers();
-	}
+			}, MakeArrayView<UE::Tasks::FTask>({ CollectPassBarriersTask, AllocatePooledBuffersTask, AllocatePooledTexturesTask }), TaskPriority);
 
-	{
-		SCOPED_NAMED_EVENT_TEXT("FRDGBuilder::Finalize", FColor::Magenta);
+			CreateViewsTask = AddCommandListSetupTask([this, Views = MoveTemp(CollectResourceContext.Views)] (FRHICommandListBase& RHICmdListTask)
+			{
+				CreateViews(RHICmdListTask, Views);
 
-		EpilogueResourceAccesses.Reserve(Textures.Num() + Buffers.Num());
+			}, MakeArrayView<UE::Tasks::FTask>({ AllocatePooledBuffersTask, AllocatePooledTexturesTask, SubmitBufferUploadsTask}), TaskPriority);
 
-		Textures.Enumerate([&](FRDGTextureRef Texture)
-		{
-			AddEpilogueTransition(Texture);
-		});
+			if (TransientResourceAllocator)
+			{
+#if RDG_ENABLE_TRACE
+				TransientResourceAllocator->Flush(RHICmdList, Trace.IsEnabled() ? &Trace.TransientAllocationStats : nullptr);
+#else
+				TransientResourceAllocator->Flush(RHICmdList);
+#endif
+			}
+		}
 
-		Buffers.Enumerate([&](FRDGBufferRef Buffer)
+		AddSetupTask([this, UniformBuffers = MoveTemp(CollectResourceContext.UniformBuffers)]
 		{
-			AddEpilogueTransition(Buffer);
-		});
-	}
+			CreateUniformBuffers(UniformBuffers);
 
-	const ENamedThreads::Type RenderThread = ENamedThreads::GetRenderThread_Local();
+		}, CreateViewsTask, TaskPriority); // Uniform buffer creation require views to be valid.
 
-	CreatePassBarriers([&]
+		AllocatePooledBuffersTask.Wait();
+		AllocatePooledTexturesTask.Wait();
+	}
+	else
 	{
-		CreateUniformBuffersTask.Wait();
+		SubmitBufferUploads(RHICmdList);
+		FinalizeResources();
+	}
 
-		if (SubmitBufferUploadsTask.IsValid())
-		{
-			SubmitBufferUploadsTask.Wait();
-			check(RHICmdListBufferUploads);
-			RHICmdList.QueueAsyncCommandListSubmit(RHICmdListBufferUploads);
-			RHICmdListBufferUploads = nullptr;
-		}
+	// Process RHI thread flush before helping with barrier compilation on the render thread.
+	EndFlushResourcesRHI();
 
-		// Process RHI thread flush before helping with barrier compilation on the render thread.
-		EndFlushResourcesRHI();
-	});
+	WaitForParallelSetupTasks();
 
-	UE::Tasks::FTask ParallelExecuteTask;
+	if (ParallelExecute.DispatchTaskEvent)
+	{
+		// Launch a task to absorb the cost of waking up threads and avoid stalling the render thread.
+		UE::Tasks::Launch(UE_SOURCE_LOCATION, [this] { ParallelExecute.DispatchTaskEvent->Trigger(); });
+	}
 
-	if (bParallelExecuteEnabled)
+	if (!ParallelSetup.CommandLists.IsEmpty())
 	{
-		SetupParallelExecuteTask.Wait();
-		ParallelExecuteTask = CompilePipe.Launch(TEXT("DispatchParallelExecute"), [this] { DispatchParallelExecute(); }, LowLevelTasks::ETaskPriority::High);
+		RHICmdList.QueueAsyncCommandListSubmit(ParallelSetup.CommandLists);
+		ParallelSetup.CommandLists.Empty();
 	}
 
-	IF_RDG_ENABLE_DEBUG(GRDGAllowRHIAccess = bParallelExecuteEnabled);
+	IF_RDG_ENABLE_DEBUG(GRDGAllowRHIAccess = ParallelExecute.bEnabled);
 	IF_RDG_ENABLE_TRACE(Trace.OutputGraphBegin());
 
 	if (!IsImmediateMode())
@@ -1952,7 +1916,7 @@ void FRDGBuilder::Execute()
 				continue;
 			}
 
-			if (bParallelExecuteEnabled)
+			if (ParallelExecute.bEnabled)
 			{
 				if (Pass->bParallelExecute)
 				{
@@ -1962,10 +1926,7 @@ void FRDGBuilder::Execute()
 
 					if (Pass->bParallelExecuteBegin)
 					{
-						FParallelPassSet& ParallelPassSet = ParallelPassSets[Pass->ParallelPassSetIndex];
-
-						// Busy wait until our pass set is ready. This will be set by the dispatch task.
-						while (!FPlatformAtomics::AtomicRead(&ParallelPassSet.bInitialized)) {};
+						FParallelPassSet& ParallelPassSet = ParallelExecute.ParallelPassSets[Pass->ParallelPassSetIndex];
 
 						FRHICommandListImmediate::ETranslatePriority TranslatePriority = ParallelPassSet.bParallelTranslate ? FRHICommandListImmediate::ETranslatePriority::Normal : FRHICommandListImmediate::ETranslatePriority::Disabled;
 
@@ -2007,15 +1968,12 @@ void FRDGBuilder::Execute()
 
 	RHICmdList.SetTrackedAccess(EpilogueResourceAccesses);
 
-	// Wait for the parallel dispatch task before attempting to wait on the execute event array (the former mutates the array).
-	ParallelExecuteTask.Wait();
-
 	// Wait on the actual parallel execute tasks in the Execute call. When draining is okay to let them overlap with other graph setup.
 	// This also needs to be done before extraction of external resources to be consistent with non-parallel rendering.
-	if (!ParallelExecuteEvents.IsEmpty())
+	if (!ParallelExecute.Tasks.IsEmpty())
 	{
-		UE::Tasks::Wait(ParallelExecuteEvents);
-		ParallelExecuteEvents.Empty();
+		UE::Tasks::Wait(ParallelExecute.Tasks);
+		ParallelExecute.Tasks.Empty();
 	}
 
 	for (const FExtractedTexture& ExtractedTexture : ExtractedTextures)
@@ -2043,7 +2001,7 @@ void FRDGBuilder::Execute()
 	GRDGStatBufferCount += Buffers.Num();
 	GRDGStatTextureCount += Textures.Num();
 	GRDGStatViewCount += Views.Num();
-	GRDGStatMemoryWatermark = FMath::Max(GRDGStatMemoryWatermark, Allocator.GetByteCount());
+	GRDGStatMemoryWatermark = FMath::Max(GRDGStatMemoryWatermark, Allocators.GetByteCount());
 #endif
 
 	RasterPassCount = 0;
@@ -2145,6 +2103,8 @@ void FRDGBuilder::MarkResourcesAsProduced(FRDGPass* Pass)
 
 void FRDGBuilder::SetupPassDependencies(FRDGPass* Pass)
 {
+	bool bIsCullRootProducer = false;
+
 	for (auto& PassState : Pass->TextureStates)
 	{
 		FRDGTextureRef Texture = PassState.Texture;
@@ -2154,26 +2114,26 @@ void FRDGBuilder::SetupPassDependencies(FRDGPass* Pass)
 
 		for (uint32 Index = 0, Count = LastProducers.Num(); Index < Count; ++Index)
 		{
-			const auto& SubresourceState = PassState.State[Index];
+			const FRDGSubresourceState* SubresourceState = PassState.State[Index];
 
-			if (SubresourceState.Access == ERHIAccess::Unknown)
+			if (!SubresourceState)
 			{
 				continue;
 			}
 
 			FRDGProducerState ProducerState;
 			ProducerState.Pass = Pass;
-			ProducerState.Access = SubresourceState.Access;
-			ProducerState.NoUAVBarrierHandle = SubresourceState.NoUAVBarrierFilter.GetUniqueHandle();
+			ProducerState.Access = SubresourceState->Access;
+			ProducerState.NoUAVBarrierHandle = SubresourceState->NoUAVBarrierFilter.GetUniqueHandle();
 
-			AddCullingDependency(LastProducers[Index], ProducerState, Pass->Pipeline);
+			bIsCullRootProducer |= AddCullingDependency(LastProducers[Index], ProducerState, Pass->Pipeline) && Texture->IsCullRoot();
 		}
 	}
 
 	for (auto& PassState : Pass->BufferStates)
 	{
 		FRDGBufferRef Buffer = PassState.Buffer;
-		const auto& SubresourceState = PassState.State;
+		const FRDGSubresourceState& SubresourceState = PassState.State;
 
 		Buffer->ReferenceCount += PassState.ReferenceCount;
 
@@ -2182,15 +2142,17 @@ void FRDGBuilder::SetupPassDependencies(FRDGPass* Pass)
 		ProducerState.Access = SubresourceState.Access;
 		ProducerState.NoUAVBarrierHandle = SubresourceState.NoUAVBarrierFilter.GetUniqueHandle();
 
-		AddCullingDependency(Buffer->LastProducer, ProducerState, Pass->Pipeline);
+		bIsCullRootProducer |= AddCullingDependency(Buffer->LastProducer, ProducerState, Pass->Pipeline) && Buffer->IsCullRoot();
 	}
 
 	const bool bCullPasses = GRDGCullPasses > 0;
 	Pass->bCulled = bCullPasses;
 
-	if (bCullPasses && (Pass->bHasExternalOutputs || EnumHasAnyFlags(Pass->Flags, ERDGPassFlags::NeverCull)))
+	if (bCullPasses && (bIsCullRootProducer || Pass->bHasExternalOutputs || EnumHasAnyFlags(Pass->Flags, ERDGPassFlags::NeverCull)))
 	{
-		CullPassStack.Emplace(Pass->Handle);
+		CullPassStack.Emplace(Pass);
+
+		FlushCullStack();
 	}
 }
 
@@ -2246,14 +2208,19 @@ void FRDGBuilder::SetupPassResources(FRDGPass* Pass)
 
 		PassState->ReferenceCount++;
 
-		EnumerateSubresourceRange(PassState->State, Texture->Layout, Range, [&](FRDGSubresourceState& State)
+		EnumerateSubresourceRange(PassState->State, Texture->Layout, Range, [&](FRDGSubresourceState*& State)
 		{
-			IF_RDG_ENABLE_DEBUG(UserValidation.ValidateAddSubresourceAccess(Texture, State, Access));
+			if (!State)
+			{
+				State = AllocSubresource();
+			}
+
+			IF_RDG_ENABLE_DEBUG(UserValidation.ValidateAddSubresourceAccess(Texture, *State, Access));
 
-			State.Access = MakeValidAccess(State.Access, Access);
-			State.Flags |= TransitionFlags;
-			State.NoUAVBarrierFilter.AddHandle(NoUAVBarrierHandle);
-			State.SetPass(PassPipeline, PassHandle);
+			State->Access = MakeValidAccess(State->Access, Access);
+			State->Flags |= TransitionFlags;
+			State->NoUAVBarrierFilter.AddHandle(NoUAVBarrierHandle);
+			State->SetPass(PassPipeline, PassHandle);
 		});
 
 		if (IsWritableAccess(Access))
@@ -2261,7 +2228,7 @@ void FRDGBuilder::SetupPassResources(FRDGPass* Pass)
 			bRenderPassOnlyWrites &= EnumHasAnyFlags(AccessFlags, ERDGTextureAccessFlags::RenderTarget);
 
 			// When running in parallel this is set via MarkResourcesAsProduced. We also can't touch this as its a bitfield and not atomic.
-			if (!bParallelSetupEnabled)
+			if (!ParallelSetup.bEnabled)
 			{
 				Texture->bProduced = true;
 			}
@@ -2310,7 +2277,7 @@ void FRDGBuilder::SetupPassResources(FRDGPass* Pass)
 			bRenderPassOnlyWrites = false;
 
 			// When running in parallel this is set via MarkResourcesAsProduced. We also can't touch this as its a bitfield and not atomic.
-			if (!bParallelSetupEnabled)
+			if (!ParallelSetup.bEnabled)
 			{
 				Buffer->bProduced = true;
 			}
@@ -2327,7 +2294,7 @@ void FRDGBuilder::SetupPassResources(FRDGPass* Pass)
 		Pass->UniformBuffers.Emplace(UniformBuffer.GetUniformBuffer()->Handle);
 	});
 
-	if (bParallelSetupEnabled)
+	if (ParallelSetup.bEnabled)
 	{
 		SetupPassDependencies(Pass);
 
@@ -2344,8 +2311,6 @@ void FRDGBuilder::SetupPassInternals(FRDGPass* Pass)
 	const ERDGPassFlags PassFlags = Pass->Flags;
 	const ERHIPipeline PassPipeline = Pass->Pipeline;
 
-	Pass->GraphicsForkPass = PassHandle;
-	Pass->GraphicsJoinPass = PassHandle;
 	Pass->PrologueBarrierPass = PassHandle;
 	Pass->EpilogueBarrierPass = PassHandle;
 
@@ -2395,40 +2360,56 @@ void FRDGBuilder::SetupAuxiliaryPasses(FRDGPass* Pass)
 		SCOPED_NAMED_EVENT(FRDGBuilder_ExecutePass, FColor::Emerald);
 		RDG_ALLOW_RHI_ACCESS_SCOPE();
 
-		// Trivially redirect the merge states to the pass states, since we won't be compiling the graph.
 		for (auto& PassState : Pass->TextureStates)
 		{
-			const uint32 SubresourceCount = PassState.State.Num();
-			PassState.MergeState.SetNum(SubresourceCount);
-			for (uint32 Index = 0; Index < SubresourceCount; ++Index)
+			FRDGTexture* Texture = PassState.Texture;
+
+			if (Texture->ResourceRHI)
+			{
+				Texture->bCollectForAllocate = false;
+			}
+
+			for (FRDGSubresourceState*& SubresourceState : Texture->State)
 			{
-				if (PassState.State[Index].Access != ERHIAccess::Unknown)
+				if (!SubresourceState)
 				{
-					PassState.MergeState[Index] = &PassState.State[Index];
+					SubresourceState = &PrologueSubresourceState;
 				}
 			}
+
+			PassState.MergeState = PassState.State;
 		}
 
 		for (auto& PassState : Pass->BufferStates)
 		{
+			FRDGBuffer* Buffer = PassState.Buffer;
+
+			if (Buffer->ResourceRHI || Buffer->bQueuedForUpload)
+			{
+				Buffer->bCollectForAllocate = false;
+			}
+
+			if (!Buffer->State)
+			{
+				Buffer->State = &PrologueSubresourceState;
+			}
+
 			PassState.MergeState = &PassState.State;
 		}
 
 		check(!EnumHasAnyFlags(Pass->Pipeline, ERHIPipeline::AsyncCompute));
-		check(ParallelSetupEvents.IsEmpty());
+		check(ParallelSetup.Tasks.IsEmpty());
 
-		PrepareBufferUploads();
-		SubmitBufferUploads();
+		FCollectResourceContext Context;
+		SubmitBufferUploads(RHICmdList);
 		CompilePassOps(Pass);
-		BeginResourcesRHI(Pass, Pass->Handle);
-		CreateUniformBuffers();
-		CollectPassBarriers(Pass);
-		CreatePassBarriers([] {});
-		if (!ParallelSetupEvents.IsEmpty())
-		{
-			UE::Tasks::Wait(ParallelSetupEvents);
-			ParallelSetupEvents.Reset();
-		}
+		CollectAllocations(Context, Pass);
+		AllocatePooledTextures(RHICmdList, Context.PooledTextures);
+		AllocatePooledBuffers(RHICmdList, Context.PooledBuffers);
+		CreateViews(RHICmdList, Context.Views);
+		CreateUniformBuffers(Context.UniformBuffers);
+		CollectPassBarriers(Pass->Handle);
+		CreatePassBarriers();
 		ExecutePass(Pass, RHICmdList);
 	}
 
@@ -2446,10 +2427,10 @@ FRDGPass* FRDGBuilder::SetupParameterPass(FRDGPass* Pass)
 
 	SetupPassInternals(Pass);
 
-	if (bParallelSetupEnabled)
+	if (ParallelSetup.bEnabled)
 	{
 		MarkResourcesAsProduced(Pass);
-		SetupPassQueue.Push(Pass);
+		AsyncSetupQueue.Push(FAsyncSetupOp::SetupPassResources(Pass));
 	}
 	else
 	{
@@ -2497,81 +2478,63 @@ void FRDGBuilder::CompilePassOps(FRDGPass* Pass)
 
 ///////////////////////////////////////////////////////////////////////////////////////////////////
 
-void FRDGBuilder::PrepareBufferUploads()
+void FRDGBuilder::SubmitBufferUploads(FRHICommandListBase& RHICmdListUpload, UE::Tasks::FTaskEvent* AllocateUploadBuffersTask)
 {
-	SCOPED_NAMED_EVENT_TEXT("FRDGBuilder::PrepareBufferUploads", FColor::Magenta);
+	SCOPED_NAMED_EVENT_TEXT("FRDGBuilder::SubmitBufferUploads", FColor::Magenta);
 
-	for (FUploadedBuffer& UploadedBuffer : UploadedBuffers)
 	{
-		FRDGBuffer* Buffer = UploadedBuffer.Buffer;
+		SCOPED_NAMED_EVENT_TEXT("Allocate", FColor::Magenta);
+		UE::TScopeLock Lock(GRenderGraphResourcePool.Mutex);
 
-		if (!Buffer->HasRHI())
+		for (FUploadedBuffer& UploadedBuffer : UploadedBuffers)
 		{
-			BeginResourceRHI(GetProloguePassHandle(), Buffer);
+			FRDGBuffer* Buffer = UploadedBuffer.Buffer;
+			if (!Buffer->ResourceRHI)
+			{
+				SetPooledBufferRHI(Buffer, AllocatePooledBufferRHI(RHICmdListUpload, Buffer));
+			}
 		}
+	}
 
-		check(UploadedBuffer.DataSize <= Buffer->Desc.GetSize());
+	if (AllocateUploadBuffersTask)
+	{
+		AllocateUploadBuffersTask->Trigger();
 	}
-}
 
-UE::Tasks::FTask FRDGBuilder::SubmitBufferUploads()
-{
-	const auto SubmitUploadsLambda = [this](FRHICommandList& RHICmdListUpload)
 	{
-		SCOPED_NAMED_EVENT_TEXT("FRDGBuilder::SubmitBufferUploads", FColor::Magenta);
+		SCOPED_NAMED_EVENT_TEXT("Upload", FColor::Magenta);
 
 		for (FUploadedBuffer& UploadedBuffer : UploadedBuffers)
 		{
-			if (UploadedBuffer.bUseDataCallbacks)
+			FRDGBuffer* Buffer = UploadedBuffer.Buffer;
+
+			if (UploadedBuffer.DataFillCallback)
 			{
-				UploadedBuffer.Data = UploadedBuffer.DataCallback();
-				UploadedBuffer.DataSize = UploadedBuffer.DataSizeCallback();
+				const uint32 DataSize = Buffer->Desc.GetSize();
+				void* DestPtr = RHICmdListUpload.LockBuffer(Buffer->GetRHIUnchecked(), 0, DataSize, RLM_WriteOnly);
+				UploadedBuffer.DataFillCallback(DestPtr, DataSize);
+				RHICmdListUpload.UnlockBuffer(Buffer->GetRHIUnchecked());
 			}
-
-			if (UploadedBuffer.Data && UploadedBuffer.DataSize)
+			else
 			{
-				void* DestPtr = RHICmdListUpload.LockBuffer(UploadedBuffer.Buffer->GetRHIUnchecked(), 0, UploadedBuffer.DataSize, RLM_WriteOnly);
-				FMemory::Memcpy(DestPtr, UploadedBuffer.Data, UploadedBuffer.DataSize);
-				RHICmdListUpload.UnlockBuffer(UploadedBuffer.Buffer->GetRHIUnchecked());
-
-				if (UploadedBuffer.bUseFreeCallbacks)
+				if (UploadedBuffer.bUseDataCallbacks)
 				{
-					UploadedBuffer.DataFreeCallback(UploadedBuffer.Data);
+					UploadedBuffer.Data = UploadedBuffer.DataCallback();
+					UploadedBuffer.DataSize = UploadedBuffer.DataSizeCallback();
 				}
-			}
-			else if (UploadedBuffer.DataFillCallback)
-			{
-				const uint32 BufferSize = UploadedBuffer.Buffer->Desc.GetSize();
-				check(BufferSize > 0);
 
-				void* DestPtr = RHICmdListUpload.LockBuffer(UploadedBuffer.Buffer->GetRHIUnchecked(), 0, BufferSize, RLM_WriteOnly);
-				UploadedBuffer.DataFillCallback(DestPtr, BufferSize);
-				RHICmdListUpload.UnlockBuffer(UploadedBuffer.Buffer->GetRHIUnchecked());
+				if (UploadedBuffer.Data && UploadedBuffer.DataSize)
+				{
+					check(UploadedBuffer.DataSize <= Buffer->Desc.GetSize());
+					void* DestPtr = RHICmdListUpload.LockBuffer(Buffer->GetRHIUnchecked(), 0, UploadedBuffer.DataSize, RLM_WriteOnly);
+					FMemory::Memcpy(DestPtr, UploadedBuffer.Data, UploadedBuffer.DataSize);
+					RHICmdListUpload.UnlockBuffer(Buffer->GetRHIUnchecked());
+				}
 			}
 		}
-
-		UploadedBuffers.Reset();
-	};
-
-	if (bParallelSetupEnabled)
-	{
-		SCOPED_NAMED_EVENT_TEXT("FRDGBuilder::SubmitBufferUploads", FColor::Magenta);
-
-		return UE::Tasks::Launch(TEXT("FRDGBuilder::SubmitBufferUploads"), [this, SubmitUploadsLambda]
-		{
-			FOptionalTaskTagScope Scope(ETaskTag::EParallelRenderingThread);
-			RHICmdListBufferUploads = new FRHICommandList(FRHIGPUMask::All());
-			RHICmdListBufferUploads->SwitchPipeline(ERHIPipeline::Graphics);
-			SubmitUploadsLambda(*RHICmdListBufferUploads);
-			RHICmdListBufferUploads->FinishRecording();
-
-		}, LowLevelTasks::ETaskPriority::High);
-	}
-	else
-	{
-		SubmitUploadsLambda(RHICmdList);
-		return {};
 	}
+
+	UploadedBuffers.Reset();
 }
 
 ///////////////////////////////////////////////////////////////////////////////////////////////////
@@ -2579,13 +2542,15 @@ UE::Tasks::FTask FRDGBuilder::SubmitBufferUploads()
 void FRDGBuilder::SetupParallelExecute()
 {
 	SCOPED_NAMED_EVENT(SetupParallelExecute, FColor::Emerald);
-	FOptionalTaskTagScope Scope(ETaskTag::EParallelRenderingThread);
-
+	FRDGAllocatorScope AllocatorScope(Allocators.Task);
 	TArray<FRDGPass*, TInlineAllocator<64, FRDGArrayAllocator>> ParallelPassCandidates;
 	uint32 ParallelPassCandidatesWorkload = 0;
 	bool bDispatchAfterExecute = false;
 	bool bParallelTranslate = false;
 
+	GPUScopeStacks.ReserveOps(Passes.Num());
+	IF_RDG_CPU_SCOPES(CPUScopeStacks.ReserveOps());
+
 	const auto FlushParallelPassCandidates = [&]()
 	{
 		if (ParallelPassCandidates.IsEmpty())
@@ -2634,18 +2599,18 @@ void FRDGBuilder::SetupParallelExecute()
 		{
 			FRDGPass* PassBegin = ParallelPassCandidates[PassBeginIndex];
 			PassBegin->bParallelExecuteBegin = 1;
-			PassBegin->ParallelPassSetIndex = ParallelPassSets.Num();
+			PassBegin->ParallelPassSetIndex = ParallelExecute.ParallelPassSets.Num();
 
 			FRDGPass* PassEnd = ParallelPassCandidates[PassEndIndex - 1];
 			PassEnd->bParallelExecuteEnd = 1;
-			PassEnd->ParallelPassSetIndex = ParallelPassSets.Num();
+			PassEnd->ParallelPassSetIndex = ParallelExecute.ParallelPassSets.Num();
 
 			for (int32 PassIndex = PassBeginIndex; PassIndex < PassEndIndex; ++PassIndex)
 			{
 				ParallelPassCandidates[PassIndex]->bParallelExecute = 1;
 			}
 
-			FParallelPassSet& ParallelPassSet = ParallelPassSets.Emplace_GetRef();
+			FParallelPassSet& ParallelPassSet = ParallelExecute.ParallelPassSets.Emplace_GetRef();
 			ParallelPassSet.Passes.Append(ParallelPassCandidates.GetData() + PassBeginIndex, ParallelPassCandidateCount);
 			ParallelPassSet.bDispatchAfterExecute = bDispatchAfterExecute;
 			ParallelPassSet.bParallelTranslate = bParallelTranslate;
@@ -2657,7 +2622,7 @@ void FRDGBuilder::SetupParallelExecute()
 		bParallelTranslate = false;
 	};
 
-	ParallelPassSets.Reserve(32);
+	ParallelExecute.ParallelPassSets.Reserve(32);
 	ParallelPassCandidates.Emplace(ProloguePass);
 
 	for (FRDGPassHandle PassHandle = GetProloguePassHandle() + 1; PassHandle < GetEpiloguePassHandle(); ++PassHandle)
@@ -2687,7 +2652,11 @@ void FRDGBuilder::SetupParallelExecute()
 		bParallelTranslate |= bPassHasParallelTranslate;
 
 		ParallelPassCandidates.Emplace(Pass);
-		ParallelPassCandidatesWorkload += Pass->Workload;
+
+		if (!Pass->bSkipRenderPassBegin && !Pass->bSkipRenderPassEnd)
+		{
+			ParallelPassCandidatesWorkload += Pass->Workload;
+		}
 
 		if (ParallelPassCandidatesWorkload >= (uint32)GRDGParallelExecutePassMax)
 		{
@@ -2712,8 +2681,8 @@ void FRDGBuilder::SetupParallelExecute()
 
 		if (Pass->bParallelExecuteBegin)
 		{
-			FParallelPassSet& ParallelPassSet = ParallelPassSets[Pass->ParallelPassSetIndex];
-			ParallelPassSet.BreadcrumbStateBegin = BreadcrumbState->Copy(Allocator);
+			FParallelPassSet& ParallelPassSet = ParallelExecute.ParallelPassSets[Pass->ParallelPassSetIndex];
+			ParallelPassSet.BreadcrumbStateBegin = BreadcrumbState->Copy(Allocators.Task);
 			ParallelPassSet.BreadcrumbStateEnd = ParallelPassSet.BreadcrumbStateBegin;
 		}
 
@@ -2722,34 +2691,27 @@ void FRDGBuilder::SetupParallelExecute()
 
 		if (Pass->bParallelExecuteEnd)
 		{
-			FParallelPassSet& ParallelPassSet = ParallelPassSets[Pass->ParallelPassSetIndex];
+			FParallelPassSet& ParallelPassSet = ParallelExecute.ParallelPassSets[Pass->ParallelPassSetIndex];
 
 			if (ParallelPassSet.BreadcrumbStateEnd->Version != BreadcrumbState->Version)
 			{
-				ParallelPassSet.BreadcrumbStateEnd = BreadcrumbState->Copy(Allocator);
+				ParallelPassSet.BreadcrumbStateEnd = BreadcrumbState->Copy(Allocators.Task);
 			}
 		}
 	}
 #endif
-}
-
-void FRDGBuilder::DispatchParallelExecute()
-{
-	SCOPED_NAMED_EVENT(DispatchParallelExecute, FColor::Emerald);
 
-	check(ParallelExecuteEvents.IsEmpty());
-	ParallelExecuteEvents.Reserve(ParallelPassSets.Num());
+	check(ParallelExecute.Tasks.IsEmpty());
+	ParallelExecute.Tasks.Reserve(ParallelExecute.ParallelPassSets.Num());
+	ParallelExecute.DispatchTaskEvent.Emplace(UE_SOURCE_LOCATION);
 
-	for (FParallelPassSet& ParallelPassSet : ParallelPassSets)
+	for (FParallelPassSet& ParallelPassSet : ParallelExecute.ParallelPassSets)
 	{
 		FRHICommandList* RHICmdListPass = new FRHICommandList(FRHIGPUMask::All());
 		ParallelPassSet.CmdList = RHICmdListPass;
 		IF_RHI_WANT_BREADCRUMB_EVENTS(RHICmdListPass->ImportBreadcrumbState(*ParallelPassSet.BreadcrumbStateBegin));
 
-		// Mark this set as initialized so that it can be submitted.
-		FPlatformAtomics::AtomicStore(&ParallelPassSet.bInitialized, 1);
-
-		ParallelExecuteEvents.Emplace(UE::Tasks::Launch(TEXT("FRDGBuilder::ParallelExecute"), [this, &ParallelPassSet, RHICmdListPass]
+		ParallelExecute.Tasks.Emplace(UE::Tasks::Launch(TEXT("FRDGBuilder::ParallelExecute"), [this, &ParallelPassSet, RHICmdListPass]
 		{
 			SCOPED_NAMED_EVENT(ParallelExecute, FColor::Emerald);
 			FOptionalTaskTagScope Scope(ETaskTag::EParallelRenderingThread);
@@ -2761,53 +2723,187 @@ void FRDGBuilder::DispatchParallelExecute()
 
 			RHICmdListPass->FinishRecording();
 
-		}, LowLevelTasks::ETaskPriority::High));
+		}, *ParallelExecute.DispatchTaskEvent, LowLevelTasks::ETaskPriority::High));
 	}
 }
 
 ///////////////////////////////////////////////////////////////////////////////////////////////////
 
-UE::Tasks::FTask FRDGBuilder::CreateUniformBuffers()
+void FRDGBuilder::AllocatePooledTextures(FRHICommandListBase& InRHICmdList, TConstArrayView<FCollectResourceOp> Ops)
 {
-	SCOPED_NAMED_EVENT_TEXT("FRDGBuilder::CreateUniformBuffers", FColor::Magenta);
+	SCOPED_NAMED_EVENT_TEXT("FRDGBuilder::AllocatePooledTextures", FColor::Magenta);
+	UE::TScopeLock Lock(GRenderTargetPool.Mutex);
+
+	for (FCollectResourceOp Op : Ops)
+	{
+		FRDGTexture* Texture = Textures[FRDGTextureHandle(Op.ResourceIndex)];
 
-	const int32 ParallelDispatchThreshold = 4;
+		switch (Op.Op)
+		{
+		case FCollectResourceOp::EOp::Allocate:
+			SetPooledRenderTargetRHI(Texture, AllocatePooledRenderTargetRHI(InRHICmdList, Texture));
+			break;
+		case FCollectResourceOp::EOp::Deallocate:
+			if (static_cast<FPooledRenderTarget*>(Texture->RenderTarget)->IsTracked())
+			{
+				// This releases the reference without invoking a virtual function call.
+				TRefCountPtr<FPooledRenderTarget>(MoveTemp(Texture->Allocation));
+			}
+			break;
+		}
+	}
 
-	const auto CreateUniformBuffersFunction = [this]
+	for (FCollectResourceOp Op : Ops)
 	{
-		SCOPED_NAMED_EVENT_TEXT("FRDGBuilder::CreateUniformBuffers", FColor::Magenta);
-		for (FRDGUniformBufferHandle UniformBufferHandle : UniformBuffersToCreate)
+		FRDGTexture* Texture = Textures[FRDGTextureHandle(Op.ResourceIndex)];
+
+		if (!Texture->bSkipLastTransition)
 		{
-			UniformBuffers[UniformBufferHandle]->InitRHI();
+			// Hold the last reference in a chain of pooled allocations.
+			Texture->Allocation = Texture->RenderTarget;
 		}
-		UniformBuffersToCreate.Reset();
-	};
+	}
+}
 
-	UE::Tasks::FTask Task;
+void FRDGBuilder::AllocatePooledBuffers(FRHICommandListBase& InRHICmdList, TConstArrayView<FCollectResourceOp> Ops)
+{
+	SCOPED_NAMED_EVENT_TEXT("FRDGBuilder::AllocatePooledBuffers", FColor::Magenta);
+	UE::TScopeLock Lock(GRenderGraphResourcePool.Mutex);
 
-	if (bParallelSetupEnabled && UniformBuffersToCreate.Num() > ParallelDispatchThreshold)
+	for (FCollectResourceOp Op : Ops)
 	{
-		Task = UE::Tasks::Launch(TEXT("FRDGBuilder::CreateUniformBuffer"),
-			[CreateUniformBuffersFunction]
-		{
-			FOptionalTaskTagScope Scope(ETaskTag::EParallelRenderingThread);
-			CreateUniformBuffersFunction();
+		FRDGBuffer* Buffer = Buffers[FRDGBufferHandle(Op.ResourceIndex)];
 
-		}, LowLevelTasks::ETaskPriority::High);
+		switch (Op.Op)
+		{
+		case FCollectResourceOp::EOp::Allocate:
+			SetPooledBufferRHI(Buffer, AllocatePooledBufferRHI(InRHICmdList, Buffer));
+			break;
+		case FCollectResourceOp::EOp::Deallocate:
+			Buffer->Allocation = nullptr;
+			break;
+		}
 	}
-	else
+
+	for (FCollectResourceOp Op : Ops)
 	{
-		CreateUniformBuffersFunction();
+		FRDGBuffer* Buffer = Buffers[FRDGBufferHandle(Op.ResourceIndex)];
+
+		if (!Buffer->bSkipLastTransition)
+		{
+			// Hold the last reference in a chain of pooled allocations.
+			Buffer->Allocation = Buffer->PooledBuffer;
+		}
 	}
+}
+
+void FRDGBuilder::AllocateTransientResources(TConstArrayView<FCollectResourceOp> Ops)
+{
+	SCOPED_NAMED_EVENT_TEXT("FRDGBuilder::AllocateTransientResources", FColor::Magenta);
+	for (FCollectResourceOp Op : Ops)
+	{
+		const FRDGPassHandle PassHandle = Op.PassHandle;
+
+		switch (Op.Op)
+		{
+		case FCollectResourceOp::EOp::Allocate:
+		{
+			if (Op.ResourceType == ERDGViewableResourceType::Buffer)
+			{
+				FRDGBuffer* Buffer = Buffers[FRDGBufferHandle(Op.ResourceIndex)];
+				FRHITransientBuffer* TransientBuffer = TransientResourceAllocator->CreateBuffer(Translate(Buffer->Desc), Buffer->Name, PassHandle.GetIndex());
+
+				SetTransientBufferRHI(Buffer, TransientBuffer);
+
+				Buffer->MinAcquirePass = FRDGPassHandle(TransientBuffer->GetAcquirePasses().Min);
+			}
+			else
+			{
+				FRDGTexture* Texture = Textures[FRDGTextureHandle(Op.ResourceIndex)];
+				FRHITransientTexture* TransientTexture = TransientResourceAllocator->CreateTexture(Texture->Desc, Texture->Name, PassHandle.GetIndex());
+
+				if (Texture->bExternal || Texture->bExtracted)
+				{
+					SetPooledRenderTargetRHI(Texture, GRDGTransientResourceAllocator.AllocateRenderTarget(TransientTexture));
+				}
+				else
+				{
+					SetTransientTextureRHI(Texture, TransientTexture);
+				}
+
+				Texture->MinAcquirePass = FRDGPassHandle(TransientTexture->GetAcquirePasses().Min);
+			}
+		}
+		break;
+		case FCollectResourceOp::EOp::Deallocate:
+		{
+			if (Op.ResourceType == ERDGViewableResourceType::Buffer)
+			{
+				FRDGBuffer* Buffer = Buffers[FRDGBufferHandle(Op.ResourceIndex)];
+				FRHITransientBuffer* TransientBuffer = Buffer->TransientBuffer;
+				TransientResourceAllocator->DeallocateMemory(TransientBuffer, PassHandle.GetIndex());
+
+				Buffer->MinDiscardPass = FRDGPassHandle(TransientBuffer->GetDiscardPasses().Min);
+				Buffer->MaxDiscardPass = FRDGPassHandle(FMath::Min<uint32>(TransientBuffer->GetDiscardPasses().Max, GetEpiloguePassHandle().GetIndex()));
+			}
+			else
+			{
+				FRDGTexture* Texture = Textures[FRDGTextureHandle(Op.ResourceIndex)];
+				FRHITransientTexture* TransientTexture = Texture->TransientTexture;
+
+				// Texture is using a transient external render target.
+				if (Texture->RenderTarget)
+				{
+					if (!Texture->bExtracted)
+					{
+						// This releases the reference without invoking a virtual function call.
+						GRDGTransientResourceAllocator.Release(TRefCountPtr<FRDGTransientRenderTarget>(MoveTemp(Texture->Allocation)), PassHandle);
+					}
+				}
+				// Texture is using an internal transient texture.
+				else
+				{
+					TransientResourceAllocator->DeallocateMemory(TransientTexture, PassHandle.GetIndex());
+				}
 
-	return Task;
+				if (!TransientTexture->IsAcquired())
+				{
+					Texture->MinDiscardPass = FRDGPassHandle(TransientTexture->GetDiscardPasses().Min);
+					Texture->MaxDiscardPass = FRDGPassHandle(FMath::Min<uint32>(TransientTexture->GetDiscardPasses().Max, GetEpiloguePassHandle().GetIndex()));
+				}
+			}
+		}
+		break;
+		}
+	}
 }
 
-///////////////////////////////////////////////////////////////////////////////////////////////////
+void FRDGBuilder::CreateViews(FRHICommandListBase& InRHICmdList, TConstArrayView<FRDGViewHandle> ViewsToCreate)
+{
+	SCOPED_NAMED_EVENT_TEXT("FRDGBuilder::CreateViews", FColor::Magenta);
+	for (FRDGViewHandle ViewHandle : ViewsToCreate)
+	{
+		FRDGView* View = Views[ViewHandle];
+
+		if (!View->ResourceRHI)
+		{
+			InitViewRHI(InRHICmdList, View);
+		}
+	}
+}
 
-void FRDGBuilder::AddProloguePass()
+void FRDGBuilder::CreateUniformBuffers(TConstArrayView<FRDGUniformBufferHandle> UniformBuffersToCreate)
 {
-	ProloguePass = SetupEmptyPass(Passes.Allocate<FRDGSentinelPass>(Allocator, RDG_EVENT_NAME("Graph Prologue (Graphics)")));
+	SCOPED_NAMED_EVENT_TEXT("FRDGBuilder::CreateUniformBuffers", FColor::Magenta);
+	for (FRDGUniformBufferHandle UniformBufferHandle : UniformBuffersToCreate)
+	{
+		FRDGUniformBuffer* UniformBuffer = UniformBuffers[UniformBufferHandle];
+
+		if (!UniformBuffer->ResourceRHI)
+		{
+			UniformBuffer->InitRHI();
+		}
+	}
 }
 
 ///////////////////////////////////////////////////////////////////////////////////////////////////
@@ -2961,7 +3057,7 @@ void FRDGBuilder::ExecutePass(FRDGPass* Pass, FRHIComputeCommandList& RHICmdList
 
 	if (GRDGDebugFlushGPU)
 	{
-		check(!GRDGAsyncCompute && !bParallelExecuteEnabled);
+		check(!GRDGAsyncCompute && !ParallelExecute.bEnabled);
 		RHICmdList.SubmitCommandsAndFlushGPU();
 		RHICmdList.BlockUntilGPUIdle();
 	}
@@ -2969,218 +3065,618 @@ void FRDGBuilder::ExecutePass(FRDGPass* Pass, FRHIComputeCommandList& RHICmdList
 
 ///////////////////////////////////////////////////////////////////////////////////////////////////
 
-void FRDGBuilder::BeginResourcesRHI(FRDGPass* ResourcePass, FRDGPassHandle ExecutePassHandle)
+void FRDGBuilder::CollectAllocations(FCollectResourceContext& Context, FRDGPass* Pass)
 {
-	for (FRDGPass* PassToBegin : ResourcePass->ResourcesToBegin)
-	{
-		for (const auto& PassState : PassToBegin->TextureStates)
-		{
-			BeginResourceRHI(ExecutePassHandle, PassState.Texture);
-		}
-
-		for (const auto& PassState : PassToBegin->BufferStates)
-		{
-			BeginResourceRHI(ExecutePassHandle, PassState.Buffer);
-		}
-
-		for (FRDGUniformBufferHandle UniformBufferHandle : PassToBegin->UniformBuffers)
-		{
-			if (FRDGUniformBuffer* UniformBuffer = UniformBuffers[UniformBufferHandle]; !UniformBuffer->bQueuedForCreate)
-			{
-				UniformBuffer->bQueuedForCreate = true;
-				UniformBuffersToCreate.Add(UniformBufferHandle);
-			}
-		}
-
-		for (FRDGViewHandle ViewHandle : PassToBegin->Views)
-		{
-			InitRHI(Views[ViewHandle]);
-		}
+	for (FRDGPass* PassToBegin : Pass->ResourcesToBegin)
+	{
+	  for (FRDGPass::FTextureState& PassState : PassToBegin->TextureStates)
+	  {
+		  CollectAllocateTexture(Context, Pass->Handle, PassState.Texture);
+	  }
+  
+	  for (FRDGPass::FBufferState& PassState : PassToBegin->BufferStates)
+	  {
+		  CollectAllocateBuffer(Context, Pass->Handle, PassState.Buffer);
+	  }
+  
+	  if (!IsImmediateMode())
+	  {
+		  for (FRDGUniformBufferHandle UniformBufferHandle : PassToBegin->UniformBuffers)
+		  {
+			  if (auto BitRef = Context.UniformBufferMap[UniformBufferHandle]; BitRef)
+			  {
+				  Context.UniformBuffers.Add(UniformBufferHandle);
+				  BitRef = false;
+			  }
+		  }
+  
+		  for (FRDGViewHandle ViewHandle : PassToBegin->Views)
+		  {
+			  if (auto BitRef = Context.ViewMap[ViewHandle]; BitRef)
+			  {
+				  Context.Views.Add(ViewHandle);
+				  BitRef = false;
+			  }
+		  }
+	  }
+	  else
+	  {
+		  Context.UniformBuffers = PassToBegin->UniformBuffers;
+		  Context.Views = PassToBegin->Views;
+	  }
 	}
 }
 
-void FRDGBuilder::EndResourcesRHI(FRDGPass* ResourcePass, FRDGPassHandle ExecutePassHandle)
+void FRDGBuilder::CollectDeallocations(FCollectResourceContext& Context, FRDGPass* Pass)
 {
-	for (FRDGPass* PassToEnd : ResourcePass->ResourcesToEnd)
+	for (FRDGPass* PassToEnd : Pass->ResourcesToEnd)
 	{
-		for (const auto& PassState : PassToEnd->TextureStates)
+		for (FRDGPass::FTextureState& PassState : PassToEnd->TextureStates)
 		{
-			EndResourceRHI(ExecutePassHandle, PassState.Texture, PassState.ReferenceCount);
+			CollectDeallocateTexture(Context, Pass->Handle, PassState.Texture, PassState.ReferenceCount);
 		}
 
-		for (const auto& PassState : PassToEnd->BufferStates)
+		for (FRDGPass::FBufferState& PassState : PassToEnd->BufferStates)
 		{
-			EndResourceRHI(ExecutePassHandle, PassState.Buffer, PassState.ReferenceCount);
+			CollectDeallocateBuffer(Context, Pass->Handle, PassState.Buffer, PassState.ReferenceCount);
 		}
 	}
 }
 
-///////////////////////////////////////////////////////////////////////////////////////////////////
-
-void FRDGBuilder::CollectPassBarriers(FRDGPass* Pass)
+void FRDGBuilder::CollectAllocateTexture(FCollectResourceContext& Context, FRDGPassHandle PassHandle, FRDGTextureRef Texture)
 {
-	IF_RDG_ENABLE_DEBUG(ConditionalDebugBreak(RDG_BREAKPOINT_PASS_COMPILE, BuilderName.GetTCHAR(), Pass->GetName()));
+	check(Texture->ReferenceCount > 0 || Texture->bExternal || IsImmediateMode());
 
-	for (auto& PassState : Pass->TextureStates)
+#if RDG_ENABLE_DEBUG
 	{
-		FRDGTextureRef Texture = PassState.Texture;
-		AddTransition(Pass->Handle, Texture, PassState.MergeState);
+		FRDGPass* Pass = Passes[PassHandle];
 
-		IF_RDG_ENABLE_TRACE(Trace.AddTexturePassDependency(Texture, Pass));
+		// Cannot begin a resource on an async compute pass.
+		check(Pass->Pipeline == ERHIPipeline::Graphics);
+
+		// Cannot begin a resource within a merged render pass region.
+		checkf(GetPrologueBarrierPassHandle(PassHandle) == PassHandle,
+			TEXT("Cannot begin a resource within a merged render pass. Pass (Handle: %d, Name: %s), Resource %s"), PassHandle.GetIndex(), Pass->GetName(), Texture->Name);
 	}
+#endif
 
-	for (auto& PassState : Pass->BufferStates)
+	if (Texture->FirstPass.IsNull())
 	{
-		FRDGBufferRef Buffer = PassState.Buffer;
-		AddTransition(Pass->Handle, Buffer, *PassState.MergeState);
-
-		IF_RDG_ENABLE_TRACE(Trace.AddBufferPassDependency(Buffer, Pass));
+		Texture->FirstPass = PassHandle;
 	}
-}
-
-void FRDGBuilder::CreatePassBarriers(TFunctionRef<void()> PreWork)
-{
-	const int32 NumBarriersPerTask = 32;
 
-	if (bParallelSetupEnabled && TransitionCreateQueue.Num() > NumBarriersPerTask)
+	if (Texture->bCollectForAllocate)
 	{
-		ParallelForWithPreWork(TEXT("FRDGBuilder::CreatePassBarriers"), TransitionCreateQueue.Num(), NumBarriersPerTask, [this](int32 Index)
-		{
-			TransitionCreateQueue[Index]->CreateTransition();
+		Texture->bCollectForAllocate = false;
+		check(!Texture->ResourceRHI);
 
-		}, PreWork);
-	}
-	else
-	{
-		PreWork();
+		const FCollectResourceOp AllocateOp = FCollectResourceOp::Allocate(PassHandle, Texture->Handle);
 
-		SCOPED_NAMED_EVENT_TEXT("FRDGBuilder::CreatePassBarriers", FColor::Magenta);
+		if (Texture->bTransient)
+		{
+			Context.TransientResources.Emplace(AllocateOp);
 
-		for (FRDGBarrierBatchBegin* BeginBatch : TransitionCreateQueue)
+	#if RDG_STATS
+			GRDGStatTransientTextureCount++;
+	#endif
+		}
+		else
 		{
-			BeginBatch->CreateTransition();
+			Context.PooledTextures.Emplace(AllocateOp);
 		}
 	}
-
-	TransitionCreateQueue.Reset();
 }
 
-
-///////////////////////////////////////////////////////////////////////////////////////////////////
-
-void FRDGBuilder::AddEpilogueTransition(FRDGTextureRef Texture)
+void FRDGBuilder::CollectDeallocateTexture(FCollectResourceContext& Context, FRDGPassHandle PassHandle, FRDGTexture* Texture, uint32 ReferenceCount)
 {
-	if (!Texture->HasRHI() || Texture->IsCulled() || !Texture->bLastOwner)
-	{
-		return;
-	}
-
-	check(IsImmediateMode() || Texture->bExtracted || Texture->ReferenceCount == FRDGViewableResource::DeallocatedReferenceCount);
+	check(!IsImmediateMode());
+	check(Texture->ReferenceCount != FRDGViewableResource::DeallocatedReferenceCount);
+	check(Texture->ReferenceCount >= ReferenceCount);
+	Texture->ReferenceCount -= ReferenceCount;
 
-	if (!EnumHasAnyFlags(Texture->Flags, ERDGTextureFlags::SkipTracking))
+	if (Texture->ReferenceCount == 0)
 	{
-		const FRDGPassHandle EpiloguePassHandle = GetEpiloguePassHandle();
-
-		FRDGSubresourceState SubresourceState;
-		SubresourceState.SetPass(ERHIPipeline::Graphics, EpiloguePassHandle);
+		check(!Texture->bCollectForAllocate);
+		const FCollectResourceOp DeallocateOp = FCollectResourceOp::Deallocate(PassHandle, Texture->Handle);
 
-		// Texture is using the RHI transient allocator. Transition it back to Discard in the final pass it is used.
-		if (Texture->bTransient && !Texture->TransientTexture->IsAcquired())
+		if (Texture->bTransient)
 		{
-			const TInterval<uint32> DiscardPasses = Texture->TransientTexture->GetDiscardPasses();
-			const FRDGPassHandle MinDiscardPassHandle(DiscardPasses.Min);
-			const FRDGPassHandle MaxDiscardPassHandle(FMath::Min<uint32>(DiscardPasses.Max, EpiloguePassHandle.GetIndex()));
-
-			AddAliasingTransition(MinDiscardPassHandle, MaxDiscardPassHandle, Texture, FRHITransientAliasingInfo::Discard(Texture->GetRHIUnchecked()));
-
-			SubresourceState.SetPass(ERHIPipeline::Graphics, MaxDiscardPassHandle);
-			Texture->EpilogueAccess = ERHIAccess::Discard;
+			Context.TransientResources.Emplace(DeallocateOp);
 		}
-
-		// Memoryless textures should be discarded at the end of the pass
-		if (EnumHasAnyFlags(Texture->Desc.Flags, TexCreate_Memoryless))
+		else
 		{
-			Texture->EpilogueAccess = ERHIAccess::Discard;
+			Context.PooledTextures.Emplace(DeallocateOp);
 		}
 
-		SubresourceState.Access = Texture->EpilogueAccess;
-
-		InitTextureSubresources(ScratchTextureState, Texture->Layout, &SubresourceState);
-		AddTransition(EpiloguePassHandle, Texture, ScratchTextureState);
-		ScratchTextureState.Reset();
-
-		EpilogueResourceAccesses.Emplace(Texture->GetRHI(), Texture->EpilogueAccess);
-	}
-
-	if (Texture->Allocation)
-	{
-		ActivePooledTextures.Emplace(MoveTemp(Texture->Allocation));
-	}
-	else if (!Texture->bTransient)
-	{
-		// Non-transient textures need to be 'resurrected' to hold the last reference in the chain.
-		// Transient textures don't have that restriction since there's no actual pooling happening.
-		ActivePooledTextures.Emplace(Texture->RenderTarget);
+		Texture->LastPass = PassHandle;
+		Texture->ReferenceCount = FRDGViewableResource::DeallocatedReferenceCount;
 	}
 }
 
-void FRDGBuilder::AddEpilogueTransition(FRDGBufferRef Buffer)
+void FRDGBuilder::CollectAllocateBuffer(FCollectResourceContext& Context, FRDGPassHandle PassHandle, FRDGBuffer* Buffer)
 {
-	if (!Buffer->HasRHI() || Buffer->IsCulled() || !Buffer->bLastOwner)
+	check(Buffer->ReferenceCount > 0 || IsImmediateMode());
+
+#if RDG_ENABLE_DEBUG
 	{
-		return;
+		const FRDGPass* Pass = Passes[PassHandle];
+
+		// Cannot begin a resource on an async compute pass.
+		check(Pass->Pipeline == ERHIPipeline::Graphics);
+
+		// Cannot begin a resource within a merged render pass region.
+		checkf(GetPrologueBarrierPassHandle(PassHandle) == PassHandle,
+			TEXT("Cannot begin a resource within a merged render pass. Pass (Handle: %d, Name: %s), Resource %s"), PassHandle.GetIndex(), Pass->GetName(), Buffer->Name);
 	}
+#endif
 
-	check(IsImmediateMode() || Buffer->bExtracted || Buffer->ReferenceCount == FRDGViewableResource::DeallocatedReferenceCount);
+	if (Buffer->FirstPass.IsNull())
+	{
+		Buffer->FirstPass = PassHandle;
+	}
 
-	if (!EnumHasAnyFlags(Buffer->Flags, ERDGBufferFlags::SkipTracking))
+	if (Buffer->bCollectForAllocate)
 	{
-		const FRDGPassHandle EpiloguePassHandle = GetEpiloguePassHandle();
+		Buffer->bCollectForAllocate = false;
+		check(!Buffer->ResourceRHI);
 
-		FRDGSubresourceState SubresourceState;
-		SubresourceState.SetPass(ERHIPipeline::Graphics, EpiloguePassHandle);
+		const FCollectResourceOp AllocateOp = FCollectResourceOp::Allocate(PassHandle, Buffer->Handle);
 
-		// Texture is using the RHI transient allocator. Transition it back to Discard in the final pass it is used.
 		if (Buffer->bTransient)
 		{
-			const TInterval<uint32> DiscardPasses = Buffer->TransientBuffer->GetDiscardPasses();
-			const FRDGPassHandle MinDiscardPassHandle(DiscardPasses.Min);
-			const FRDGPassHandle MaxDiscardPassHandle(FMath::Min<uint32>(DiscardPasses.Max, EpiloguePassHandle.GetIndex()));
-
-			AddAliasingTransition(MinDiscardPassHandle, MaxDiscardPassHandle, Buffer, FRHITransientAliasingInfo::Discard(Buffer->GetRHIUnchecked()));
-
-			SubresourceState.SetPass(ERHIPipeline::Graphics, MaxDiscardPassHandle);
-			Buffer->EpilogueAccess = ERHIAccess::Discard;
-		}
+			Context.TransientResources.Emplace(AllocateOp);
 
-		SubresourceState.Access = Buffer->EpilogueAccess;
+#if RDG_STATS
+			GRDGStatTransientBufferCount++;
+#endif
+		}
+		else
+		{
+			Context.PooledBuffers.Emplace(AllocateOp);
+		}
+	}
+}
+
+void FRDGBuilder::CollectDeallocateBuffer(FCollectResourceContext& Context, FRDGPassHandle PassHandle, FRDGBuffer* Buffer, uint32 ReferenceCount)
+{
+	check(!IsImmediateMode());
+	check(Buffer->ReferenceCount != FRDGViewableResource::DeallocatedReferenceCount);
+	check(Buffer->ReferenceCount >= ReferenceCount);
+	Buffer->ReferenceCount -= ReferenceCount;
+
+	if (Buffer->ReferenceCount == 0)
+	{
+		const FCollectResourceOp DeallocateOp = FCollectResourceOp::Deallocate(PassHandle, Buffer->Handle);
+
+		if (Buffer->bTransient)
+		{
+			Context.TransientResources.Emplace(DeallocateOp);
+		}
+		else
+		{
+			Context.PooledBuffers.Emplace(DeallocateOp);
+		}
+
+		Buffer->LastPass = PassHandle;
+		Buffer->ReferenceCount = FRDGViewableResource::DeallocatedReferenceCount;
+	}
+}
+
+///////////////////////////////////////////////////////////////////////////////////////////////////
+
+void FRDGBuilder::CompilePassBarriers()
+{
+	// Walk the culled graph and compile barriers for each subresource. Certain transitions are redundant; read-to-read, for example.
+	// We can avoid them by traversing and merging compatible states together. The merging states removes a transition, but the merging
+	// heuristic is conservative and choosing not to merge doesn't necessarily mean a transition is performed. They are two distinct steps.
+	// Merged states track the first and last pass used for all pipelines.
+
+	SCOPED_NAMED_EVENT(CompileBarriers, FColor::Emerald);
+	FRDGAllocatorScope AllocatorScope(Allocators.Transition);
+
+	for (FRDGPassHandle PassHandle = GetProloguePassHandle() + 1; PassHandle < GetEpiloguePassHandle(); ++PassHandle)
+	{
+		FRDGPass* Pass = Passes[PassHandle];
+
+		if (Pass->bCulled)
+		{
+			continue;
+		}
+
+		const ERHIPipeline PassPipeline = Pass->Pipeline;
+
+		const auto MergeSubresourceStates = [&](ERDGViewableResourceType ResourceType, FRDGSubresourceState*& PassMergeState, FRDGSubresourceState*& ResourceMergeState, FRDGSubresourceState* PassState)
+		{
+			if (!ResourceMergeState || !FRDGSubresourceState::IsMergeAllowed(ResourceType, *ResourceMergeState, *PassState))
+			{
+				// Use the new pass state as the merge state for future passes.
+				ResourceMergeState = PassState;
+			}
+			else
+			{
+				// Merge the pass state into the merged state.
+				ResourceMergeState->Access |= PassState->Access;
+
+				FRDGPassHandle& FirstPassHandle = ResourceMergeState->FirstPass[PassPipeline];
+
+				if (FirstPassHandle.IsNull())
+				{
+					FirstPassHandle = PassHandle;
+				}
+
+				ResourceMergeState->LastPass[PassPipeline] = PassHandle;
+			}
+
+			PassMergeState = ResourceMergeState;
+		};
+
+		for (auto& PassState : Pass->TextureStates)
+		{
+			FRDGTexture* Texture = PassState.Texture;
+
+		#if RDG_STATS
+			GRDGStatTextureReferenceCount += PassState.ReferenceCount;
+		#endif
+
+			for (int32 Index = 0; Index < PassState.State.Num(); ++Index)
+			{
+				if (!PassState.State[Index])
+				{
+					continue;
+				}
+
+				MergeSubresourceStates(ERDGViewableResourceType::Texture, PassState.MergeState[Index], Texture->MergeState[Index], PassState.State[Index]);
+			}
+		}
+
+		for (auto& PassState : Pass->BufferStates)
+		{
+			FRDGBuffer* Buffer = PassState.Buffer;
+
+		#if RDG_STATS
+			GRDGStatBufferReferenceCount += PassState.ReferenceCount;
+		#endif
+
+			MergeSubresourceStates(ERDGViewableResourceType::Buffer, PassState.MergeState, Buffer->MergeState, &PassState.State);
+		}
+	}
+}
+
+void FRDGBuilder::CollectPassBarriers()
+{
+	SCOPED_NAMED_EVENT_TEXT("FRDGBuilder::CollectBarriers", FColor::Magenta);
+	SCOPE_CYCLE_COUNTER(STAT_RDG_CollectBarriersTime);
+	CSV_SCOPED_TIMING_STAT_EXCLUSIVE_CONDITIONAL(RDG_CollectBarriers, GRDGVerboseCSVStats != 0);
+	FRDGAllocatorScope AllocatorScope(Allocators.Transition);
+
+	for (FRDGPassHandle PassHandle = GetProloguePassHandle() + 1; PassHandle < GetEpiloguePassHandle(); ++PassHandle)
+	{
+		CollectPassBarriers(PassHandle);
+	}
+}
+
+void FRDGBuilder::CollectPassBarriers(FRDGPassHandle PassHandle)
+{
+	IF_RDG_ENABLE_DEBUG(ConditionalDebugBreak(RDG_BREAKPOINT_PASS_COMPILE, BuilderName.GetTCHAR(), Passes[PassHandle]->GetName()));
+
+	FRDGPass* Pass = Passes[PassHandle];
+
+	if (Pass->bCulled || Pass->bEmptyParameters)
+	{
+		return;
+	}
+
+
+	for (auto& PassState : Pass->TextureStates)
+	{
+		FRDGTexture* Texture = PassState.Texture;
+
+		AddTextureTransition(PassState.Texture, Texture->State, PassState.MergeState, [Texture] (FRDGSubresourceState* StateAfter, int32 SubresourceIndex)
+		{
+			if (!Texture->FirstState[SubresourceIndex])
+			{
+				Texture->FirstState[SubresourceIndex] = StateAfter;
+				return IsImmediateMode();
+			}
+			return true;
+		});
+
+		IF_RDG_ENABLE_TRACE(Trace.AddTexturePassDependency(Texture, Pass));
+	}
+
+	for (auto& PassState : Pass->BufferStates)
+	{
+		FRDGBuffer* Buffer = PassState.Buffer;
+
+		AddBufferTransition(PassState.Buffer, Buffer->State, PassState.MergeState, [Buffer] (FRDGSubresourceState* StateAfter)
+		{
+			if (!Buffer->FirstState)
+			{
+				Buffer->FirstState = StateAfter;
+				Buffer->FirstState->bReservedCommit = Buffer->PendingCommitSize > 0;
+				return IsImmediateMode();
+			}
+			return true;
+		});
+
+		IF_RDG_ENABLE_TRACE(Trace.AddBufferPassDependency(Buffer, Pass));
+	}
+}
+
+void FRDGBuilder::CreatePassBarriers()
+{
+	struct FTaskContext
+	{
+		TArray<FRHITransitionInfo, FConcurrentLinearArrayAllocator> Transitions;
+	};
+
+	const auto CreateTransition = [this] (FTaskContext& Context, FRDGBarrierBatchBegin* BeginBatch)
+	{
+		Context.Transitions.Reset(BeginBatch->Transitions.Num());
+
+		for (FRDGTransitionInfo InfoRDG : BeginBatch->Transitions)
+		{
+			FRHITransitionInfo& InfoRHI = Context.Transitions.Emplace_GetRef();
+			InfoRHI.AccessBefore = InfoRDG.AccessBefore;
+			InfoRHI.AccessAfter  = InfoRDG.AccessAfter;
+			InfoRHI.ArraySlice   = InfoRDG.ArraySlice;
+			InfoRHI.MipIndex     = InfoRDG.MipIndex;
+			InfoRHI.PlaneSlice   = InfoRDG.PlaneSlice;
+			InfoRHI.Flags        = InfoRDG.Flags;
+
+			if (InfoRDG.Type == ERDGViewableResourceType::Texture)
+			{
+				InfoRHI.Resource = Textures[FRDGTextureHandle(InfoRDG.Handle)]->ResourceRHI;
+				InfoRHI.Type = FRHITransitionInfo::EType::Texture;
+			}
+			else
+			{
+				FRDGBuffer* Buffer = Buffers[FRDGBufferHandle(InfoRDG.Handle)];
+
+				InfoRHI.Resource = Buffer->ResourceRHI;
+				InfoRHI.Type = FRHITransitionInfo::EType::Buffer;
+
+				if (InfoRDG.bReservedCommit)
+				{
+					InfoRHI.CommitInfo.Emplace(Buffer->PendingCommitSize);
+				}
+			}
+		}
 
-		AddTransition(Buffer->LastPass, Buffer, SubresourceState);
+		BeginBatch->CreateTransition(Context.Transitions);
+	};
 
-		EpilogueResourceAccesses.Emplace(Buffer->GetRHI(), Buffer->EpilogueAccess);
+	TArray<FTaskContext, TInlineAllocator<1, FRDGArrayAllocator>> TaskContexts;
+	ParallelForWithTaskContext(TEXT("FRDGBuilder::CreatePassBarriers"), TaskContexts, TransitionCreateQueue.Num(), 1, [&](FTaskContext& TaskContext, int32 Index)
+	{
+		CreateTransition(TaskContext, TransitionCreateQueue[Index]);
+
+	}, ParallelSetup.bEnabled ? EParallelForFlags::None : EParallelForFlags::ForceSingleThread);
+
+	TransitionCreateQueue.Reset();
+}
+
+///////////////////////////////////////////////////////////////////////////////////////////////////
+
+void FRDGBuilder::FinalizeResources()
+{
+	SCOPED_NAMED_EVENT_TEXT("FRDGBuilder::FinalizeResources", FColor::Magenta);
+	FRDGAllocatorScope AllocatorScope(Allocators.Transition);
+
+	{
+		SCOPED_NAMED_EVENT_TEXT("Textures", FColor::Magenta);
+		Textures.Enumerate([&](FRDGTextureRef Texture)
+		{
+			if (Texture->FirstPass.IsValid())
+			{
+				if (!IsImmediateMode())
+				{
+					AddFirstTextureTransition(Texture);
+				}
+
+				if (!Texture->bSkipLastTransition)
+				{
+					AddLastTextureTransition(Texture);
+				}
+			}
+
+			if (Texture->Allocation)
+			{
+				ActivePooledTextures.Emplace(MoveTemp(Texture->Allocation));
+			}
+		});
 	}
 
-	if (Buffer->Allocation)
 	{
-		ActivePooledBuffers.Emplace(MoveTemp(Buffer->Allocation));
+		SCOPED_NAMED_EVENT_TEXT("Buffers", FColor::Magenta);
+		Buffers.Enumerate([&](FRDGBufferRef Buffer)
+		{
+			if (Buffer->FirstPass.IsValid())
+			{
+				if (!IsImmediateMode())
+				{
+					AddFirstBufferTransition(Buffer);
+				}
+
+				if (!Buffer->bSkipLastTransition)
+				{
+					AddLastBufferTransition(Buffer);
+				}
+			}
+
+			if (Buffer->Allocation)
+			{
+				ActivePooledBuffers.Emplace(MoveTemp(Buffer->Allocation));
+			}
+		});
+	}
+
+	CreatePassBarriers();
+}
+
+void FRDGBuilder::AddFirstTextureTransition(FRDGTexture* Texture)
+{
+	check(!IsImmediateMode());
+	check(Texture->HasRHI());
+
+	FRDGTextureSubresourceState* StateBefore = &ScratchTextureState;
+	FRDGSubresourceState& SubresourceStateBefore = *AllocSubresource(FRDGSubresourceState(ERHIPipeline::Graphics, GetProloguePassHandle()));
+
+	if (Texture->PreviousOwner.IsValid())
+	{
+		// Previous state is the last used state of RDG texture that previously aliased the underlying pooled texture.
+		StateBefore = &Textures[Texture->PreviousOwner]->State;
+
+		for (int32 Index = 0; Index < Texture->FirstState.Num(); ++Index)
+		{
+			// If the new owner doesn't touch the subresource but the previous owner did, pull the previous owner subresource in so that the last transition is respected.
+			if (!Texture->FirstState[Index])
+			{
+				Texture->State[Index] = (*StateBefore)[Index];
+			}
+			// If the previous owner didn't touch the subresource but the new owner does, assign the prologue subresource state so the first transition is respected.
+			else if (!(*StateBefore)[Index])
+			{
+				(*StateBefore)[Index] = &SubresourceStateBefore;
+			}
+		}
 	}
-	else if (!Buffer->bTransient)
+	else
 	{
-		// Non-transient buffers need to be 'resurrected' to hold the last reference in the chain.
-		// Transient buffers don't have that restriction since there's no actual pooling happening.
-		ActivePooledBuffers.Emplace(Buffer->PooledBuffer);
+		if (Texture->MinAcquirePass.IsValid())
+		{
+			AddAliasingTransition(Texture->MinAcquirePass, Texture->FirstPass, Texture, FRHITransientAliasingInfo::Acquire(Texture->GetRHI(), Texture->AliasingOverlaps));
+
+			SubresourceStateBefore.SetPass(ERHIPipeline::Graphics, Texture->MinAcquirePass);
+			SubresourceStateBefore.Access = ERHIAccess::Discard;
+		}
+		else if (!Texture->bSplitFirstTransition)
+		{
+			SubresourceStateBefore.SetPass(ERHIPipeline::Graphics, Texture->FirstPass);
+		}
+
+		InitTextureSubresources(*StateBefore, Texture->Layout, &SubresourceStateBefore);
 	}
+
+	AddTextureTransition(Texture, *StateBefore, Texture->FirstState);
+
+	ScratchTextureState.Reset();
 }
 
-void FRDGBuilder::AddTransition(FRDGPassHandle PassHandle, FRDGTextureRef Texture, FRDGTextureSubresourceStateIndirect& StateAfter)
+void FRDGBuilder::AddLastTextureTransition(FRDGTexture* Texture)
+{
+	check(IsImmediateMode() || Texture->bExtracted || Texture->ReferenceCount == FRDGViewableResource::DeallocatedReferenceCount);
+	check(Texture->HasRHI());
+
+	const FRDGPassHandle EpiloguePassHandle = GetEpiloguePassHandle();
+
+	FRDGSubresourceState& SubresourceStateAfter = *AllocSubresource();
+	SubresourceStateAfter.SetPass(ERHIPipeline::Graphics, EpiloguePassHandle);
+
+	// Texture is using the RHI transient allocator. Transition it back to Discard in the final pass it is used.
+	if (Texture->MinDiscardPass.IsValid())
+	{
+		AddAliasingTransition(Texture->MinDiscardPass, Texture->MaxDiscardPass, Texture, FRHITransientAliasingInfo::Discard(Texture->GetRHIUnchecked()));
+
+		SubresourceStateAfter.SetPass(ERHIPipeline::Graphics, Texture->MaxDiscardPass);
+		SubresourceStateAfter.Access = ERHIAccess::Discard;
+	}
+	else
+	{
+		FRDGSubresourceState& SubresourceStateBefore = *AllocSubresource();
+		SubresourceStateBefore.SetPass(ERHIPipeline::Graphics, Texture->FirstPass);
+
+		SubresourceStateAfter.Access = Texture->EpilogueAccess;
+
+		// Transition any unused (null) sub-resources to the epilogue state since we are assigning a monolithic state across all subresources.
+		for (FRDGSubresourceState*& State : Texture->State)
+		{
+			if (!State)
+			{
+				State = &SubresourceStateBefore;
+			}
+		}
+
+		EpilogueResourceAccesses.Emplace(Texture->GetRHI(), SubresourceStateAfter.Access);
+	}
+
+	InitTextureSubresources(ScratchTextureState, Texture->Layout, &SubresourceStateAfter);
+	AddTextureTransition(Texture, Texture->State, ScratchTextureState);
+	ScratchTextureState.Reset();
+}
+
+void FRDGBuilder::AddFirstBufferTransition(FRDGBuffer* Buffer)
+{
+	check(!IsImmediateMode());
+	check(Buffer->HasRHI());
+
+	FRDGSubresourceState* StateBefore = nullptr;
+
+	if (Buffer->PreviousOwner.IsValid())
+	{
+		// Previous state is the last used state of RDG buffer that previously aliased the underlying pooled buffer.
+		StateBefore = Buffers[Buffer->PreviousOwner]->State;
+	}
+
+	if (!StateBefore)
+	{
+		StateBefore = AllocSubresource();
+
+		if (Buffer->MinAcquirePass.IsValid())
+		{
+			AddAliasingTransition(Buffer->MinAcquirePass, Buffer->FirstPass, Buffer, FRHITransientAliasingInfo::Acquire(Buffer->GetRHI(), Buffer->AliasingOverlaps));
+
+			StateBefore->SetPass(ERHIPipeline::Graphics, Buffer->MinAcquirePass);
+			StateBefore->Access = ERHIAccess::Discard;
+		}
+		else if (!Buffer->bSplitFirstTransition)
+		{
+			StateBefore->SetPass(ERHIPipeline::Graphics, Buffer->FirstPass);
+		}
+		else
+		{
+			StateBefore->SetPass(ERHIPipeline::Graphics, GetProloguePassHandle());
+		}
+	}
+
+	AddBufferTransition(Buffer, StateBefore, Buffer->FirstState);
+}
+
+void FRDGBuilder::AddLastBufferTransition(FRDGBuffer* Buffer)
+{
+	check(IsImmediateMode() || Buffer->bExtracted || Buffer->ReferenceCount == FRDGViewableResource::DeallocatedReferenceCount);
+	check(Buffer->HasRHI());
+
+	const FRDGPassHandle EpiloguePassHandle = GetEpiloguePassHandle();
+
+	FRDGSubresourceState* StateAfter = AllocSubresource();
+
+	// Texture is using the RHI transient allocator. Transition it back to Discard in the final pass it is used.
+	if (Buffer->MinDiscardPass.IsValid())
+	{
+		AddAliasingTransition(Buffer->MinDiscardPass, Buffer->MaxDiscardPass, Buffer, FRHITransientAliasingInfo::Discard(Buffer->GetRHIUnchecked()));
+
+		StateAfter->SetPass(ERHIPipeline::Graphics, Buffer->MaxDiscardPass);
+		StateAfter->Access = ERHIAccess::Discard;
+	}
+	else
+	{
+		StateAfter->SetPass(ERHIPipeline::Graphics, EpiloguePassHandle);
+		StateAfter->Access = Buffer->EpilogueAccess;
+
+		EpilogueResourceAccesses.Emplace(Buffer->GetRHI(), StateAfter->Access);
+	}
+
+	AddBufferTransition(Buffer, Buffer->State, StateAfter);
+}
+
+template <typename FilterSubresourceLambdaType>
+void FRDGBuilder::AddTextureTransition(FRDGTexture* Texture, FRDGTextureSubresourceState& StateBefore, FRDGTextureSubresourceState& StateAfter, FilterSubresourceLambdaType&& FilterSubresourceLambda)
 {
 	const FRDGTextureSubresourceLayout Layout = Texture->Layout;
-	FRDGTextureSubresourceState& StateBefore = Texture->GetState();
-	const uint32 SubresourceCount = StateBefore.Num();
+	const uint32 SubresourceCount = Texture->SubresourceCount;
 
-	checkf(StateBefore.Num() == Layout.GetSubresourceCount() && StateBefore.Num() == StateAfter.Num(),
-		TEXT("Before state array (%d) does not match after state array (%d) for resource %s on pass %s."),
-		StateBefore.Num(), StateAfter.Num(), Texture->Name, Passes[PassHandle]->GetName());
+	check(SubresourceCount == Layout.GetSubresourceCount() && StateBefore.Num() == StateAfter.Num());
 
 	if (!GRHISupportsSeparateDepthStencilCopyAccess && Texture->Desc.Format == PF_DepthStencil)
 	{
@@ -3215,40 +3711,37 @@ void FRDGBuilder::AddTransition(FRDGPassHandle PassHandle, FRDGTextureRef Textur
 				continue;
 			}
 
-			FRDGSubresourceState& DepthStateBefore   = StateBefore[DepthIndex];
-			FRDGSubresourceState& StencilStateBefore = StateBefore[StencilIndex];
+			FRDGSubresourceState*& DepthStateBefore   = StateBefore[DepthIndex];
+			FRDGSubresourceState*& StencilStateBefore = StateBefore[StencilIndex];
 
 			// Case 1: transitioning into a fused copy state.
 			if (DepthStateAfter && EnumHasAnyFlags(DepthStateAfter->Access, ERHIAccess::CopySrc | ERHIAccess::CopyDest))
 			{
 				check(StencilStateAfter && StencilStateAfter->Access == DepthStateAfter->Access);
 
-				const FRDGPassHandle MaxPassHandle = FRDGPassHandle::Max(DepthStateBefore.LastPass[GraphicsPipe], StencilStateBefore.LastPass[GraphicsPipe]);
-				DepthStateBefore.LastPass[GraphicsPipe]   = MaxPassHandle;
-				StencilStateBefore.LastPass[GraphicsPipe] = MaxPassHandle;
+				const FRDGPassHandle MaxPassHandle = FRDGPassHandle::Max(DepthStateBefore->LastPass[GraphicsPipe], StencilStateBefore->LastPass[GraphicsPipe]);
+
+				DepthStateBefore = AllocSubresource(*DepthStateBefore);
+				DepthStateAfter  = AllocSubresource(*DepthStateAfter);
+
+				DepthStateBefore->LastPass[GraphicsPipe]   = MaxPassHandle;
+				StencilStateBefore->LastPass[GraphicsPipe] = MaxPassHandle;
 			}
 			// Case 2: transitioning out of a fused copy state.
-			else if (EnumHasAnyFlags(DepthStateBefore.Access, ERHIAccess::CopySrc | ERHIAccess::CopyDest))
+			else if (EnumHasAnyFlags(DepthStateBefore->Access, ERHIAccess::CopySrc | ERHIAccess::CopyDest))
 			{
-				check(StencilStateBefore.Access        == DepthStateBefore.Access);
-				check(StencilStateBefore.GetLastPass() == DepthStateBefore.GetLastPass());
+				check(StencilStateBefore->Access        == DepthStateBefore->Access);
+				check(StencilStateBefore->GetLastPass() == DepthStateBefore->GetLastPass());
 
 				// Case 2a: depth unknown, so transition to match stencil.
 				if (!DepthStateAfter)
 				{
 					DepthStateAfter = AllocSubresource(*StencilStateAfter);
-					DepthStateAfter->SetPass(GraphicsPipe, PassHandle);
 				}
 				// Case 2b: stencil unknown, so transition to match depth.
 				else if (!StencilStateAfter)
 				{
 					StencilStateAfter = AllocSubresource(*DepthStateAfter);
-					StencilStateAfter->SetPass(GraphicsPipe, PassHandle);
-				}
-				// Two valid after states should be transitioning on this pass.
-				else
-				{
-					check(StencilStateAfter->GetFirstPass() == PassHandle && DepthStateAfter->GetFirstPass() == PassHandle);
 				}
 			}
 		}
@@ -3256,61 +3749,71 @@ void FRDGBuilder::AddTransition(FRDGPassHandle PassHandle, FRDGTextureRef Textur
 
 	for (uint32 SubresourceIndex = 0; SubresourceIndex < SubresourceCount; ++SubresourceIndex)
 	{
-		if (const FRDGSubresourceState* SubresourceStateAfter = StateAfter[SubresourceIndex])
+		FRDGSubresourceState*& SubresourceStateBefore = StateBefore[SubresourceIndex];
+		FRDGSubresourceState* SubresourceStateAfter = StateAfter[SubresourceIndex];
+
+		if (!SubresourceStateAfter)
 		{
-			check(SubresourceStateAfter->Access != ERHIAccess::Unknown);
+			continue;
+		}
 
-			FRDGSubresourceState& SubresourceStateBefore = StateBefore[SubresourceIndex];
+		if (FilterSubresourceLambda(SubresourceStateAfter, SubresourceIndex))
+		{
+			check(SubresourceStateAfter->Access != ERHIAccess::Unknown);
 
-			if (FRDGSubresourceState::IsTransitionRequired(SubresourceStateBefore, *SubresourceStateAfter))
+			if (SubresourceStateBefore && FRDGSubresourceState::IsTransitionRequired(*SubresourceStateBefore, *SubresourceStateAfter))
 			{
 				const FRDGTextureSubresource Subresource = Layout.GetSubresource(SubresourceIndex);
 
-				FRHITransitionInfo Info;
-				Info.Texture      = Texture->GetRHIUnchecked();
-				Info.Type         = FRHITransitionInfo::EType::Texture;
-				Info.Flags        = SubresourceStateAfter->Flags;
-				Info.AccessBefore = SubresourceStateBefore.Access;
-				Info.AccessAfter  = SubresourceStateAfter->Access;
-				Info.MipIndex     = Subresource.MipIndex;
-				Info.ArraySlice   = Subresource.ArraySlice;
-				Info.PlaneSlice   = Subresource.PlaneSlice;
+				FRDGTransitionInfo Info;
+				Info.AccessBefore = SubresourceStateBefore->Access;
+				Info.AccessAfter = SubresourceStateAfter->Access;
+				Info.Handle = Texture->Handle.GetIndex();
+				Info.Type = ERDGViewableResourceType::Texture;
+				Info.Flags = SubresourceStateAfter->Flags;
+				Info.ArraySlice = Subresource.ArraySlice;
+				Info.MipIndex = Subresource.MipIndex;
+				Info.PlaneSlice = Subresource.PlaneSlice;
+				Info.bReservedCommit = 0;
 
 				if (Info.AccessBefore == ERHIAccess::Discard)
 				{
 					Info.Flags |= EResourceTransitionFlags::Discard;
 				}
 
-				AddTransition(Texture, SubresourceStateBefore, *SubresourceStateAfter, Info);
+				AddTransition(Texture, *SubresourceStateBefore, *SubresourceStateAfter, Info);
 			}
-
-			SubresourceStateBefore = *SubresourceStateAfter;
 		}
+
+		SubresourceStateBefore = SubresourceStateAfter;
 	}
 }
 
-void FRDGBuilder::AddTransition(FRDGPassHandle PassHandle, FRDGBufferRef Buffer, FRDGSubresourceState StateAfter)
+template <typename FilterSubresourceLambdaType>
+void FRDGBuilder::AddBufferTransition(FRDGBufferRef Buffer, FRDGSubresourceState*& StateBefore, FRDGSubresourceState* StateAfter, FilterSubresourceLambdaType&& FilterSubresourceLambda)
 {
-	check(StateAfter.Access != ERHIAccess::Unknown);
-
-	FRDGSubresourceState& StateBefore = Buffer->GetState();
+	check(StateAfter);
+	check(StateAfter->Access != ERHIAccess::Unknown);
 
-	if (FRDGSubresourceState::IsTransitionRequired(StateBefore, StateAfter))
+	if (FilterSubresourceLambda(StateAfter))
 	{
-		FRHITransitionInfo Info;
-		Info.Resource = Buffer->GetRHIUnchecked();
-		Info.Type = FRHITransitionInfo::EType::Buffer;
-		Info.Flags = StateAfter.Flags;
-		Info.AccessBefore = StateBefore.Access;
-		Info.AccessAfter = StateAfter.Access;
+		check(StateBefore);
 
-		if (Buffer->PendingCommitSize > 0)
+		if (FRDGSubresourceState::IsTransitionRequired(*StateBefore, *StateAfter))
 		{
-			Info.CommitInfo.Emplace(Buffer->PendingCommitSize);
-			Buffer->PendingCommitSize = 0;
-		}
+			FRDGTransitionInfo Info;
+			Info.AccessBefore = StateBefore->Access;
+			Info.AccessAfter = StateAfter->Access;
+			Info.Handle = Buffer->Handle.GetIndex();
+			Info.Type = ERDGViewableResourceType::Buffer;
+			Info.Flags = StateAfter->Flags;
+			Info.ArraySlice = 0;
+			Info.MipIndex = 0;
+			Info.PlaneSlice = 0;
+			Info.bReservedCommit = StateAfter->bReservedCommit;
 
-		AddTransition(Buffer, StateBefore, StateAfter, Info);
+			AddTransition(Buffer, *StateBefore, *StateAfter, Info);
+		}
 	}
 
 	StateBefore = StateAfter;
@@ -3320,7 +3823,7 @@ void FRDGBuilder::AddTransition(
 	FRDGViewableResource* Resource,
 	FRDGSubresourceState StateBefore,
 	FRDGSubresourceState StateAfter,
-	const FRHITransitionInfo& TransitionInfo)
+	FRDGTransitionInfo TransitionInfo)
 {
 	const ERHIPipeline Graphics = ERHIPipeline::Graphics;
 	const ERHIPipeline AsyncCompute = ERHIPipeline::AsyncCompute;
@@ -3362,7 +3865,7 @@ void FRDGBuilder::AddTransition(
 		if (BeginPassHandle < FirstEndPassHandle)
 		{
 			BeginPass = GetEpilogueBarrierPass(BeginPassHandle);
-			BarriersToBegin = &BeginPass->GetEpilogueBarriersToBeginFor(Allocator, TransitionCreateQueue, PipelinesAfter);
+			BarriersToBegin = &BeginPass->GetEpilogueBarriersToBeginFor(Allocators.Transition, TransitionCreateQueue, PipelinesAfter);
 		}
 		// This is an immediate prologue transition in the same pass. Issue the begin in the prologue.
 		else
@@ -3372,7 +3875,7 @@ void FRDGBuilder::AddTransition(
 				Resource->Name, *GetRHIPipelineName(PipelinesAfter));
 
 			BeginPass = GetPrologueBarrierPass(BeginPassHandle);
-			BarriersToBegin = &BeginPass->GetPrologueBarriersToBegin(Allocator, TransitionCreateQueue);
+			BarriersToBegin = &BeginPass->GetPrologueBarriersToBegin(Allocators.Transition, TransitionCreateQueue);
 		}
 
 		BarriersToBegin->AddTransition(Resource, TransitionInfo);
@@ -3427,7 +3930,7 @@ void FRDGBuilder::AddTransition(
 			BarrierBatchPasses[Graphics]     = Passes[Id.Passes[Graphics]];
 			BarrierBatchPasses[AsyncCompute] = Passes[Id.Passes[AsyncCompute]];
 
-			BarriersToBegin = Allocator.AllocNoDestruct<FRDGBarrierBatchBegin>(PipelinesBefore, PipelinesAfter, GetEpilogueBarriersToBeginDebugName(PipelinesAfter), BarrierBatchPasses);
+			BarriersToBegin = Allocators.Transition.AllocNoDestruct<FRDGBarrierBatchBegin>(PipelinesBefore, PipelinesAfter, GetEpilogueBarriersToBeginDebugName(PipelinesAfter), BarrierBatchPasses);
 			TransitionCreateQueue.Emplace(BarriersToBegin);
 
 			for (FRDGPass* Pass : BarrierBatchPasses)
@@ -3461,9 +3964,8 @@ void FRDGBuilder::AddAliasingTransition(FRDGPassHandle BeginPassHandle, FRDGPass
 		EndPass = BeginPass;
 
 		check(GetPrologueBarrierPassHandle(BeginPassHandle) == BeginPassHandle);
-		check(BeginPass->GetPipeline() == ERHIPipeline::Graphics);
 
-		BarriersToBegin = &BeginPass->GetPrologueBarriersToBegin(Allocator, TransitionCreateQueue);
+		BarriersToBegin = &BeginPass->GetPrologueBarriersToBegin(Allocators.Transition, TransitionCreateQueue);
 	}
 	else
 	{
@@ -3471,19 +3973,30 @@ void FRDGBuilder::AddAliasingTransition(FRDGPassHandle BeginPassHandle, FRDGPass
 		EndPass = Passes[EndPassHandle];
 
 		check(GetPrologueBarrierPassHandle(EndPassHandle) == EndPassHandle);
-		check(BeginPass->GetPipeline() == ERHIPipeline::Graphics);
-		check(EndPass->GetPipeline() == ERHIPipeline::Graphics);
 
-		BarriersToBegin = &BeginPass->GetEpilogueBarriersToBeginForGraphics(Allocator, TransitionCreateQueue);
+		BarriersToBegin = &BeginPass->GetEpilogueBarriersToBeginFor(Allocators.Transition, TransitionCreateQueue, EndPass->GetPipeline());
 	}
 
 	BarriersToBegin->AddAlias(Resource, Info);
-	EndPass->GetPrologueBarriersToEnd(Allocator).AddDependency(BarriersToBegin);
+	EndPass->GetPrologueBarriersToEnd(Allocators.Transition).AddDependency(BarriersToBegin);
 }
 
 ///////////////////////////////////////////////////////////////////////////////////////////////////
 
-void FRDGBuilder::SetRHI(FRDGTexture* Texture, IPooledRenderTarget* RenderTarget, FRDGPassHandle PassHandle)
+TRefCountPtr<IPooledRenderTarget> FRDGBuilder::AllocatePooledRenderTargetRHI(FRHICommandListBase& InRHICmdList, FRDGTextureRef Texture)
+{
+	return GRenderTargetPool.FindFreeElement(InRHICmdList, Texture->Desc, Texture->Name);
+}
+
+TRefCountPtr<FRDGPooledBuffer> FRDGBuilder::AllocatePooledBufferRHI(FRHICommandListBase& InRHICmdList, FRDGBufferRef Buffer)
+{
+	Buffer->FinalizeDesc();
+	return GRenderGraphResourcePool.FindFreeBuffer(InRHICmdList, Buffer->Desc, Buffer->Name);
+}
+
+///////////////////////////////////////////////////////////////////////////////////////////////////
+
+void FRDGBuilder::SetPooledRenderTargetRHI(FRDGTexture* Texture, IPooledRenderTarget* RenderTarget)
 {
 	Texture->RenderTarget = RenderTarget;
 
@@ -3492,256 +4005,118 @@ void FRDGBuilder::SetRHI(FRDGTexture* Texture, IPooledRenderTarget* RenderTarget
 		FRDGTransientRenderTarget* TransientRenderTarget = static_cast<FRDGTransientRenderTarget*>(RenderTarget);
 		Texture->Allocation = TRefCountPtr<FRDGTransientRenderTarget>(TransientRenderTarget);
 
-		SetRHI(Texture, TransientTexture, PassHandle);
+		SetTransientTextureRHI(Texture, TransientTexture);
 	}
 	else
 	{
 		FPooledRenderTarget* PooledRenderTarget = static_cast<FPooledRenderTarget*>(RenderTarget);
 		Texture->Allocation = TRefCountPtr<FPooledRenderTarget>(PooledRenderTarget);
 
-		SetRHI(Texture, &PooledRenderTarget->PooledTexture, PassHandle);
+		SetPooledTextureRHI(Texture, &PooledRenderTarget->PooledTexture);
 	}
 }
 
-void FRDGBuilder::SetRHI(FRDGTexture* Texture, FRDGPooledTexture* PooledTexture, FRDGPassHandle PassHandle)
+void FRDGBuilder::SetPooledTextureRHI(FRDGTexture* Texture, FRDGPooledTexture* PooledTexture)
 {
+	check(!Texture->ResourceRHI);
+
 	FRHITexture* TextureRHI = PooledTexture->GetRHI();
 
 	Texture->ResourceRHI = TextureRHI;
 	Texture->PooledTexture = PooledTexture;
 	Texture->ViewCache = &PooledTexture->ViewCache;
-	Texture->FirstPass = PassHandle;
 
-	FRDGTexture*& Owner = PooledTextureOwnershipMap.FindOrAdd(PooledTexture);
+	FRDGTexture*& Owner = *PooledTextureOwnershipMap.FindOrAdd(PooledTexture, nullptr);
 
 	// Link the previous alias to this one.
 	if (Owner)
 	{
+		Texture->PreviousOwner = Owner->Handle;
 		Owner->NextOwner = Texture->Handle;
-		Owner->bLastOwner = false;
-
-		// Chain the state allocation between all RDG textures which share this pooled texture.
-		Texture->State = Owner->State;
-	}
-	else
-	{
-		Texture->State = Allocator.AllocNoDestruct<FRDGTextureSubresourceState>();
-
-		FRDGSubresourceState State;
-		State.SetPass(ERHIPipeline::Graphics, GetProloguePassHandle());
-		InitTextureSubresources(*Texture->State, Texture->Layout, State);
+		Owner->bSkipLastTransition = true;
 	}
 
 	Owner = Texture;
 }
 
-void FRDGBuilder::SetRHI(FRDGTexture* Texture, FRHITransientTexture* TransientTexture, FRDGPassHandle PassHandle)
+void FRDGBuilder::SetTransientTextureRHI(FRDGTexture* Texture, FRHITransientTexture* TransientTexture)
 {
+	check(!Texture->ResourceRHI);
+
 	Texture->ResourceRHI = TransientTexture->GetRHI();
 	Texture->TransientTexture = TransientTexture;
 	Texture->ViewCache = &TransientTexture->ViewCache;
-	Texture->FirstPass = PassHandle;
 	Texture->bTransient = true;
-	Texture->State = Allocator.AllocNoDestruct<FRDGTextureSubresourceState>();
+	Texture->AliasingOverlaps = TransientTexture->GetAliasingOverlaps();
 }
 
-void FRDGBuilder::SetRHI(FRDGBuffer* Buffer, FRDGPooledBuffer* PooledBuffer, FRDGPassHandle PassHandle)
+void FRDGBuilder::SetPooledBufferRHI(FRDGBuffer* Buffer, FRDGPooledBuffer* PooledBuffer)
 {
+	check(!Buffer->ResourceRHI);
+
 	FRHIBuffer* BufferRHI = PooledBuffer->GetRHI();
 
 	Buffer->ResourceRHI = BufferRHI;
 	Buffer->PooledBuffer = PooledBuffer;
 	Buffer->ViewCache = &PooledBuffer->ViewCache;
 	Buffer->Allocation = PooledBuffer;
-	Buffer->FirstPass = PassHandle;
 
-	FRDGBuffer*& Owner = PooledBufferOwnershipMap.FindOrAdd(PooledBuffer);
+#if RHI_ENABLE_RESOURCE_INFO
+	Buffer->ResourceRHI->SetOwnerName(Buffer->OwnerName);
+#endif
+
+	FRDGBuffer*& Owner = *PooledBufferOwnershipMap.FindOrAdd(PooledBuffer, nullptr);
 
 	// Link the previous owner to this one.
 	if (Owner)
 	{
+		Buffer->PreviousOwner = Owner->Handle;
 		Owner->NextOwner = Buffer->Handle;
-		Owner->bLastOwner = false;
-
-		// Chain the state allocation between all RDG buffers which share this pooled buffer.
-		Buffer->State = Owner->State;
-	}
-	else
-	{
-		FRDGSubresourceState State;
-		State.SetPass(ERHIPipeline::Graphics, GetProloguePassHandle());
-		Buffer->State = Allocator.AllocNoDestruct<FRDGSubresourceState>(State);
+		Owner->bSkipLastTransition = true;
 	}
 
 	Owner = Buffer;
 }
 
-void FRDGBuilder::SetRHI(FRDGBuffer* Buffer, FRHITransientBuffer* TransientBuffer, FRDGPassHandle PassHandle)
+void FRDGBuilder::SetTransientBufferRHI(FRDGBuffer* Buffer, FRHITransientBuffer* TransientBuffer)
 {
-	check(!Buffer->ResourceRHI);
+	check(!Buffer->ResourceRHI && Buffer->bTransient);
 
 	Buffer->ResourceRHI = TransientBuffer->GetRHI();
 	Buffer->TransientBuffer = TransientBuffer;
 	Buffer->ViewCache = &TransientBuffer->ViewCache;
-	Buffer->FirstPass = PassHandle;
-	Buffer->bTransient = true;
-	Buffer->State = Allocator.AllocNoDestruct<FRDGSubresourceState>();
-}
-
-///////////////////////////////////////////////////////////////////////////////////////////////////
-
-void FRDGBuilder::BeginResourceRHI(FRDGPassHandle PassHandle, FRDGTextureRef Texture)
-{
-	check(Texture);
-
-	if (Texture->HasRHI())
-	{
-		return;
-	}
-
-	check(Texture->ReferenceCount > 0 || Texture->bExternal || IsImmediateMode());
-
-#if RDG_ENABLE_DEBUG
-	{
-		FRDGPass* Pass = Passes[PassHandle];
-
-		// Cannot begin a resource on an async compute pass.
-		check(Pass->Pipeline == ERHIPipeline::Graphics);
+	Buffer->AliasingOverlaps = TransientBuffer->GetAliasingOverlaps();
 
-		// Cannot begin a resource within a merged render pass region.
-		checkf(GetPrologueBarrierPassHandle(PassHandle) == PassHandle,
-			TEXT("Cannot begin a resource within a merged render pass. Pass (Handle: %d, Name: %s), Resource %s"), PassHandle.GetIndex(), Pass->GetName(), Texture->Name);
-	}
+#if RHI_ENABLE_RESOURCE_INFO
+	Buffer->ResourceRHI->SetOwnerName(Buffer->OwnerName);
 #endif
-
-	if (TransientResourceAllocator && bSupportsTransientTextures && IsTransient(Texture))
-	{
-		if (FRHITransientTexture* TransientTexture = TransientResourceAllocator->CreateTexture(Texture->Desc, Texture->Name, PassHandle.GetIndex()))
-		{
-			if (Texture->bExternal || Texture->bExtracted)
-			{
-				SetRHI(Texture, GRDGTransientResourceAllocator.AllocateRenderTarget(TransientTexture), PassHandle);
-			}
-			else
-			{
-				SetRHI(Texture, TransientTexture, PassHandle);
-			}
-
-			const FRDGPassHandle MinAcquirePassHandle(TransientTexture->GetAcquirePasses().Min);
-
-			AddAliasingTransition(MinAcquirePassHandle, PassHandle, Texture, FRHITransientAliasingInfo::Acquire(TransientTexture->GetRHI(), TransientTexture->GetAliasingOverlaps()));
-
-			FRDGSubresourceState InitialState;
-			InitialState.SetPass(ERHIPipeline::Graphics, MinAcquirePassHandle);
-			InitialState.Access = ERHIAccess::Discard;
-			InitTextureSubresources(*Texture->State, Texture->Layout, InitialState);
-
-		#if RDG_STATS
-			GRDGStatTransientTextureCount++;
-		#endif
-		}
-	}
-
-	if (!Texture->ResourceRHI)
-	{
-		SetRHI(Texture, GRenderTargetPool.FindFreeElement(RHICmdList, Texture->Desc, Texture->Name), PassHandle);
-	}
 }
 
-void FRDGBuilder::InitRHI(FRDGTextureSRVRef SRV)
-{
-	check(SRV);
+///////////////////////////////////////////////////////////////////////////////////////////////////
 
-	if (SRV->HasRHI())
-	{
-		return;
-	}
+void FRDGBuilder::InitTextureViewRHI(FRHICommandListBase& InRHICmdList, FRDGTextureSRVRef SRV)
+{
+	check(SRV && !SRV->ResourceRHI);
 
 	FRDGTextureRef Texture = SRV->Desc.Texture;
 	FRHITexture* TextureRHI = Texture->GetRHIUnchecked();
 	check(TextureRHI);
 
-	SRV->ResourceRHI = Texture->ViewCache->GetOrCreateSRV(RHICmdList, TextureRHI, SRV->Desc);
+	SRV->ResourceRHI = Texture->ViewCache->GetOrCreateSRV(InRHICmdList, TextureRHI, SRV->Desc);
 }
 
-void FRDGBuilder::InitRHI(FRDGTextureUAVRef UAV)
+void FRDGBuilder::InitTextureViewRHI(FRHICommandListBase& InRHICmdList, FRDGTextureUAVRef UAV)
 {
-	check(UAV);
-
-	if (UAV->HasRHI())
-	{
-		return;
-	}
+	check(UAV && !UAV->ResourceRHI);
 
 	FRDGTextureRef Texture = UAV->Desc.Texture;
 	FRHITexture* TextureRHI = Texture->GetRHIUnchecked();
 	check(TextureRHI);
 
-	UAV->ResourceRHI = Texture->ViewCache->GetOrCreateUAV(RHICmdList, TextureRHI, UAV->Desc);
-}
-
-void FRDGBuilder::BeginResourceRHI(FRDGPassHandle PassHandle, FRDGBufferRef Buffer)
-{
-	check(Buffer);
-
-	if (Buffer->HasRHI())
-	{
-		return;
-	}
-
-	check(Buffer->ReferenceCount > 0 || Buffer->bExternal || Buffer->bQueuedForUpload || IsImmediateMode());
-
-#if RDG_ENABLE_DEBUG
-	{
-		const FRDGPass* Pass = Passes[PassHandle];
-
-		// Cannot begin a resource on an async compute pass.
-		check(Pass->Pipeline == ERHIPipeline::Graphics);
-
-		// Cannot begin a resource within a merged render pass region.
-		checkf(GetPrologueBarrierPassHandle(PassHandle) == PassHandle,
-			TEXT("Cannot begin a resource within a merged render pass. Pass (Handle: %d, Name: %s), Resource %s"), PassHandle.GetIndex(), Pass->GetName(), Buffer->Name);
-	}
-#endif
-	Buffer->FinalizeDesc();
-
-	// If transient then create the resource on the transient allocator. External or extracted resource can't be transient because of lifetime tracking issues.
-	if (TransientResourceAllocator && bSupportsTransientBuffers && IsTransient(Buffer))
-	{
-		if (FRHITransientBuffer* TransientBuffer = TransientResourceAllocator->CreateBuffer(Translate(Buffer->Desc), Buffer->Name, PassHandle.GetIndex()))
-		{
-			SetRHI(Buffer, TransientBuffer, PassHandle);
-
-			const FRDGPassHandle MinAcquirePassHandle(TransientBuffer->GetAcquirePasses().Min);
-
-			AddAliasingTransition(MinAcquirePassHandle, PassHandle, Buffer, FRHITransientAliasingInfo::Acquire(TransientBuffer->GetRHI(), TransientBuffer->GetAliasingOverlaps()));
-
-			FRDGSubresourceState* InitialState = Buffer->State;
-			InitialState->SetPass(ERHIPipeline::Graphics, MinAcquirePassHandle);
-			InitialState->Access = ERHIAccess::Discard;
-
-		#if RDG_STATS
-			GRDGStatTransientBufferCount++;
-		#endif
-		}
-	}
-
-	if (!Buffer->ResourceRHI)
-	{
-		const ERDGPooledBufferAlignment Alignment = Buffer->bQueuedForUpload ? ERDGPooledBufferAlignment::PowerOfTwo : ERDGPooledBufferAlignment::Page;
-
-		SetRHI(Buffer, GRenderGraphResourcePool.FindFreeBuffer(RHICmdList, Buffer->Desc, Buffer->Name, Alignment), PassHandle);
-	}
-
-#if RHI_ENABLE_RESOURCE_INFO
-	if (Buffer->ResourceRHI)
-	{
-		Buffer->ResourceRHI->SetOwnerName(Buffer->OwnerName);
-	}
-#endif
+	UAV->ResourceRHI = Texture->ViewCache->GetOrCreateUAV(InRHICmdList, TextureRHI, UAV->Desc);
 }
 
-void FRDGBuilder::InitRHI(FRDGBufferSRVRef SRV)
+void FRDGBuilder::InitBufferViewRHI(FRHICommandListBase& InRHICmdList, FRDGBufferSRVRef SRV)
 {
 	check(SRV);
 
@@ -3764,10 +4139,10 @@ void FRDGBuilder::InitRHI(FRDGBufferSRVRef SRV)
 		SRVCreateInfo.Format = PF_Unknown;
 	}
 
-	SRV->ResourceRHI = Buffer->ViewCache->GetOrCreateSRV(RHICmdList, BufferRHI, SRVCreateInfo);
+	SRV->ResourceRHI = Buffer->ViewCache->GetOrCreateSRV(InRHICmdList, BufferRHI, SRVCreateInfo);
 }
 
-void FRDGBuilder::InitRHI(FRDGBufferUAV* UAV)
+void FRDGBuilder::InitBufferViewRHI(FRHICommandListBase& InRHICmdList, FRDGBufferUAV* UAV)
 {
 	check(UAV);
 
@@ -3789,101 +4164,32 @@ void FRDGBuilder::InitRHI(FRDGBufferUAV* UAV)
 		UAVCreateInfo.Format = PF_Unknown;
 	}
 
-	UAV->ResourceRHI = Buffer->ViewCache->GetOrCreateUAV(RHICmdList, Buffer->GetRHIUnchecked(), UAVCreateInfo);
+	UAV->ResourceRHI = Buffer->ViewCache->GetOrCreateUAV(InRHICmdList, Buffer->GetRHIUnchecked(), UAVCreateInfo);
 }
 
-void FRDGBuilder::InitRHI(FRDGView* View)
+void FRDGBuilder::InitViewRHI(FRHICommandListBase& InRHICmdList, FRDGView* View)
 {
-	if (View->HasRHI())
-	{
-		return;
-	}
+	check(!View->ResourceRHI);
 
 	switch (View->Type)
 	{
 	case ERDGViewType::TextureUAV:
-		InitRHI(static_cast<FRDGTextureUAV*>(View));
+		InitTextureViewRHI(InRHICmdList, static_cast<FRDGTextureUAV*>(View));
 		break;
 	case ERDGViewType::TextureSRV:
-		InitRHI(static_cast<FRDGTextureSRV*>(View));
+		InitTextureViewRHI(InRHICmdList, static_cast<FRDGTextureSRV*>(View));
 		break;
 	case ERDGViewType::BufferUAV:
-		InitRHI(static_cast<FRDGBufferUAV*>(View));
+		InitBufferViewRHI(InRHICmdList, static_cast<FRDGBufferUAV*>(View));
 		break;
 	case ERDGViewType::BufferSRV:
-		InitRHI(static_cast<FRDGBufferSRV*>(View));
+		InitBufferViewRHI(InRHICmdList, static_cast<FRDGBufferSRV*>(View));
 		break;
 	}
 }
 
 ///////////////////////////////////////////////////////////////////////////////////////////////////
 
-void FRDGBuilder::EndResourceRHI(FRDGPassHandle PassHandle, FRDGTextureRef Texture, uint32 ReferenceCount)
-{
-	check(Texture);
-	check(Texture->ReferenceCount != FRDGViewableResource::DeallocatedReferenceCount);
-	check(Texture->ReferenceCount >= ReferenceCount || IsImmediateMode());
-	Texture->ReferenceCount -= ReferenceCount;
-
-	if (Texture->ReferenceCount == 0)
-	{
-		if (Texture->bTransient)
-		{
-			// Texture is using a transient external render target.
-			if (Texture->RenderTarget)
-			{
-				if (!Texture->bExtracted)
-				{
-					// This releases the reference without invoking a virtual function call.
-					GRDGTransientResourceAllocator.Release(TRefCountPtr<FRDGTransientRenderTarget>(MoveTemp(Texture->Allocation)), PassHandle);
-				}
-			}
-			// Texture is using an internal transient texture.
-			else
-			{
-				TransientResourceAllocator->DeallocateMemory(Texture->TransientTexture, PassHandle.GetIndex());
-			}
-		}
-		else
-		{
-			// Only tracked render targets are released. Untracked ones persist until the end of the frame.
-			if (static_cast<FPooledRenderTarget*>(Texture->RenderTarget)->IsTracked())
-			{
-				// This releases the reference without invoking a virtual function call.
-				TRefCountPtr<FPooledRenderTarget>(MoveTemp(Texture->Allocation));
-			}
-		}
-
-		Texture->LastPass = PassHandle;
-		Texture->ReferenceCount = FRDGViewableResource::DeallocatedReferenceCount;
-	}
-}
-
-void FRDGBuilder::EndResourceRHI(FRDGPassHandle PassHandle, FRDGBufferRef Buffer, uint32 ReferenceCount)
-{
-	check(Buffer);
-	check(Buffer->ReferenceCount != FRDGViewableResource::DeallocatedReferenceCount);
-	check(Buffer->ReferenceCount >= ReferenceCount || IsImmediateMode());
-	Buffer->ReferenceCount -= ReferenceCount;
-
-	if (Buffer->ReferenceCount == 0)
-	{
-		if (Buffer->bTransient)
-		{
-			TransientResourceAllocator->DeallocateMemory(Buffer->TransientBuffer, PassHandle.GetIndex());
-		}
-		else
-		{
-			Buffer->Allocation = nullptr;
-		}
-
-		Buffer->LastPass = PassHandle;
-		Buffer->ReferenceCount = FRDGViewableResource::DeallocatedReferenceCount;
-	}
-}
-
-///////////////////////////////////////////////////////////////////////////////////////////////////
-
 #if RDG_ENABLE_DEBUG
 
 void FRDGBuilder::VisualizePassOutputs(const FRDGPass* Pass)
@@ -3985,7 +4291,6 @@ void FRDGBuilder::ClobberPassOutputs(const FRDGPass* Pass)
 	}
 
 	RDG_RECURSION_COUNTER_SCOPE(AuxiliaryPasses.Clobber);
-
 	RDG_EVENT_SCOPE(*this, "RDG ClobberResources");
 
 	const FLinearColor ClobberColor = GetClobberColor();
@@ -4002,6 +4307,42 @@ void FRDGBuilder::ClobberPassOutputs(const FRDGPass* Pass)
 		}
 	};
 
+	const auto ClobberTextureAccess = [&](FRDGTextureAccess TextureAccess)
+	{
+		if (IsWritableAccess(TextureAccess.GetAccess()))
+		{
+			FRDGTextureRef Texture = TextureAccess.GetTexture();
+	
+			if (UserValidation.TryMarkForClobber(Texture))
+			{
+				if (EnumHasAnyFlags(TextureAccess.GetAccess(), ERHIAccess::UAVMask))
+				{
+					for (int32 MipLevel = 0; MipLevel < Texture->Desc.NumMips; MipLevel++)
+					{
+						ClobberTextureUAV(CreateUAV(FRDGTextureUAVDesc(Texture, MipLevel)));
+					}
+				}
+				else if (EnumHasAnyFlags(TextureAccess.GetAccess(), ERHIAccess::RTV))
+				{
+					AddClearRenderTargetPass(*this, Texture, ClobberColor);
+				}
+			}
+		}
+	};
+
+	const auto ClobberBufferAccess = [&](FRDGBufferAccess BufferAccess)
+	{
+		if (IsWritableAccess(BufferAccess.GetAccess()))
+		{
+			FRDGBufferRef Buffer = BufferAccess.GetBuffer();
+	
+			if (UserValidation.TryMarkForClobber(Buffer))
+			{
+				AddClearUAVPass(*this, CreateUAV(Buffer), GetClobberBufferValue());
+			}
+		}
+	};
+
 	Pass->GetParameters().Enumerate([&](FRDGParameter Parameter)
 	{
 		switch (Parameter.GetType())
@@ -4021,24 +4362,31 @@ void FRDGBuilder::ClobberPassOutputs(const FRDGPass* Pass)
 		break;
 		case UBMT_RDG_TEXTURE_ACCESS:
 		{
-			if (FRDGTextureAccess TextureAccess = Parameter.GetAsTextureAccess())
+			ClobberTextureAccess(Parameter.GetAsTextureAccess());
+		}
+		break;
+		case UBMT_RDG_TEXTURE_ACCESS_ARRAY:
+		{
+			const FRDGTextureAccessArray& TextureAccessArray = Parameter.GetAsTextureAccessArray();
+
+			for (FRDGTextureAccess TextureAccess : TextureAccessArray)
 			{
-				FRDGTextureRef Texture = TextureAccess.GetTexture();
+				ClobberTextureAccess(TextureAccess);
+			}
+		}
+		break;
+		case UBMT_RDG_BUFFER_ACCESS:
+		{
+			ClobberBufferAccess(Parameter.GetAsBufferAccess());
+		}
+		break;
+		case UBMT_RDG_BUFFER_ACCESS_ARRAY:
+		{
+			const FRDGBufferAccessArray& BufferAccessArray = Parameter.GetAsBufferAccessArray();
 
-				if (UserValidation.TryMarkForClobber(Texture))
-				{
-					if (EnumHasAnyFlags(TextureAccess.GetAccess(), ERHIAccess::UAVMask))
-					{
-						for (int32 MipLevel = 0; MipLevel < Texture->Desc.NumMips; MipLevel++)
-						{
-							ClobberTextureUAV(CreateUAV(FRDGTextureUAVDesc(Texture, MipLevel)));
-						}
-					}
-					else if (EnumHasAnyFlags(TextureAccess.GetAccess(), ERHIAccess::RTV))
-					{
-						AddClearRenderTargetPass(*this, Texture, ClobberColor);
-					}
-				}
+			for (FRDGBufferAccess BufferAccess : BufferAccessArray)
+			{
+				ClobberBufferAccess(BufferAccess);
 			}
 		}
 		break;
@@ -4097,49 +4445,55 @@ void FRDGBuilder::ClobberPassOutputs(const FRDGPass* Pass)
 #if WITH_MGPU
 void FRDGBuilder::ForceCopyCrossGPU()
 {
-	// Initialize set of external buffers
-	TSet<FRHIBuffer*> ExternalBufferSet;
-	ExternalBufferSet.Reserve(ExternalBuffers.Num());
-
-	for (auto ExternalBufferIt = ExternalBuffers.CreateConstIterator(); ExternalBufferIt; ++ExternalBufferIt)
+	const auto GetLastProducerGPUMask = [](FRDGProducerStatesByPipeline& LastProducers) -> TOptional<FRHIGPUMask>
 	{
-		ExternalBufferSet.Emplace(ExternalBufferIt.Value()->GetRHIUnchecked());
-	}
+		for (const FRDGProducerState& LastProducer : LastProducers)
+		{
+			if (LastProducer.Pass && !LastProducer.Pass->bCulled)
+			{
+				return LastProducer.Pass->GPUMask;
+			}
+		}
+		return {};
+	};
 
-	// Generate list of cross GPU resources from all passes, and the GPU mask where they were last written 
-	TMap<FRHIBuffer*, FRHIGPUMask> BuffersToTransfer;
-	TMap<FRHITexture*, FRHIGPUMask> TexturesToTransfer;
+	Experimental::TRobinHoodHashMap<FRHIBuffer*, FRHIGPUMask, DefaultKeyFuncs<FRHIBuffer*>, FRDGArrayAllocator> BuffersToTransfer;
+	BuffersToTransfer.Reserve(ExternalBuffers.Num());
 
-	for (FRDGPassHandle PassHandle = GetProloguePassHandle(); PassHandle <= GetEpiloguePassHandle(); ++PassHandle)
+	for (auto& ExternalBuffer : ExternalBuffers)
 	{
-		FRDGPass* Pass = Passes[PassHandle];
+		FRHIBuffer* BufferRHI = ExternalBuffer.Key;
+		FRDGBuffer* BufferRDG = ExternalBuffer.Value;
 
-		for (int32 BufferIndex = 0; BufferIndex < Pass->BufferStates.Num(); BufferIndex++)
+		if (!EnumHasAnyFlags(BufferRDG->Desc.Usage, BUF_MultiGPUAllocate | BUF_MultiGPUGraphIgnore))
 		{
-			FRHIBuffer* BufferRHI = Pass->BufferStates[BufferIndex].Buffer->GetRHIUnchecked();
+			TOptional<FRHIGPUMask> GPUMask = GetLastProducerGPUMask(BufferRDG->LastProducer);
 
-			if (ExternalBufferSet.Contains(BufferRHI) &&
-				!EnumHasAnyFlags(BufferRHI->GetUsage(), BUF_MultiGPUAllocate | BUF_MultiGPUGraphIgnore) &&
-				EnumHasAnyFlags(Pass->BufferStates[BufferIndex].State.Access, ERHIAccess::WritableMask))
+			if (GPUMask)
 			{
-				BuffersToTransfer.Emplace(BufferRHI) = Pass->GPUMask;
+				BuffersToTransfer.FindOrAdd(BufferRHI, *GPUMask);
 			}
 		}
+	}
 
-		for (int32 TextureIndex = 0; TextureIndex < Pass->TextureStates.Num(); TextureIndex++)
+	Experimental::TRobinHoodHashMap<FRHITexture*, FRHIGPUMask, DefaultKeyFuncs<FRHITexture*>, FRDGArrayAllocator> TexturesToTransfer;
+	TexturesToTransfer.Reserve(ExternalTextures.Num());
+
+	for (auto& ExternalTexture : ExternalTextures)
+	{
+		FRHITexture* TextureRHI = ExternalTexture.Key;
+		FRDGTexture* TextureRDG = ExternalTexture.Value;
+
+		if (!EnumHasAnyFlags(TextureRDG->Desc.Flags, TexCreate_MultiGPUGraphIgnore))
 		{
-			if (ExternalTextures.Contains(Pass->TextureStates[TextureIndex].Texture->GetRHIUnchecked()))
+			for (auto& LastProducer : TextureRDG->LastProducers)
 			{
-				for (int32 StateIndex = 0; StateIndex < Pass->TextureStates[TextureIndex].State.Num(); StateIndex++)
-				{
-					FRHITexture* TextureRHI = Pass->TextureStates[TextureIndex].Texture->GetRHIUnchecked();
+				TOptional<FRHIGPUMask> GPUMask = GetLastProducerGPUMask(LastProducer);
 
-					if (TextureRHI &&
-						!EnumHasAnyFlags(TextureRHI->GetFlags(), TexCreate_MultiGPUGraphIgnore) &&
-						EnumHasAnyFlags(Pass->TextureStates[TextureIndex].State[StateIndex].Access, ERHIAccess::WritableMask))
-					{
-						TexturesToTransfer.Emplace(Pass->TextureStates[TextureIndex].Texture->GetRHIUnchecked()) = Pass->GPUMask;
-					}
+				if (GPUMask)
+				{
+					TexturesToTransfer.FindOrAdd(TextureRHI, *GPUMask);
+					break;
 				}
 			}
 		}
@@ -4147,15 +4501,16 @@ void FRDGBuilder::ForceCopyCrossGPU()
 
 	// Now that we've got the list of external resources, and the GPU they were last written to, make a list of what needs to
 	// be propagated to other GPUs.
-	TArray<FTransferResourceParams> Transfers;
+	TArray<FTransferResourceParams, FRDGArrayAllocator> Transfers;
+	Transfers.Reserve(BuffersToTransfer.Num() + TexturesToTransfer.Num());
 	const FRHIGPUMask AllGPUMask = FRHIGPUMask::All();
 	const bool bPullData = false;
 	const bool bLockstepGPUs = true;
 
-	for (auto BufferIt = BuffersToTransfer.CreateConstIterator(); BufferIt; ++BufferIt)
+	for (auto& KeyValue : BuffersToTransfer)
 	{
-		FRHIBuffer* Buffer = BufferIt.Key();
-		FRHIGPUMask GPUMask = BufferIt.Value();
+		FRHIBuffer* Buffer  = KeyValue.Key;
+		FRHIGPUMask GPUMask = KeyValue.Value;
 
 		for (uint32 GPUIndex : AllGPUMask)
 		{
@@ -4166,10 +4521,10 @@ void FRDGBuilder::ForceCopyCrossGPU()
 		}
 	}
 
-	for (auto TextureIt = TexturesToTransfer.CreateConstIterator(); TextureIt; ++TextureIt)
+	for (auto& KeyValue : TexturesToTransfer)
 	{
-		FRHITexture* Texture = TextureIt.Key();
-		FRHIGPUMask GPUMask = TextureIt.Value();
+		FRHITexture* Texture = KeyValue.Key;
+		FRHIGPUMask GPUMask  = KeyValue.Value;
 
 		for (uint32 GPUIndex : AllGPUMask)
 		{
@@ -4185,4 +4540,4 @@ void FRDGBuilder::ForceCopyCrossGPU()
 		RHICmdList.TransferResources(Transfers);
 	}
 }
-#endif  // WITH_MGPU
+#endif  // WITH_MGPU
\ No newline at end of file
diff --git a/Engine/Source/Runtime/RenderCore/Private/RenderGraphEvent.cpp b/Engine/Source/Runtime/RenderCore/Private/RenderGraphEvent.cpp
--- a/Engine/Source/Runtime/RenderCore/Private/RenderGraphEvent.cpp
+++ b/Engine/Source/Runtime/RenderCore/Private/RenderGraphEvent.cpp
@@ -568,7 +568,7 @@ FString FRDGEventScope::GetPath(const FRDGEventName& Event) const
 
 FRDGEventScopeGuard::FRDGEventScopeGuard(FRDGBuilder& InGraphBuilder, FRDGEventName&& ScopeName, bool InbCondition, ERDGEventScopeFlags InFlags)
 	: GraphBuilder(InGraphBuilder)
-	, bCondition(InbCondition && !GraphBuilder.bFinalEventScopeActive)
+	, bCondition(InbCondition && !GraphBuilder.GPUScopeStacks.bFinalEventScopeActive)
 {
 	if (bCondition)
 	{
@@ -577,7 +577,7 @@ FRDGEventScopeGuard::FRDGEventScopeGuard(FRDGBuilder& InGraphBuilder, FRDGEventN
 			EnumRemoveFlags(InFlags, ERDGEventScopeFlags::Final);
 		}
 
-		GraphBuilder.bFinalEventScopeActive = EnumHasAnyFlags(InFlags, ERDGEventScopeFlags::Final);
+		GraphBuilder.GPUScopeStacks.bFinalEventScopeActive = (EnumHasAnyFlags(InFlags, ERDGEventScopeFlags::Final));
 		GraphBuilder.GPUScopeStacks.BeginEventScope(MoveTemp(ScopeName), GraphBuilder.RHICmdList.GetGPUMask(), InFlags);
 	}
 }
@@ -587,7 +587,7 @@ FRDGEventScopeGuard::~FRDGEventScopeGuard()
 	if (bCondition)
 	{
 		GraphBuilder.GPUScopeStacks.EndEventScope();
-		GraphBuilder.bFinalEventScopeActive = false;
+		GraphBuilder.GPUScopeStacks.bFinalEventScopeActive = false;
 	}
 }
 
diff --git a/Engine/Source/Runtime/RenderCore/Private/RenderGraphPass.cpp b/Engine/Source/Runtime/RenderCore/Private/RenderGraphPass.cpp
--- a/Engine/Source/Runtime/RenderCore/Private/RenderGraphPass.cpp
+++ b/Engine/Source/Runtime/RenderCore/Private/RenderGraphPass.cpp
@@ -123,7 +123,7 @@ FRDGBarrierBatchBegin::FRDGBarrierBatchBegin(ERHIPipeline InPipelinesToBegin, ER
 #endif
 {}
 
-void FRDGBarrierBatchBegin::AddTransition(FRDGViewableResource* Resource, const FRHITransitionInfo& Info)
+void FRDGBarrierBatchBegin::AddTransition(FRDGViewableResource* Resource, FRDGTransitionInfo Info)
 {
 	Transitions.Add(Info);
 	bTransitionNeeded = true;
@@ -151,10 +151,10 @@ void FRDGBarrierBatchBegin::AddAlias(FRDGViewableResource* Resource, const FRHIT
 #endif
 }
 
-void FRDGBarrierBatchBegin::CreateTransition()
+void FRDGBarrierBatchBegin::CreateTransition(TConstArrayView<FRHITransitionInfo> TransitionsRHI)
 {
 	check(bTransitionNeeded && !Transition);
-	Transition = RHICreateTransition(FRHITransitionCreateInfo(PipelinesToBegin, PipelinesToEnd, TransitionFlags, Transitions, Aliases));
+	Transition = RHICreateTransition(FRHITransitionCreateInfo(PipelinesToBegin, PipelinesToEnd, TransitionFlags, TransitionsRHI, Aliases));
 }
 
 void FRDGBarrierBatchBegin::Submit(FRHIComputeCommandList& RHICmdList, ERHIPipeline Pipeline, FRDGTransitionQueue& TransitionsToBegin)
diff --git a/Engine/Source/Runtime/RenderCore/Private/RenderGraphResourcePool.cpp b/Engine/Source/Runtime/RenderCore/Private/RenderGraphResourcePool.cpp
--- a/Engine/Source/Runtime/RenderCore/Private/RenderGraphResourcePool.cpp
+++ b/Engine/Source/Runtime/RenderCore/Private/RenderGraphResourcePool.cpp
@@ -102,7 +102,7 @@ TRefCountPtr<FRDGPooledBuffer> FRDGBufferPool::FindFreeBuffer(FRHICommandListBas
 		check(PooledBuffer->GetAlignedDesc() == AlignedDesc);
 
 		PooledBuffer->LastUsedFrame = FrameCounter;
-		PooledBuffer->ViewCache.SetDebugName(InDebugName);
+		PooledBuffer->ViewCache.SetDebugName(RHICmdList, InDebugName);
 		PooledBuffer->Name = InDebugName;
 
 	#if !(UE_BUILD_SHIPPING || UE_BUILD_TEST)
diff --git a/Engine/Source/Runtime/RenderCore/Private/RenderGraphResources.cpp b/Engine/Source/Runtime/RenderCore/Private/RenderGraphResources.cpp
--- a/Engine/Source/Runtime/RenderCore/Private/RenderGraphResources.cpp
+++ b/Engine/Source/Runtime/RenderCore/Private/RenderGraphResources.cpp
@@ -8,7 +8,32 @@ inline bool SkipUAVBarrier(const FRDGSubresourceState& Previous, const FRDGSubre
 	return SkipUAVBarrier(Previous.NoUAVBarrierFilter.GetUniqueHandle(), Next.NoUAVBarrierFilter.GetUniqueHandle());
 }
 
-FRDGViewableResource::FRDGViewableResource(const TCHAR* InName, const ERDGViewableResourceType InType, bool bSkipTracking, bool bImmediateFirstBarrier)
+FRDGTexture::FRDGTexture(const TCHAR* InName, const FRDGTextureDesc& InDesc, ERDGTextureFlags InFlags)
+	: FRDGViewableResource(InName, ERDGViewableResourceType::Texture, EnumHasAnyFlags(InFlags, ERDGTextureFlags::SkipTracking), !EnumHasAnyFlags(InFlags, ERDGTextureFlags::ForceImmediateFirstBarrier) && !EnumHasAnyFlags(InDesc.Flags, ETextureCreateFlags::Presentable))
+	, Desc(InDesc)
+	, Flags(InFlags)
+	, Layout(InDesc)
+	, WholeRange(Layout)
+	, SubresourceCount(Layout.GetSubresourceCount())
+{
+	if (EnumHasAnyFlags(Desc.Flags, ETextureCreateFlags::Foveation))
+	{
+		EpilogueAccess = ERHIAccess::ShadingRateSource;
+	}
+
+	State.SetNum(SubresourceCount);
+	FirstState.SetNum(SubresourceCount);
+	MergeState.SetNum(SubresourceCount);
+	LastProducers.SetNum(SubresourceCount);
+}
+
+FRDGBuffer::FRDGBuffer(const TCHAR* InName, const FRDGBufferDesc& InDesc, ERDGBufferFlags InFlags)
+	: FRDGViewableResource(InName, ERDGViewableResourceType::Buffer, EnumHasAnyFlags(InFlags, ERDGBufferFlags::SkipTracking), !EnumHasAnyFlags(InFlags, ERDGBufferFlags::ForceImmediateFirstBarrier))
+	, Desc(InDesc)
+	, Flags(InFlags)
+{}
+
+FRDGViewableResource::FRDGViewableResource(const TCHAR* InName, const ERDGViewableResourceType InType, bool bSkipTracking, bool bInSplitFirstTransition)
 	: FRDGResource(InName)
 	, Type(InType)
 	, bExternal(0)
@@ -16,10 +41,12 @@ FRDGViewableResource::FRDGViewableResource(const TCHAR* InName, const ERDGViewab
 	, bProduced(0)
 	, bTransient(0)
 	, bForceNonTransient(0)
-	, TransientExtractionHint(ETransientExtractionHint::None)
-	, bLastOwner(1)
+	, bSkipLastTransition(0)
+	, bSplitFirstTransition(bInSplitFirstTransition)
 	, bQueuedForUpload(0)
-	, FirstBarrier(EFirstBarrier::Split)
+	, bCollectForAllocate(1)
+	, bQueuedForReservedCommit(0)
+	, TransientExtractionHint(ETransientExtractionHint::None)
 	, ReferenceCount(IsImmediateMode() ? 1 : 0)
 {
 	if (bSkipTracking)
@@ -28,11 +55,6 @@ FRDGViewableResource::FRDGViewableResource(const TCHAR* InName, const ERDGViewab
 		AccessModeState.bLocked = 1;
 		AccessModeState.ActiveMode = AccessModeState.Mode;
 	}
-
-	if (bImmediateFirstBarrier)
-	{
-		FirstBarrier = EFirstBarrier::ImmediateRequested;
-	}
 }
 
 bool FRDGSubresourceState::IsMergeAllowed(ERDGViewableResourceType ResourceType, const FRDGSubresourceState& Previous, const FRDGSubresourceState& Next)
@@ -151,3 +173,12 @@ FRDGTextureSubresourceRange FRDGTexture::GetSubresourceRangeSRV() const
 
 	return Range;
 }
+
+void FRDGBuffer::FinalizeDesc()
+{
+	if (NumElementsCallback)
+	{
+		Desc.NumElements = FMath::Max(NumElementsCallback(), 1u);
+		NumElementsCallback = {};
+	}
+}
\ No newline at end of file
diff --git a/Engine/Source/Runtime/RenderCore/Private/RenderGraphTrace.cpp b/Engine/Source/Runtime/RenderCore/Private/RenderGraphTrace.cpp
--- a/Engine/Source/Runtime/RenderCore/Private/RenderGraphTrace.cpp
+++ b/Engine/Source/Runtime/RenderCore/Private/RenderGraphTrace.cpp
@@ -310,7 +310,7 @@ void FRDGTrace::OutputGraphEnd(const FRDGBuilder& GraphBuilder)
 
 		TransientAllocation.Reset();
 
-		if (Texture->bTransient)
+		if (Texture->TransientTexture)
 		{
 			const FRHITransientResource* Resource = Texture->TransientTexture;
 			TransientAllocation.Fill(TransientAllocationStats, Resource);
@@ -338,7 +338,7 @@ void FRDGTrace::OutputGraphEnd(const FRDGBuilder& GraphBuilder)
 			<< TextureMessage.NumSamples(Texture->Desc.NumSamples)
 			<< TextureMessage.IsExternal(bool(Texture->bExternal))
 			<< TextureMessage.IsExtracted(bool(Texture->bExtracted))
-			<< TextureMessage.IsCulled(bool(Texture->IsCulled()))
+			<< TextureMessage.IsCulled(bool(Texture->ReferenceCount == 0))
 			<< TextureMessage.IsTrackingSkipped(EnumHasAnyFlags(Texture->Flags, ERDGTextureFlags::SkipTracking))
 			<< TextureMessage.IsTransient(bool(Texture->bTransient))
 			<< TextureMessage.IsTransientUntracked(false)
@@ -351,7 +351,7 @@ void FRDGTrace::OutputGraphEnd(const FRDGBuilder& GraphBuilder)
 
 		TransientAllocation.Reset();
 
-		if (Buffer->bTransient)
+		if (Buffer->TransientBuffer)
 		{
 			const FRHITransientResource* Resource = Buffer->TransientBuffer;
 			TransientAllocation.Fill(TransientAllocationStats, Resource);
@@ -372,7 +372,7 @@ void FRDGTrace::OutputGraphEnd(const FRDGBuilder& GraphBuilder)
 			<< BufferMessage.NumElements(Buffer->Desc.NumElements)
 			<< BufferMessage.IsExternal(bool(Buffer->bExternal))
 			<< BufferMessage.IsExtracted(bool(Buffer->bExtracted))
-			<< BufferMessage.IsCulled(bool(Buffer->IsCulled()))
+			<< BufferMessage.IsCulled(bool(Buffer->ReferenceCount == 0))
 			<< BufferMessage.IsTrackingSkipped(EnumHasAnyFlags(Buffer->Flags, ERDGBufferFlags::SkipTracking))
 			<< BufferMessage.IsTransient(bool(Buffer->bTransient))
 			<< BufferMessage.IsTransientUntracked(false)
diff --git a/Engine/Source/Runtime/RenderCore/Private/RenderGraphValidation.cpp b/Engine/Source/Runtime/RenderCore/Private/RenderGraphValidation.cpp
--- a/Engine/Source/Runtime/RenderCore/Private/RenderGraphValidation.cpp
+++ b/Engine/Source/Runtime/RenderCore/Private/RenderGraphValidation.cpp
@@ -10,45 +10,6 @@
 
 namespace
 {
-template <typename TFunction>
-void EnumerateSubresources(const FRHITransitionInfo& Transition, uint32 NumMips, uint32 NumArraySlices, uint32 NumPlaneSlices, TFunction Function)
-{
-	uint32 MinMipIndex = 0;
-	uint32 MaxMipIndex = NumMips;
-	uint32 MinArraySlice = 0;
-	uint32 MaxArraySlice = NumArraySlices;
-	uint32 MinPlaneSlice = 0;
-	uint32 MaxPlaneSlice = NumPlaneSlices;
-
-	if (!Transition.IsAllMips())
-	{
-		MinMipIndex = Transition.MipIndex;
-		MaxMipIndex = MinMipIndex + 1;
-	}
-
-	if (!Transition.IsAllArraySlices())
-	{
-		MinArraySlice = Transition.ArraySlice;
-		MaxArraySlice = MinArraySlice + 1;
-	}
-
-	if (!Transition.IsAllPlaneSlices())
-	{
-		MinPlaneSlice = Transition.PlaneSlice;
-		MaxPlaneSlice = MinPlaneSlice + 1;
-	}
-
-	for (uint32 PlaneSlice = MinPlaneSlice; PlaneSlice < MaxPlaneSlice; ++PlaneSlice)
-	{
-		for (uint32 ArraySlice = MinArraySlice; ArraySlice < MaxArraySlice; ++ArraySlice)
-		{
-			for (uint32 MipIndex = MinMipIndex; MipIndex < MaxMipIndex; ++MipIndex)
-			{
-				Function(FRDGTextureSubresource(MipIndex, ArraySlice, PlaneSlice));
-			}
-		}
-	}
-}
 
 const ERHIAccess AccessMaskCopy    = ERHIAccess::CopySrc | ERHIAccess::CopyDest | ERHIAccess::CPURead;
 const ERHIAccess AccessMaskCompute = ERHIAccess::SRVCompute | ERHIAccess::UAVCompute;
@@ -200,7 +161,7 @@ void FRDGUserValidation::ValidateCreateResource(FRDGResourceRef Resource)
 void FRDGUserValidation::ValidateCreateViewableResource(FRDGViewableResource* Resource)
 {
 	ValidateCreateResource(Resource);
-	Resource->ViewableDebugData = Allocator.Alloc<FRDGViewableResourceDebugData>();
+	Resource->ViewableDebugData = Allocator.AllocNoDestruct<FRDGViewableResourceDebugData>();
 }
 
 void FRDGUserValidation::ValidateCreateTexture(FRDGTextureRef Texture)
@@ -211,7 +172,7 @@ void FRDGUserValidation::ValidateCreateTexture(FRDGTextureRef Texture)
 	}
 
 	ValidateCreateViewableResource(Texture);
-	Texture->TextureDebugData = Allocator.Alloc<FRDGTextureDebugData>();
+	Texture->TextureDebugData = Allocator.AllocNoDestruct<FRDGTextureDebugData>();
 	if (GRDGDebug)
 	{
 		TrackedTextures.Add(Texture);
@@ -513,11 +474,7 @@ void FRDGUserValidation::ValidateCommitBuffer(FRDGBufferRef Buffer, uint64 Commi
 	check(Buffer);
 	checkf(EnumHasAnyFlags(Buffer->Desc.Usage, BUF_ReservedResource), TEXT("Buffer %s is not marked as reserved and thus cannot be queued for reserved resource commit."), Buffer->Name);
 	checkf(Buffer->IsExternal(), TEXT("Only external buffers support commit operation. It is expected that reserved resource commit mechanism is only used when perserving buffer contents is required."));
-	checkf(!Buffer->IsTransient(), TEXT("Transient buffers may not be reserved and do not support commit operation. It is expected that reserved resource commit mechanism is only used when perserving buffer contents is required."));
 	checkf(CommitSizeInBytes > 0, TEXT("Attempted to set a reserved buffer commit size of 0 for buffer %s"), Buffer->Name);
-
-	// This may be relaxed in the future, by committing the maximum amount of memory requested during the graph setup
-	checkf(!Buffer->PendingCommitSize, TEXT("Buffer %s is already queued for a reserved commit size of %ull"), Buffer->Name, Buffer->PendingCommitSize); 
 }
 
 void FRDGUserValidation::ValidateExtractTexture(FRDGTextureRef Texture, TRefCountPtr<IPooledRenderTarget>* OutTexturePtr)
@@ -1388,7 +1345,7 @@ void FRDGBarrierValidation::ValidateBarrierBatchBegin(const FRDGPass* Pass, cons
 		for (int32 Index = 0; Index < Batch.Transitions.Num(); ++Index)
 		{
 			FRDGViewableResource* Resource = Batch.DebugTransitionResources[Index];
-			const FRHITransitionInfo& Transition = Batch.Transitions[Index];
+			const FRDGTransitionInfo& Transition = Batch.Transitions[Index];
 
 			if (Resource->Type == ERDGViewableResourceType::Texture)
 			{
@@ -1456,29 +1413,21 @@ void FRDGBarrierValidation::ValidateBarrierBatchBegin(const FRDGPass* Pass, cons
 
 		const FRDGTextureSubresourceLayout SubresourceLayout = Texture->GetSubresourceLayout();
 
-		for (const FRHITransitionInfo& Transition : Transitions)
+		for (const FRDGTransitionInfo& Transition : Transitions)
 		{
-			check(SubresourceLayout.GetSubresourceCount() > 0);
-
-			EnumerateSubresources(Transition, SubresourceLayout.NumMips, SubresourceLayout.NumArraySlices, SubresourceLayout.NumPlaneSlices,
-				[&](FRDGTextureSubresource Subresource)
-			{
-				const int32 SubresourceIndex = SubresourceLayout.GetSubresourceIndex(Subresource);
-
-				UE_LOG(LogRDG, Display, TEXT("\t\tMip(%d), Array(%d), Slice(%d): [%s, %s] -> [%s, %s]"),
-					Subresource.MipIndex, Subresource.ArraySlice, Subresource.PlaneSlice,
-					*GetRHIAccessName(Transition.AccessBefore),
-					*GetRHIPipelineName(Batch.DebugPipelinesToBegin),
-					*GetRHIAccessName(Transition.AccessAfter),
-					*GetRHIPipelineName(Batch.DebugPipelinesToEnd));
-			});
+			UE_LOG(LogRDG, Display, TEXT("\t\tMip(%d), Array(%d), Slice(%d): [%s, %s] -> [%s, %s]"),
+				Transition.MipIndex, Transition.ArraySlice, Transition.PlaneSlice,
+				*GetRHIAccessName(Transition.AccessBefore),
+				*GetRHIPipelineName(Batch.DebugPipelinesToBegin),
+				*GetRHIAccessName(Transition.AccessAfter),
+				*GetRHIPipelineName(Batch.DebugPipelinesToEnd));
 		}
 	}
 
 	for (const auto& Pair : ResourceMap->Buffers)
 	{
 		FRDGBufferRef Buffer = Pair.Key;
-		const FRHITransitionInfo& Transition = Pair.Value;
+		const FRDGTransitionInfo& Transition = Pair.Value;
 
 		if (!IsDebugAllowedForResource(Buffer->Name))
 		{
diff --git a/Engine/Source/Runtime/RenderCore/Public/RenderGraphAllocator.h b/Engine/Source/Runtime/RenderCore/Public/RenderGraphAllocator.h
--- a/Engine/Source/Runtime/RenderCore/Public/RenderGraphAllocator.h
+++ b/Engine/Source/Runtime/RenderCore/Public/RenderGraphAllocator.h
@@ -3,6 +3,7 @@
 #pragma once
 
 #include "Containers/Array.h"
+#include "Containers/StaticArray.h"
 #include "Containers/ContainerAllocationPolicies.h"
 #include "CoreGlobals.h"
 #include "CoreMinimal.h"
@@ -15,9 +16,11 @@
 #include "Templates/UnrealTemplate.h"
 
 #define RDG_USE_MALLOC USING_ADDRESS_SANITISER
+#define RDG_ALLOCATOR_DEBUG UE_BUILD_DEBUG
 
-/** Private allocator used by RDG to track its internal memory. All memory is released after RDG builder execution. */
-class FRDGAllocator final
+class FRDGAllocator;
+
+class FRDGAllocator
 {
 public:
 	class FObject
@@ -39,19 +42,35 @@ public:
 		T Alloc;
 	};
 
-	static RENDERCORE_API FRDGAllocator& Get();
-	RENDERCORE_API ~FRDGAllocator();
+	RENDERCORE_API static FRDGAllocator& GetTLS();
+
+	FRDGAllocator();
+	FRDGAllocator(FRDGAllocator&&);
+	FRDGAllocator& operator=(FRDGAllocator&&);
+	~FRDGAllocator();
 
 	/** Allocates raw memory. */
 	FORCEINLINE void* Alloc(uint64 SizeInBytes, uint32 AlignInBytes)
 	{
+		void* Memory;
+
+#if RDG_ALLOCATOR_DEBUG
+		AcquireAccess();
+#endif
+
 #if RDG_USE_MALLOC
-		void* Memory = FMemory::Malloc(SizeInBytes, AlignInBytes);
-		GetContext().RawAllocs.Emplace(Memory);
-		return Memory;
+		Memory = FMemory::Malloc(SizeInBytes, AlignInBytes);
+		Mallocs.Emplace(Memory);
+		NumMallocBytes += SizeInBytes;
 #else
-		return GetContext().MemStack.Alloc(SizeInBytes, AlignInBytes);
+		Memory = MemStack.Alloc(SizeInBytes, AlignInBytes);
 #endif
+
+#if RDG_ALLOCATOR_DEBUG
+		ReleaseAccess();
+#endif
+
+		return Memory;
 	}
 
 	/** Allocates an uninitialized type without destructor tracking. */
@@ -65,14 +84,22 @@ public:
 	template <typename T, typename... TArgs>
 	FORCEINLINE T* Alloc(TArgs&&... Args)
 	{
-		FContext& LocalContext = GetContext();
-	#if RDG_USE_MALLOC
+#if RDG_ALLOCATOR_DEBUG
+		AcquireAccess();
+#endif
+
+#if RDG_USE_MALLOC
 		TObject<T>* Object = new TObject<T>(Forward<TArgs&&>(Args)...);
-	#else
-		TObject<T>* Object = new(LocalContext.MemStack) TObject<T>(Forward<TArgs&&>(Args)...);
-	#endif
+		NumMallocBytes += sizeof(Object);
+#else
+		TObject<T>* Object = new(MemStack) TObject<T>(Forward<TArgs&&>(Args)...);
+#endif
 		check(Object);
-		LocalContext.Objects.Add(Object);
+		Objects.Add(Object);
+
+#if RDG_ALLOCATOR_DEBUG
+		ReleaseAccess();
+#endif
 		return &Object->Alloc;
 	}
 
@@ -80,79 +107,53 @@ public:
 	template <typename T, typename... TArgs>
 	FORCEINLINE T* AllocNoDestruct(TArgs&&... Args)
 	{
-#if RDG_USE_MALLOC
 		return new (AllocUninitialized<T>(1)) T(Forward<TArgs&&>(Args)...);
-#else
-		return new (GetContext().MemStack) T(Forward<TArgs&&>(Args)...);
-#endif
 	}
 
 	FORCEINLINE int32 GetByteCount() const
 	{
-	#if RDG_USE_MALLOC
-		return 0;
-	#else
-		return Context.MemStack.GetByteCount() + ContextForTasks.MemStack.GetByteCount();
-	#endif
+#if RDG_USE_MALLOC
+		return NumMallocBytes;
+#else
+		return MemStack.GetByteCount();
+#endif
 	}
 
-private:
-	FRDGAllocator() = default;
-	FRDGAllocator(FRDGAllocator&&) = default;
-	FRDGAllocator& operator = (FRDGAllocator&&) = default;
+	void ReleaseAll();
 
-	RENDERCORE_API void ReleaseAll();
-
-	struct FContext
-	{
+private:
 #if RDG_USE_MALLOC
-		TArray<void*> RawAllocs;
+	TArray<void*> Mallocs;
+	uint32 NumMallocBytes = 0;
 #else
-		FMemStackBase MemStack;
+	FMemStackBase MemStack;
 #endif
-		TArray<FObject*> Objects;
-
-		void ReleaseAll();
-	};
+	TArray<FObject*> Objects;
 
-	FContext Context;
-	uint8 Pad[PLATFORM_CACHE_LINE_SIZE];
-	FContext ContextForTasks;
+#if RDG_ALLOCATOR_DEBUG
+	RENDERCORE_API void AcquireAccess();
+	RENDERCORE_API void ReleaseAccess();
 
-	FORCEINLINE FContext& GetContext()
-	{
-		return FTaskTagScope::IsCurrentTag(ETaskTag::ERenderingThread) ? Context : ContextForTasks;
-	}
+	std::atomic_int32_t NumAccesses{0};
+	static thread_local int32 NumAccessesTLS;
+#endif
 
-	template <uint32>
-	friend class TRDGArrayAllocator;
 	friend class FRDGAllocatorScope;
+	static thread_local FRDGAllocator* AllocatorTLS;
 };
 
-#define RDG_FRIEND_ALLOCATOR_FRIEND(Type) friend class FRDGAllocator::TObject<Type>
-
-/** Base class for RDG builder which scopes the allocations and releases them in the destructor. */
 class FRDGAllocatorScope
 {
 public:
-	FRDGAllocatorScope()
-		: Allocator(FRDGAllocator::Get())
-	{}
-
+	RENDERCORE_API FRDGAllocatorScope(FRDGAllocator& Allocator);
 	RENDERCORE_API ~FRDGAllocatorScope();
 
-protected:
-	FRDGAllocator& Allocator;
-
-	void BeginAsyncDelete(TUniqueFunction<void()>&& InAsyncDeleteFunction)
-	{
-		AsyncDeleteFunction = MoveTemp(InAsyncDeleteFunction);
-	}
-
 private:
-	TUniqueFunction<void()> AsyncDeleteFunction;
+	FRDGAllocator* AllocatorToRestore;
 };
 
+#define RDG_FRIEND_ALLOCATOR_FRIEND(Type) friend class FRDGAllocator::TObject<Type>
+
 namespace UE::RenderCore::Private
 {
 	[[noreturn]] RENDERCORE_API void OnInvalidRDGAllocatorNum(int32 NewNum, SIZE_T NumBytesPerElement);
@@ -203,7 +204,7 @@ public:
 				// Allocate memory from the allocator.
 				const int32 AllocSize = (int32)(NumElements * NumBytesPerElement);
 				const int32 AllocAlignment = FMath::Max(Alignment, (uint32)alignof(ElementType));
-				Data = (ElementType*)FRDGAllocator::Get().Alloc(AllocSize, FMath::Max(AllocSize >= 16 ? (int32)16 : (int32)8, AllocAlignment));
+				Data = (ElementType*)FRDGAllocator::GetTLS().Alloc(AllocSize, FMath::Max(AllocSize >= 16 ? (int32)16 : (int32)8, AllocAlignment));
 
 				// If the container previously held elements, copy them into the new allocation.
 				if (OldData && PreviousNumElements)
diff --git a/Engine/Source/Runtime/RenderCore/Public/RenderGraphBuilder.h b/Engine/Source/Runtime/RenderCore/Public/RenderGraphBuilder.h
--- a/Engine/Source/Runtime/RenderCore/Public/RenderGraphBuilder.h
+++ b/Engine/Source/Runtime/RenderCore/Public/RenderGraphBuilder.h
@@ -30,6 +30,7 @@
 #include "Templates/RefCounting.h"
 #include "Templates/UnrealTemplate.h"
 #include "Tasks/Pipe.h"
+#include "Experimental/Containers/RobinHoodHashTable.h"
 
 enum class ERenderTargetTexture : uint8;
 struct FParallelPassSet;
@@ -44,8 +45,27 @@
  *  destruction.
  */
 class FRDGBuilder
-	: FRDGAllocatorScope
 {
+	struct
+	{
+		// Allocator for all root graph allocations on the graph builder thread.
+		FRDGAllocator Root;
+
+		// Allocator for async pass and parallel execute setup.
+		FRDGAllocator Task;
+
+		// Allocator for all allocations related to states / transitions.
+		FRDGAllocator Transition;
+
+		int32 GetByteCount() const
+		{
+			return Root.GetByteCount() + Task.GetByteCount() + Transition.GetByteCount();
+		}
+
+	} Allocators;
+
+	FRDGAllocatorScope RootAllocatorScope;
+
 public:
 	RENDERCORE_API FRDGBuilder(FRHICommandListImmediate& RHICmdList, FRDGEventName Name = {}, ERDGBuilderFlags Flags = ERDGBuilderFlags::None);
 	FRDGBuilder(const FRDGBuilder&) = delete;
@@ -255,7 +275,7 @@
 	/** Whether RDG will launch async tasks when AddSetup{CommandList}Task is called. */
 	inline bool IsParallelSetupEnabled() const
 	{
-		return bParallelSetupEnabled;
+		return ParallelSetup.bEnabled;
 	}
 
 	/** Tells the builder to delete unused RHI resources. The behavior of this method depends on whether RDG immediate mode is enabled:
@@ -418,21 +438,37 @@
 #endif
 
 private:
-	static RENDERCORE_API const char* const kDefaultUnaccountedCSVStat;
+	static const char* const kDefaultUnaccountedCSVStat;
 
 	const FRDGEventName BuilderName;
 
-	template <typename ParameterStructType, typename ExecuteLambdaType>
-	FRDGPassRef AddPassInternal(
-		FRDGEventName&& Name,
-		const FShaderParametersMetadata* ParametersMetadata,
-		const ParameterStructType* ParameterStruct,
-		ERDGPassFlags Flags,
-		ExecuteLambdaType&& ExecuteLambda);
+	//////////////////////////////////////////////////////////////////////////////
+	// Passes
 
+	/** The epilogue and prologue passes are sentinels that are used to simplify graph logic around barriers
+	*  and traversal. The prologue pass is used exclusively for barriers before the graph executes, while the
+	*  epilogue pass is used for resource extraction barriers--a property that also makes it the main root of
+	*  the graph for culling purposes. The epilogue pass is added to the very end of the pass array for traversal
+	*  purposes. The prologue does not need to participate in any graph traversal behavior.
+	*/
+	FRDGPass* ProloguePass = nullptr;
+	FRDGPass* EpiloguePass = nullptr;
+
+	uint32 AsyncComputePassCount = 0;
+	uint32 RasterPassCount = 0;
+
+	/** Current scope's async compute budget. This is passed on to every pass created. */
+	EAsyncComputeBudget AsyncComputeBudgetScope = EAsyncComputeBudget::EAll_4;
+	EAsyncComputeBudget AsyncComputeBudgetState = EAsyncComputeBudget(~0u);
+
+	IF_RDG_CMDLIST_STATS(TStatId CommandListStatScope);
+	IF_RDG_CMDLIST_STATS(TStatId CommandListStatState);
+
+	IF_RDG_CPU_SCOPES(FRDGCPUScopeStacks CPUScopeStacks);
+	FRDGGPUScopeStacksByPipeline GPUScopeStacks;
+	IF_RHI_WANT_BREADCRUMB_EVENTS(FRDGBreadcrumbState* BreadcrumbState{});
+
 	static RENDERCORE_API ERDGPassFlags OverridePassFlags(const TCHAR* PassName, ERDGPassFlags Flags);
 
-	RENDERCORE_API void AddProloguePass();
-
 	FORCEINLINE FRDGPass* GetProloguePass() const
 	{
@@ -452,12 +488,269 @@
 		return Passes.Last();
 	}
 
+	FRHIRenderPassInfo GetRenderPassInfo(const FRDGPass* Pass) const;
+
+	template <typename ParameterStructType, typename ExecuteLambdaType>
+	FRDGPassRef AddPassInternal(
+		FRDGEventName&& Name,
+		const FShaderParametersMetadata* ParametersMetadata,
+		const ParameterStructType* ParameterStruct,
+		ERDGPassFlags Flags,
+		ExecuteLambdaType&& ExecuteLambda);
+
+	void MarkResourcesAsProduced(FRDGPass* Pass);
+
+	RENDERCORE_API FRDGPass* SetupEmptyPass(FRDGPass* Pass);
+	RENDERCORE_API FRDGPass* SetupParameterPass(FRDGPass* Pass);
+
+	void SetupPassInternals(FRDGPass* Pass);
+	void SetupPassResources(FRDGPass* Pass);
+	void SetupPassDependencies(FRDGPass* Pass);
+
+	void Compile();
+	void CompilePassOps(FRDGPass* Pass);
+	void ExecutePass(FRDGPass* Pass, FRHIComputeCommandList& RHICmdListPass);
+
+	void ExecutePassPrologue(FRHIComputeCommandList& RHICmdListPass, FRDGPass* Pass);
+	void ExecutePassEpilogue(FRHIComputeCommandList& RHICmdListPass, FRDGPass* Pass);
+
+	//////////////////////////////////////////////////////////////////////////////
+	// Resource Registries
+
+	/** Registry of graph objects. */
+	FRDGPassRegistry Passes;
+	FRDGTextureRegistry Textures;
+	FRDGBufferRegistry Buffers;
+	FRDGViewRegistry Views;
+	FRDGUniformBufferRegistry UniformBuffers;
+
+	struct FExtractedTexture
+	{
+		FExtractedTexture() = default;
+
+		FExtractedTexture(FRDGTexture* InTexture, TRefCountPtr<IPooledRenderTarget>* InPooledTexture)
+			: Texture(InTexture)
+			, PooledTexture(InPooledTexture)
+		{}
+
+		FRDGTexture* Texture{};
+		TRefCountPtr<IPooledRenderTarget>* PooledTexture{};
+	};
+
+	TArray<FExtractedTexture, FRDGArrayAllocator> ExtractedTextures;
+
+	struct FExtractedBuffer
+	{
+		FExtractedBuffer() = default;
+
+		FExtractedBuffer(FRDGBuffer* InBuffer, TRefCountPtr<FRDGPooledBuffer>* InPooledBuffer)
+			: Buffer(InBuffer)
+			, PooledBuffer(InPooledBuffer)
+		{}
+
+		FRDGBuffer* Buffer{};
+		TRefCountPtr<FRDGPooledBuffer>* PooledBuffer{};
+	};
+
+	TArray<FExtractedBuffer, FRDGArrayAllocator> ExtractedBuffers;
+
+	/** Tracks external resources to their registered render graph counterparts for de-duplication. */
+	Experimental::TRobinHoodHashMap<FRHITexture*, FRDGTexture*, DefaultKeyFuncs<FRHITexture*>, FRDGArrayAllocator> ExternalTextures;
+	Experimental::TRobinHoodHashMap<FRHIBuffer*,  FRDGBuffer*,  DefaultKeyFuncs<FRHIBuffer*>,  FRDGArrayAllocator> ExternalBuffers;
+
+	/** Tracks buffers that have a defered num elements callback. */
+	TArray<FRDGBuffer*, FRDGArrayAllocator> NumElementsCallbackBuffers;
+
+	//////////////////////////////////////////////////////////////////////////////
+	// Resource Collection and Allocation
+
+	IRHITransientResourceAllocator* TransientResourceAllocator = nullptr; 
+	bool bSupportsTransientTextures = false;
+	bool bSupportsTransientBuffers = false;
+
+	bool IsTransient(FRDGTextureRef Texture) const;
+	bool IsTransient(FRDGBufferRef Buffer) const;
+	bool IsTransientInternal(FRDGViewableResource* Resource, bool bFastVRAM) const;
+
+	struct FCollectResourceOp
+	{
+		enum class EOp : uint8
+		{
+			Allocate,
+			Deallocate
+		};
+
+		static FCollectResourceOp Allocate(FRDGPassHandle PassHandle, FRDGBufferHandle BufferHandle)
+		{
+			return FCollectResourceOp(PassHandle, BufferHandle.GetIndex(), ERDGViewableResourceType::Buffer, EOp::Allocate);
+		}
+
+		static FCollectResourceOp Allocate(FRDGPassHandle PassHandle, FRDGTextureHandle TextureHandle)
+		{
+			return FCollectResourceOp(PassHandle, TextureHandle.GetIndex(), ERDGViewableResourceType::Texture, EOp::Allocate);
+		}
+
+		static FCollectResourceOp Deallocate(FRDGPassHandle PassHandle, FRDGBufferHandle BufferHandle)
+		{
+			return FCollectResourceOp(PassHandle, BufferHandle.GetIndex(), ERDGViewableResourceType::Buffer, EOp::Deallocate);
+		}
+
+		static FCollectResourceOp Deallocate(FRDGPassHandle PassHandle, FRDGTextureHandle TextureHandle)
+		{
+			return FCollectResourceOp(PassHandle, TextureHandle.GetIndex(), ERDGViewableResourceType::Texture, EOp::Deallocate);
+		}
+
+		FCollectResourceOp() = default;
+		FCollectResourceOp(FRDGPassHandle InPassHandle, uint16 InResourceIndex, ERDGViewableResourceType InResourceType, EOp InOp)
+			: PassHandle(InPassHandle)
+			, ResourceIndex(InResourceIndex)
+			, ResourceType(InResourceType)
+			, Op(InOp)
+		{}
+
+		FRDGPassHandle PassHandle;
+		uint16 ResourceIndex;
+		ERDGViewableResourceType ResourceType;
+		EOp Op;
+	};
+
+	using FCollectResourceOpArray = TArray<FCollectResourceOp, FRDGArrayAllocator>;
+
+	/** A temporary context used to collect resources for allocation. */
+	struct FCollectResourceContext
+	{
+		FCollectResourceOpArray TransientResources;
+		FCollectResourceOpArray PooledTextures;
+		FCollectResourceOpArray PooledBuffers;
+		TArray<FRDGUniformBufferHandle, FRDGArrayAllocator> UniformBuffers;
+		TArray<FRDGViewHandle, FRDGArrayAllocator> Views;
+		FRDGUniformBufferBitArray UniformBufferMap;
+		FRDGViewBitArray ViewMap;
+	};
+
+	/** Tracks the latest RDG resource to own an alias of a pooled resource (multiple RDG resources can reference the same pooled resource). */
+	Experimental::TRobinHoodHashMap<FRDGPooledTexture*, FRDGTexture*, DefaultKeyFuncs<FRDGPooledTexture*>, FConcurrentLinearArrayAllocator> PooledTextureOwnershipMap;
+	Experimental::TRobinHoodHashMap<FRDGPooledBuffer*, FRDGBuffer*, DefaultKeyFuncs<FRDGPooledBuffer*>, FConcurrentLinearArrayAllocator> PooledBufferOwnershipMap;
+
+	/** Collects new resource allocations for the pass into the provided context. */
+	void CollectAllocations(FCollectResourceContext& Context, FRDGPass* Pass);
+	void CollectAllocateTexture(FCollectResourceContext& Context, FRDGPassHandle PassHandle, FRDGTexture* Texture);
+	void CollectAllocateBuffer(FCollectResourceContext& Context, FRDGPassHandle PassHandle, FRDGBuffer* Buffer);
+
+	/** Collects new resource deallocations for the pass into the provided context. */
+	void CollectDeallocations(FCollectResourceContext& Context, FRDGPass* Pass);
+	void CollectDeallocateTexture(FCollectResourceContext& Context, FRDGPassHandle PassHandle, FRDGTexture* Texture, uint32 ReferenceCount);
+	void CollectDeallocateBuffer(FCollectResourceContext& Context, FRDGPassHandle PassHandle, FRDGBuffer* Buffer, uint32 ReferenceCount);
+
+	/** Allocates resources using the provided lifetime op arrays. */
+	void AllocateTransientResources(TConstArrayView<FCollectResourceOp> Ops);
+	void AllocatePooledTextures(FRHICommandListBase& RHICmdList, TConstArrayView<FCollectResourceOp> Ops);
+	void AllocatePooledBuffers(FRHICommandListBase& RHICmdList, TConstArrayView<FCollectResourceOp> Ops);
+
+	/** Creates resources for the provided handles. */
+	void CreateViews(FRHICommandListBase& RHICmdList, TConstArrayView<FRDGViewHandle> ViewsToCreate);
+	void CreateUniformBuffers(TConstArrayView<FRDGUniformBufferHandle> UniformBuffersToCreate);
+
+	/** Allocates and returns a pooled resource for the RDG resource. Does not assign it. */
+	TRefCountPtr<IPooledRenderTarget> AllocatePooledRenderTargetRHI(FRHICommandListBase& RHICmdList, FRDGTextureRef Texture);
+	TRefCountPtr<FRDGPooledBuffer> AllocatePooledBufferRHI(FRHICommandListBase& RHICmdList, FRDGBufferRef Buffer);
+
+	/** Assigns an underlying RHI resource to an RDG resource. */
+	void SetPooledRenderTargetRHI(FRDGTexture* Texture, IPooledRenderTarget* RenderTarget);
+	void SetPooledTextureRHI(FRDGTexture* Texture, FRDGPooledTexture* PooledTexture);
+	void SetTransientTextureRHI(FRDGTexture* Texture, FRHITransientTexture* TransientTexture);
+	void SetPooledBufferRHI(FRDGBuffer* Buffer, FRDGPooledBuffer* PooledBuffer);
+	void SetTransientBufferRHI(FRDGBuffer* Buffer, FRHITransientBuffer* TransientBuffer);
+
+	/** Initializes various view types. Assumes that the underlying RHI viewable resource type is assigned. */
+	void InitViewRHI(FRHICommandListBase& RHICmdList, FRDGView* View);
+	void InitBufferViewRHI(FRHICommandListBase& RHICmdList, FRDGBufferSRV* SRV);
+	void InitBufferViewRHI(FRHICommandListBase& RHICmdList, FRDGBufferUAV* UAV);
+	void InitTextureViewRHI(FRHICommandListBase& RHICmdList, FRDGTextureSRV* SRV);
+	void InitTextureViewRHI(FRHICommandListBase& RHICmdList, FRDGTextureUAV* UAV);
+
+	//////////////////////////////////////////////////////////////////////////////
+	// Resource Transitions and State Tracking
+
+	/** Map of barrier batches begun from more than one pipe. */
+	TMap<FRDGBarrierBatchBeginId, FRDGBarrierBatchBegin*, FRDGSetAllocator> BarrierBatchMap;
+
+	/** Tracks the final access used on resources in order to call SetTrackedAccess. */
+	TArray<FRHITrackedAccessInfo, FRDGArrayAllocator> EpilogueResourceAccesses;
+
+	/** Array of all pooled references held during execution. */
+	TArray<TRefCountPtr<IPooledRenderTarget>, FRDGArrayAllocator> ActivePooledTextures;
+	TArray<TRefCountPtr<FRDGPooledBuffer>, FRDGArrayAllocator> ActivePooledBuffers;
+
+	/** Set of all active barrier batch begin instances; used to create transitions. */
+	FRDGTransitionCreateQueue TransitionCreateQueue;
+
+	/** Texture state used for intermediate operations. Held here to avoid re-allocating. */
+	FRDGTextureSubresourceState ScratchTextureState;
+
+	/** Subresource state representing the graph prologue. Used for immediate mode. */
+	FRDGSubresourceState PrologueSubresourceState;
+
+	void CompilePassBarriers();
+	void CollectPassBarriers();
+	void CollectPassBarriers(FRDGPassHandle PassHandle);
+	void CreatePassBarriers();
+	void FinalizeResources();
+
+	void AddFirstTextureTransition(FRDGTextureRef Texture);
+	void AddFirstBufferTransition(FRDGBufferRef Buffer);
+
+	void AddLastTextureTransition(FRDGTextureRef Texture);
+	void AddLastBufferTransition(FRDGBufferRef Buffer);
+
+	template <typename FilterSubresourceLambdaType>
+	void AddTextureTransition(
+		FRDGTextureRef Texture,
+		FRDGTextureSubresourceState& StateBefore,
+		FRDGTextureSubresourceState& StateAfter,
+		FilterSubresourceLambdaType&& FilterSubresourceLambda);
+
+	void AddTextureTransition(
+		FRDGTextureRef Texture,
+		FRDGTextureSubresourceState& StateBefore,
+		FRDGTextureSubresourceState& StateAfter)
+	{
+		AddTextureTransition(Texture, StateBefore, StateAfter, [](FRDGSubresourceState*, int32) { return true; });
+	}
+
+	template <typename FilterSubresourceLambdaType>
+	void AddBufferTransition(
+		FRDGBufferRef Buffer,
+		FRDGSubresourceState*& StateBefore,
+		FRDGSubresourceState* StateAfter,
+		FilterSubresourceLambdaType&& FilterSubresourceLambda);
+
+	void AddBufferTransition(
+		FRDGBufferRef Buffer,
+		FRDGSubresourceState*& StateBefore,
+		FRDGSubresourceState* StateAfter)
+	{
+		AddBufferTransition(Buffer, StateBefore, StateAfter, [](FRDGSubresourceState*) { return true; });
+	}
+
+	void AddTransition(
+		FRDGViewableResource* Resource,
+		FRDGSubresourceState StateBefore,
+		FRDGSubresourceState StateAfter,
+		FRDGTransitionInfo TransitionInfo);
+
+	void AddAliasingTransition(
+		FRDGPassHandle BeginPassHandle,
+		FRDGPassHandle EndPassHandle,
+		FRDGViewableResource* Resource,
+		const FRHITransientAliasingInfo& Info);
+
 	/** Prologue and Epilogue barrier passes are used to plan transitions around RHI render pass merging,
-	 *  as it is illegal to issue a barrier during a render pass. If passes [A, B, C] are merged together,
-	 *  'A' becomes 'B's prologue pass and 'C' becomes 'A's epilogue pass. This way, any transitions that
-	 *  need to happen before the merged pass (i.e. in the prologue) are done in A. Any transitions after
-	 *  the render pass merge are done in C.
-	 */
+	*  as it is illegal to issue a barrier during a render pass. If passes [A, B, C] are merged together,
+	*  'A' becomes 'B's prologue pass and 'C' becomes 'A's epilogue pass. This way, any transitions that
+	*  need to happen before the merged pass (i.e. in the prologue) are done in A. Any transitions after
+	*  the render pass merge are done in C.
+	*/
 	FRDGPassHandle GetEpilogueBarrierPassHandle(FRDGPassHandle Handle)
 	{
 		return Passes[Handle]->EpilogueBarrierPass;
@@ -482,14 +775,14 @@
 	void AddToPrologueBarriersToEnd(FRDGPassHandle Handle, FRDGBarrierBatchBegin& BarriersToBegin)
 	{
 		FRDGPass* Pass = GetPrologueBarrierPass(Handle);
-		Pass->GetPrologueBarriersToEnd(Allocator).AddDependency(&BarriersToBegin);
+		Pass->GetPrologueBarriersToEnd(Allocators.Transition).AddDependency(&BarriersToBegin);
 	}
 
 	/** Ends the barrier batch in the epilogue of the provided pass. */
 	void AddToEpilogueBarriersToEnd(FRDGPassHandle Handle, FRDGBarrierBatchBegin& BarriersToBegin)
 	{
 		FRDGPass* Pass = GetEpilogueBarrierPass(Handle);
-		Pass->GetEpilogueBarriersToEnd(Allocator).AddDependency(&BarriersToBegin);
+		Pass->GetEpilogueBarriersToEnd(Allocators.Transition).AddDependency(&BarriersToBegin);
 	}
 
 	/** Utility function to add an immediate barrier dependency in the prologue of the provided pass. */
@@ -497,9 +790,9 @@
 	void AddToPrologueBarriers(FRDGPassHandle PassHandle, FunctionType Function)
 	{
 		FRDGPass* Pass = GetPrologueBarrierPass(PassHandle);
-		FRDGBarrierBatchBegin& BarriersToBegin = Pass->GetPrologueBarriersToBegin(Allocator, TransitionCreateQueue);
+		FRDGBarrierBatchBegin& BarriersToBegin = Pass->GetPrologueBarriersToBegin(Allocators.Transition, TransitionCreateQueue);
 		Function(BarriersToBegin);
-		Pass->GetPrologueBarriersToEnd(Allocator).AddDependency(&BarriersToBegin);
+		Pass->GetPrologueBarriersToEnd(Allocators.Transition).AddDependency(&BarriersToBegin);
 	}
 
 	/** Utility function to add an immediate barrier dependency in the epilogue of the provided pass. */
@@ -507,107 +800,123 @@
 	void AddToEpilogueBarriers(FRDGPassHandle PassHandle, FunctionType Function)
 	{
 		FRDGPass* Pass = GetEpilogueBarrierPass(PassHandle);
-		FRDGBarrierBatchBegin& BarriersToBegin = Pass->GetEpilogueBarriersToBeginFor(Allocator, TransitionCreateQueue, Pass->GetPipeline());
+		FRDGBarrierBatchBegin& BarriersToBegin = Pass->GetEpilogueBarriersToBeginFor(Allocators.Transition, TransitionCreateQueue, Pass->GetPipeline());
 		Function(BarriersToBegin);
-		Pass->GetEpilogueBarriersToEnd(Allocator).AddDependency(&BarriersToBegin);
+		Pass->GetEpilogueBarriersToEnd(Allocators.Transition).AddDependency(&BarriersToBegin);
 	}
 
-#if WITH_MGPU
-	RENDERCORE_API void ForceCopyCrossGPU();
-#endif
+	FRDGSubresourceState* AllocSubresource(const FRDGSubresourceState& Other);
+	FRDGSubresourceState* AllocSubresource();
 
-	/** Registry of graph objects. */
-	FRDGPassRegistry Passes;
-	FRDGTextureRegistry Textures;
-	FRDGBufferRegistry Buffers;
-	FRDGViewRegistry Views;
-	FRDGUniformBufferRegistry UniformBuffers;
+	//////////////////////////////////////////////////////////////////////////////
+	// Async Setup Queue
 
-	/** Uniform buffers which were used in a pass. */
-	TArray<FRDGUniformBufferHandle, FRDGArrayAllocator> UniformBuffersToCreate;
+	struct FAsyncSetupOp
+	{
+		enum class EType
+		{
+			SetupPassResources,
+			CullRootBuffer,
+			CullRootTexture
+		};
 
-	/** Tracks external resources to their registered render graph counterparts for de-duplication. */
-	TSortedMap<FRHITexture*, FRDGTexture*, FRDGArrayAllocator> ExternalTextures;
-	TSortedMap<FRHIBuffer*, FRDGBuffer*, FRDGArrayAllocator> ExternalBuffers;
+		static FAsyncSetupOp SetupPassResources(FRDGPass* Pass)
+		{
+			FAsyncSetupOp Op;
+			Op.Type = EType::SetupPassResources;
+			Op.Pass = Pass;
+			return Op;
+		}
 
-	/** Tracks the latest RDG resource to own an alias of a pooled resource (multiple RDG resources can reference the same pooled resource). */
-	TMap<FRDGPooledTexture*, FRDGTexture*, FRDGSetAllocator> PooledTextureOwnershipMap;
-	TMap<FRDGPooledBuffer*, FRDGBuffer*, FRDGSetAllocator> PooledBufferOwnershipMap;
+		static FAsyncSetupOp CullRootBuffer(FRDGBuffer* Buffer)
+		{
+			FAsyncSetupOp Op;
+			Op.Type = EType::CullRootBuffer;
+			Op.Buffer = Buffer;
+			return Op;
+		}
 
-	/** Array of all pooled references held during execution. */
-	TArray<TRefCountPtr<IPooledRenderTarget>, FRDGArrayAllocator> ActivePooledTextures;
-	TArray<TRefCountPtr<FRDGPooledBuffer>, FRDGArrayAllocator> ActivePooledBuffers;
+		static FAsyncSetupOp CullRootTexture(FRDGTexture* Texture)
+		{
+			FAsyncSetupOp Op;
+			Op.Type = EType::CullRootTexture;
+			Op.Texture = Texture;
+			return Op;
+		}
 
-	/** Map of barrier batches begun from more than one pipe. */
-	TMap<FRDGBarrierBatchBeginId, FRDGBarrierBatchBegin*, FRDGSetAllocator> BarrierBatchMap;
+		EType Type;
 
-	/** Set of all active barrier batch begin instances; used to create transitions. */
-	FRDGTransitionCreateQueue TransitionCreateQueue;
+		union
+		{
+			FRDGPass*    Pass;
+			FRDGBuffer*  Buffer;
+			FRDGTexture* Texture;
+		};
+	};
 
-	template <typename LambdaType>
-	UE::Tasks::FTask LaunchCompileTask(const TCHAR* Name, bool bCondition, LambdaType&& Lambda);
+	struct FAsyncSetupQueue
+	{
+		void Push(FAsyncSetupOp Op)
+		{
+			UE::TScopeLock Lock(Mutex);
+			Ops.Emplace(Op);
+		}
 
-	UE::Tasks::FPipe CompilePipe;
-
-	class FPassQueue
-	{
-	public:
-		void Push(FRDGPass* Pass);
-
-		template <typename LambdaType>
-		void Flush(UE::Tasks::FPipe& Pipe, const TCHAR* Name, LambdaType&& Lambda);
-
-		template <typename LambdaType>
-		void Flush(const TCHAR* Name, LambdaType&& Lambda);
-
-	private:
 		UE::FMutex Mutex;
-		TArray<FRDGPass*, FRDGArrayAllocator> Queue;
+		TArray<FAsyncSetupOp, FRDGArrayAllocator> Ops;
 		UE::Tasks::FTask LastTask;
-	};
+		UE::Tasks::FPipe Pipe{ TEXT("FRDGBuilder::AsyncSetupQueue") };
 
-	FPassQueue SetupPassQueue;
+	} AsyncSetupQueue;
 
-	TArray<FRDGPassHandle, FRDGArrayAllocator> CullPassStack;
+	void LaunchAsyncSetupQueueTask();
+	void ProcessAsyncSetupQueue();
 
-	/** The epilogue and prologue passes are sentinels that are used to simplify graph logic around barriers
-	 *  and traversal. The prologue pass is used exclusively for barriers before the graph executes, while the
-	 *  epilogue pass is used for resource extraction barriers--a property that also makes it the main root of
-	 *  the graph for culling purposes. The epilogue pass is added to the very end of the pass array for traversal
-	 *  purposes. The prologue does not need to participate in any graph traversal behavior.
-	 */
-	FRDGPass* ProloguePass = nullptr;
-	FRDGPass* EpiloguePass = nullptr;
+	//////////////////////////////////////////////////////////////////////////////
+	// Culling
 
-	struct FExtractedTexture
+	TArray<FRDGPass*, FRDGArrayAllocator> CullPassStack;
+
+	bool AddCullingDependency(FRDGProducerStatesByPipeline& LastProducers, const FRDGProducerState& NextState, ERHIPipeline NextPipeline);
+	void AddCullRootBuffer(FRDGBuffer* Buffer);
+	void AddCullRootTexture(FRDGTexture* Texture);
+	void AddLastProducersToCullStack(const FRDGProducerStatesByPipeline& LastProducers);
+	void FlushCullStack();
+
+	//////////////////////////////////////////////////////////////////////////////
+	// Parallel Setup
+
+	struct
 	{
-		FExtractedTexture() = default;
+		/** Array of all tasks for variants of AddSetupTask. */
+		TArray<UE::Tasks::FTask, FRDGArrayAllocator> Tasks;
 
-		FExtractedTexture(FRDGTexture* InTexture, TRefCountPtr<IPooledRenderTarget>* InPooledTexture)
-			: Texture(InTexture)
-			, PooledTexture(InPooledTexture)
-		{}
+		/** Array of all command lists to submit for AddCommandListSetupTask. */
+		UE::FMutex CommandListsMutex;
+		TArray<FRHICommandListImmediate::FQueuedCommandList, FConcurrentLinearArrayAllocator> CommandLists;
 
-		FRDGTexture* Texture{};
-		TRefCountPtr<IPooledRenderTarget>* PooledTexture{};
-	};
+		bool bEnabled = false;
 
-	TArray<FExtractedTexture, FRDGArrayAllocator> ExtractedTextures;
+	} ParallelSetup;
 
-	struct FExtractedBuffer
+	void WaitForParallelSetupTasks();
+
+	/////////////////////////////////////////////////////////////////////////////
+	// Parallel Execution
+
+	struct
 	{
-		FExtractedBuffer() = default;
+		TArray<FParallelPassSet, FRDGArrayAllocator> ParallelPassSets;
+		TArray<UE::Tasks::FTask, FRDGArrayAllocator> Tasks;
+		TOptional<UE::Tasks::FTaskEvent> DispatchTaskEvent;
+		bool bEnabled = false;
 
-		FExtractedBuffer(FRDGBuffer* InBuffer, TRefCountPtr<FRDGPooledBuffer>* InPooledBuffer)
-			: Buffer(InBuffer)
-			, PooledBuffer(InPooledBuffer)
-		{}
+	} ParallelExecute;
 
-		FRDGBuffer* Buffer{};
-		TRefCountPtr<FRDGPooledBuffer>* PooledBuffer{};
-	};
+	void SetupParallelExecute();
 
-	TArray<FExtractedBuffer, FRDGArrayAllocator> ExtractedBuffers;
+	/////////////////////////////////////////////////////////////////////////////
+	// Buffer Uploads
 
 	struct FUploadedBuffer
 	{
@@ -665,49 +974,28 @@
 
 	TArray<FUploadedBuffer, FRDGArrayAllocator> UploadedBuffers;
 
-	TArray<FParallelPassSet, FRDGArrayAllocator> ParallelPassSets;
+	void SubmitBufferUploads(FRHICommandListBase& InRHICmdList, UE::Tasks::FTaskEvent* AllocateUploadBuffersTask = nullptr);
 
-	/** Array of all active parallel execute tasks. */
-	TArray<UE::Tasks::FTask, FRDGArrayAllocator> ParallelExecuteEvents;
+	/////////////////////////////////////////////////////////////////////////////
+	// External Access Queue
 
-	/** Array of all task events for variants of AddSetupTask. */
-	TArray<UE::Tasks::FTask, FRDGArrayAllocator> ParallelSetupEvents;
-
-	/** Array of all command lists to submit for AddCommandListSetupTask. */
-	TArray<FRHICommandListImmediate::FQueuedCommandList, FRDGArrayAllocator> ParallelSetupCommandLists;
-
-	/** Tracks the final access used on resources in order to call SetTrackedAccess. */
-	TArray<FRHITrackedAccessInfo, FRDGArrayAllocator> EpilogueResourceAccesses;
-
 	/** Contains resources queued for either access mode change passes. */
 	TArray<FRDGViewableResource*, FRDGArrayAllocator> AccessModeQueue;
 	TSet<FRDGViewableResource*, DefaultKeyFuncs<FRDGViewableResource*>, FRDGSetAllocator> ExternalAccessResources;
 
-	/** Texture state used for intermediate operations. Held here to avoid re-allocating. */
-	FRDGTextureSubresourceStateIndirect ScratchTextureState;
+	RENDERCORE_API void FlushAccessModeQueue();
 
-	/** Current scope's async compute budget. This is passed on to every pass created. */
-	EAsyncComputeBudget AsyncComputeBudgetScope = EAsyncComputeBudget::EAll_4;
-	EAsyncComputeBudget AsyncComputeBudgetState = EAsyncComputeBudget(~0u);
+	/////////////////////////////////////////////////////////////////////////////
+	// Resource Deletion Flushing
 
-	/** Command list handle created by the parallel buffer upload task. */
-	FRHICommandList* RHICmdListBufferUploads = nullptr;
-
-	IF_RDG_CPU_SCOPES(FRDGCPUScopeStacks CPUScopeStacks);
-	FRDGGPUScopeStacksByPipeline GPUScopeStacks;
-	IF_RHI_WANT_BREADCRUMB_EVENTS(FRDGBreadcrumbState* BreadcrumbState{});
-
-	IF_RDG_ENABLE_TRACE(FRDGTrace Trace);
-
 	bool bFlushResourcesRHI = false;
-	bool bParallelExecuteEnabled = false;
-	bool bParallelSetupEnabled = false;
-	bool bFinalEventScopeActive = false;
+	FRHICommandListScopedExtendResourceLifetime ExtendResourceLifetimeScope;
 
-#if RDG_ENABLE_DEBUG
-	FRDGUserValidation UserValidation;
-	FRDGBarrierValidation BarrierValidation;
-#endif
+	void BeginFlushResourcesRHI();
+	void EndFlushResourcesRHI();
 
+	/////////////////////////////////////////////////////////////////////////////
+	// Clobber, Visualize, and DumpGPU tools.
+
 	/** Tracks stack counters of auxiliary passes to avoid calling them recursively. */
 	struct FAuxiliaryPass
@@ -726,120 +1014,15 @@
 
 	} AuxiliaryPasses;
 
-#if WITH_MGPU
-	/** Copy all cross GPU external resources (not marked MultiGPUGraphIgnore) at the end of execution (bad for perf, but useful for debugging). */
-	bool bForceCopyCrossGPU = false;
-#endif
+	void SetupAuxiliaryPasses(FRDGPass* Pass);
 
-	uint32 AsyncComputePassCount = 0;
-	uint32 RasterPassCount = 0;
-
-	IF_RDG_CMDLIST_STATS(TStatId CommandListStatScope);
-	IF_RDG_CMDLIST_STATS(TStatId CommandListStatState);
-
-	IRHITransientResourceAllocator* TransientResourceAllocator = nullptr; 
-	bool bSupportsTransientTextures = false;
-	bool bSupportsTransientBuffers = false;
-
-	FRHICommandListScopedExtendResourceLifetime ExtendResourceLifetimeScope;
-
-	RENDERCORE_API void MarkResourcesAsProduced(FRDGPass* Pass);
-
-	RENDERCORE_API void Compile();
-	RENDERCORE_API void Clear();
-
-	RENDERCORE_API void SetRHI(FRDGTexture* Texture, IPooledRenderTarget* RenderTarget, FRDGPassHandle PassHandle);
-	RENDERCORE_API void SetRHI(FRDGTexture* Texture, FRDGPooledTexture* PooledTexture, FRDGPassHandle PassHandle);
-	RENDERCORE_API void SetRHI(FRDGTexture* Texture, FRHITransientTexture* TransientTexture, FRDGPassHandle PassHandle);
-	RENDERCORE_API void SetRHI(FRDGBuffer* Buffer, FRDGPooledBuffer* PooledBuffer, FRDGPassHandle PassHandle);
-	RENDERCORE_API void SetRHI(FRDGBuffer* Buffer, FRHITransientBuffer* TransientBuffer, FRDGPassHandle PassHandle);
-
-	RENDERCORE_API void BeginResourcesRHI(FRDGPass* ResourcePass, FRDGPassHandle ExecutePassHandle);
-	RENDERCORE_API void BeginResourceRHI(FRDGPassHandle, FRDGTexture* Texture);
-	RENDERCORE_API void BeginResourceRHI(FRDGPassHandle, FRDGBuffer* Buffer);
-
-	RENDERCORE_API void EndResourcesRHI(FRDGPass* ResourcePass, FRDGPassHandle ExecutePassHandle);
-	RENDERCORE_API void EndResourceRHI(FRDGPassHandle, FRDGTexture* Texture, uint32 ReferenceCount);
-	RENDERCORE_API void EndResourceRHI(FRDGPassHandle, FRDGBuffer* Buffer, uint32 ReferenceCount);
-
-	RENDERCORE_API void InitRHI(FRDGView* View);
-	RENDERCORE_API void InitRHI(FRDGBufferSRV* SRV);
-	RENDERCORE_API void InitRHI(FRDGBufferUAV* UAV);
-	RENDERCORE_API void InitRHI(FRDGTextureSRV* SRV);
-	RENDERCORE_API void InitRHI(FRDGTextureUAV* UAV);
-
-	RENDERCORE_API void SetupParallelExecute();
-	RENDERCORE_API void DispatchParallelExecute();
-
-	RENDERCORE_API void PrepareBufferUploads();
-	RENDERCORE_API UE::Tasks::FTask SubmitBufferUploads();
-	RENDERCORE_API void BeginFlushResourcesRHI();
-	RENDERCORE_API void EndFlushResourcesRHI();
-
-	RENDERCORE_API void FlushAccessModeQueue();
-
-	RENDERCORE_API FRDGPass* SetupEmptyPass(FRDGPass* Pass);
-	RENDERCORE_API FRDGPass* SetupParameterPass(FRDGPass* Pass);
-
-	RENDERCORE_API void SetupPassInternals(FRDGPass* Pass);
-	RENDERCORE_API void SetupPassResources(FRDGPass* Pass);
-	RENDERCORE_API void SetupAuxiliaryPasses(FRDGPass* Pass);
-	RENDERCORE_API void SetupPassDependencies(FRDGPass* Pass);
-
-	RENDERCORE_API void CompilePassOps(FRDGPass* Pass);
-	RENDERCORE_API void ExecutePass(FRDGPass* Pass, FRHIComputeCommandList& RHICmdListPass);
-
-	RENDERCORE_API void ExecutePassPrologue(FRHIComputeCommandList& RHICmdListPass, FRDGPass* Pass);
-	RENDERCORE_API void ExecutePassEpilogue(FRHIComputeCommandList& RHICmdListPass, FRDGPass* Pass);
-
-	RENDERCORE_API void CompilePassBarriers();
-	RENDERCORE_API void CollectPassBarriers(FRDGPass* Pass);
-	RENDERCORE_API void CreatePassBarriers(TFunctionRef<void()> PreWork);
-
-	RENDERCORE_API UE::Tasks::FTask CreateUniformBuffers();
-
-	RENDERCORE_API void AddCullingDependency(FRDGProducerStatesByPipeline& LastProducers, const FRDGProducerState& NextState, ERHIPipeline NextPipeline);
-
-	RENDERCORE_API void AddEpilogueTransition(FRDGTextureRef Texture);
-	RENDERCORE_API void AddEpilogueTransition(FRDGBufferRef Buffer);
-
-	RENDERCORE_API void AddTransition(
-		FRDGPassHandle PassHandle,
-		FRDGTextureRef Texture,
-		FRDGTextureSubresourceStateIndirect& StateAfter);
-
-	RENDERCORE_API void AddTransition(
-		FRDGPassHandle PassHandle,
-		FRDGBufferRef Buffer,
-		FRDGSubresourceState StateAfter);
-
-	RENDERCORE_API void AddTransition(
-		FRDGViewableResource* Resource,
-		FRDGSubresourceState StateBefore,
-		FRDGSubresourceState StateAfter,
-		const FRHITransitionInfo& TransitionInfo);
-
-	RENDERCORE_API void AddAliasingTransition(
-		FRDGPassHandle BeginPassHandle,
-		FRDGPassHandle EndPassHandle,
-		FRDGViewableResource* Resource,
-		const FRHITransientAliasingInfo& Info);
-
-	RENDERCORE_API bool IsTransient(FRDGTextureRef Texture) const;
-	RENDERCORE_API bool IsTransient(FRDGBufferRef Buffer) const;
-	RENDERCORE_API bool IsTransientInternal(FRDGViewableResource* Resource, bool bFastVRAM) const;
-
-	RENDERCORE_API FRHIRenderPassInfo GetRenderPassInfo(const FRDGPass* Pass) const;
-
-	RENDERCORE_API FRDGSubresourceState* AllocSubresource(const FRDGSubresourceState& Other);
-
 #if RDG_DUMP_RESOURCES
-	RENDERCORE_API void DumpNewGraphBuilder();
-	RENDERCORE_API void DumpResourcePassOutputs(const FRDGPass* Pass);
+	void DumpNewGraphBuilder();
+	void DumpResourcePassOutputs(const FRDGPass* Pass);
 
 #if RDG_DUMP_RESOURCES_AT_EACH_DRAW
-	RENDERCORE_API void BeginPassDump(const FRDGPass* Pass);
-	RENDERCORE_API void EndPassDump(const FRDGPass* Pass);
+	void BeginPassDump(const FRDGPass* Pass);
+	void EndPassDump(const FRDGPass* Pass);
 #endif
 #endif
 
@@ -848,6 +1031,28 @@
 	RENDERCORE_API void ClobberPassOutputs(const FRDGPass* Pass);
 #endif
 
+	/////////////////////////////////////////////////////////////////////////////
+	// Multi-GPU
+
+#if WITH_MGPU
+	/** Copy all cross GPU external resources (not marked MultiGPUGraphIgnore) at the end of execution (bad for perf, but useful for debugging). */
+	bool bForceCopyCrossGPU = false;
+
+	void ForceCopyCrossGPU();
+#endif
+
+	/////////////////////////////////////////////////////////////////////////////
+	// Validation and Tracing
+
+	IF_RDG_ENABLE_TRACE(FRDGTrace Trace);
+
+#if RDG_ENABLE_DEBUG
+	FRDGUserValidation UserValidation;
+	FRDGBarrierValidation BarrierValidation;
+#endif
+
+	/////////////////////////////////////////////////////////////////////////////
+
 	friend FRDGTrace;
 	friend DynamicRenderScaling::FRDGScope;
 	friend FRDGEventScopeGuard;
diff --git a/Engine/Source/Runtime/RenderCore/Public/RenderGraphBuilder.inl b/Engine/Source/Runtime/RenderCore/Public/RenderGraphBuilder.inl
--- a/Engine/Source/Runtime/RenderCore/Public/RenderGraphBuilder.inl
+++ b/Engine/Source/Runtime/RenderCore/Public/RenderGraphBuilder.inl
@@ -57,7 +57,7 @@ inline FRDGTextureRef FRDGBuilder::CreateTexture(
 	OverrideDesc.Extent.Y = FMath::Clamp<int32>(OverrideDesc.Extent.Y, 1, GetMax2DTextureDimension());
 
 	IF_RDG_ENABLE_DEBUG(UserValidation.ValidateCreateTexture(OverrideDesc, Name, Flags));
-	FRDGTextureRef Texture = Textures.Allocate(Allocator, Name, OverrideDesc, Flags);
+	FRDGTextureRef Texture = Textures.Allocate(Allocators.Root, Name, OverrideDesc, Flags);
 	IF_RDG_ENABLE_DEBUG(UserValidation.ValidateCreateTexture(Texture));
 	IF_RDG_ENABLE_TRACE(Trace.AddResource(Texture));
 	return Texture;
@@ -72,7 +72,7 @@ inline FRDGBufferRef FRDGBuilder::CreateBuffer(
 	FRDGBufferDesc OverrideDesc = Desc;
 
 	IF_RDG_ENABLE_DEBUG(UserValidation.ValidateCreateBuffer(Desc, Name, Flags));
-	FRDGBufferRef Buffer = Buffers.Allocate(Allocator, Name, OverrideDesc, Flags);
+	FRDGBufferRef Buffer = Buffers.Allocate(Allocators.Root, Name, OverrideDesc, Flags);
 	IF_RDG_ENABLE_DEBUG(UserValidation.ValidateCreateBuffer(Buffer));
 	IF_RDG_ENABLE_TRACE(Trace.AddResource(Buffer));
 	return Buffer;
@@ -88,7 +88,8 @@ inline FRDGBufferRef FRDGBuilder::CreateBuffer(
 	FRDGBufferDesc OverrideDesc = Desc;
 
 	IF_RDG_ENABLE_DEBUG(UserValidation.ValidateCreateBuffer(Desc, Name, Flags));
-	FRDGBufferRef Buffer = Buffers.Allocate(Allocator, Name, OverrideDesc, Flags, MoveTemp(NumElementsCallback));
+	FRDGBufferRef Buffer = Buffers.Allocate(Allocators.Root, Name, OverrideDesc, Flags, MoveTemp(NumElementsCallback));
+	NumElementsCallbackBuffers.Emplace(Buffer);
 	IF_RDG_ENABLE_DEBUG(UserValidation.ValidateCreateBuffer(Buffer));
 	IF_RDG_ENABLE_TRACE(Trace.AddResource(Buffer));
 	return Buffer;
@@ -97,7 +98,7 @@ inline FRDGBufferRef FRDGBuilder::CreateBuffer(
 inline FRDGTextureSRVRef FRDGBuilder::CreateSRV(const FRDGTextureSRVDesc& Desc)
 {
 	IF_RDG_ENABLE_DEBUG(UserValidation.ValidateCreateSRV(Desc));
-	FRDGTextureSRVRef SRV = Views.Allocate<FRDGTextureSRV>(Allocator, Desc.Texture->Name, Desc);
+	FRDGTextureSRVRef SRV = Views.Allocate<FRDGTextureSRV>(Allocators.Root, Desc.Texture->Name, Desc);
 	IF_RDG_ENABLE_DEBUG(UserValidation.ValidateCreateSRV(SRV));
 	return SRV;
 }
@@ -105,7 +106,7 @@ inline FRDGTextureSRVRef FRDGBuilder::CreateSRV(const FRDGTextureSRVDesc& Desc)
 inline FRDGBufferSRVRef FRDGBuilder::CreateSRV(const FRDGBufferSRVDesc& Desc)
 {
 	IF_RDG_ENABLE_DEBUG(UserValidation.ValidateCreateSRV(Desc));
-	FRDGBufferSRVRef SRV = Views.Allocate<FRDGBufferSRV>(Allocator, Desc.Buffer->Name, Desc);
+	FRDGBufferSRVRef SRV = Views.Allocate<FRDGBufferSRV>(Allocators.Root, Desc.Buffer->Name, Desc);
 	IF_RDG_ENABLE_DEBUG(UserValidation.ValidateCreateSRV(SRV));
 	return SRV;
 }
@@ -113,7 +114,7 @@ inline FRDGBufferSRVRef FRDGBuilder::CreateSRV(const FRDGBufferSRVDesc& Desc)
 inline FRDGTextureUAVRef FRDGBuilder::CreateUAV(const FRDGTextureUAVDesc& Desc, ERDGUnorderedAccessViewFlags InFlags)
 {
 	IF_RDG_ENABLE_DEBUG(UserValidation.ValidateCreateUAV(Desc));
-	FRDGTextureUAVRef UAV = Views.Allocate<FRDGTextureUAV>(Allocator, Desc.Texture->Name, Desc, InFlags);
+	FRDGTextureUAVRef UAV = Views.Allocate<FRDGTextureUAV>(Allocators.Root, Desc.Texture->Name, Desc, InFlags);
 	IF_RDG_ENABLE_DEBUG(UserValidation.ValidateCreateUAV(UAV));
 	return UAV;
 }
@@ -121,50 +122,50 @@ inline FRDGTextureUAVRef FRDGBuilder::CreateUAV(const FRDGTextureUAVDesc& Desc,
 inline FRDGBufferUAVRef FRDGBuilder::CreateUAV(const FRDGBufferUAVDesc& Desc, ERDGUnorderedAccessViewFlags InFlags)
 {
 	IF_RDG_ENABLE_DEBUG(UserValidation.ValidateCreateUAV(Desc));
-	FRDGBufferUAVRef UAV = Views.Allocate<FRDGBufferUAV>(Allocator, Desc.Buffer->Name, Desc, InFlags);
+	FRDGBufferUAVRef UAV = Views.Allocate<FRDGBufferUAV>(Allocators.Root, Desc.Buffer->Name, Desc, InFlags);
 	IF_RDG_ENABLE_DEBUG(UserValidation.ValidateCreateUAV(UAV));
 	return UAV;
 }
 
 FORCEINLINE void* FRDGBuilder::Alloc(uint64 SizeInBytes, uint32 AlignInBytes)
 {
-	return Allocator.Alloc(SizeInBytes, AlignInBytes);
+	return Allocators.Root.Alloc(SizeInBytes, AlignInBytes);
 }
 
 template <typename PODType>
 FORCEINLINE PODType* FRDGBuilder::AllocPOD()
 {
-	return Allocator.AllocUninitialized<PODType>();
+	return Allocators.Root.AllocUninitialized<PODType>();
 }
 
 template <typename PODType>
 FORCEINLINE PODType* FRDGBuilder::AllocPODArray(uint32 Count)
 {
-	return Allocator.AllocUninitialized<PODType>(Count);
+	return Allocators.Root.AllocUninitialized<PODType>(Count);
 }
 
 template <typename ObjectType, typename... TArgs>
 FORCEINLINE ObjectType* FRDGBuilder::AllocObject(TArgs&&... Args)
 {
-	return Allocator.Alloc<ObjectType>(Forward<TArgs&&>(Args)...);
+	return Allocators.Root.Alloc<ObjectType>(Forward<TArgs&&>(Args)...);
 }
 
 template <typename ObjectType>
 FORCEINLINE TArray<ObjectType, SceneRenderingAllocator>& FRDGBuilder::AllocArray()
 {
-	return *Allocator.Alloc<TArray<ObjectType, SceneRenderingAllocator>>();
+	return *Allocators.Root.Alloc<TArray<ObjectType, SceneRenderingAllocator>>();
 }
 
 template <typename ParameterStructType>
 FORCEINLINE ParameterStructType* FRDGBuilder::AllocParameters()
 {
-	return Allocator.Alloc<ParameterStructType>();
+	return Allocators.Root.Alloc<ParameterStructType>();
 }
 
 template <typename ParameterStructType>
 FORCEINLINE ParameterStructType* FRDGBuilder::AllocParameters(ParameterStructType* StructToCopy)
 {
-	ParameterStructType* Struct = Allocator.Alloc<ParameterStructType>();
+	ParameterStructType* Struct = Allocators.Root.Alloc<ParameterStructType>();
 	*Struct = *StructToCopy;
 	return Struct;
 }
@@ -180,7 +181,7 @@ TStridedView<BaseParameterStructType> FRDGBuilder::AllocParameters(const FShader
 {
 	// NOTE: Contents are always zero! This might differ if shader parameters have a non-zero default initializer.
 	const int32 Stride = ParametersMetadata->GetSize();
-	BaseParameterStructType* Contents = reinterpret_cast<BaseParameterStructType*>(Allocator.Alloc(Stride * NumStructs, SHADER_PARAMETER_STRUCT_ALIGNMENT));
+	BaseParameterStructType* Contents = reinterpret_cast<BaseParameterStructType*>(Allocators.Root.Alloc(Stride * NumStructs, SHADER_PARAMETER_STRUCT_ALIGNMENT));
 	FMemory::Memset(Contents, 0, Stride * NumStructs);
 	TStridedView<BaseParameterStructType> ParameterArray(Stride, Contents, NumStructs);
 
@@ -211,7 +212,12 @@ TStridedView<BaseParameterStructType> FRDGBuilder::AllocParameters(const FShader
 
 FORCEINLINE FRDGSubresourceState* FRDGBuilder::AllocSubresource(const FRDGSubresourceState& Other)
 {
-	return Allocator.AllocNoDestruct<FRDGSubresourceState>(Other);
+	return Allocators.Transition.AllocNoDestruct<FRDGSubresourceState>(Other);
+}
+
+FORCEINLINE FRDGSubresourceState* FRDGBuilder::AllocSubresource()
+{
+	return Allocators.Transition.AllocNoDestruct<FRDGSubresourceState>();
 }
 
 template <typename ParameterStructType>
@@ -219,7 +225,7 @@ TRDGUniformBufferRef<ParameterStructType> FRDGBuilder::CreateUniformBuffer(const
 {
 #if !USE_NULL_RHI
 	IF_RDG_ENABLE_DEBUG(UserValidation.ValidateCreateUniformBuffer(ParameterStruct, ParameterStructType::FTypeInfo::GetStructMetadata()));
-	auto* UniformBuffer = UniformBuffers.Allocate<TRDGUniformBuffer<ParameterStructType>>(Allocator, ParameterStruct, ParameterStructType::FTypeInfo::GetStructMetadata()->GetShaderVariableName());
+	auto* UniformBuffer = UniformBuffers.Allocate<TRDGUniformBuffer<ParameterStructType>>(Allocators.Root, ParameterStruct, ParameterStructType::FTypeInfo::GetStructMetadata()->GetShaderVariableName());
 	IF_RDG_ENABLE_DEBUG(UserValidation.ValidateCreateUniformBuffer(UniformBuffer));
 	return UniformBuffer;
 #else
@@ -243,7 +249,7 @@ FRDGPassRef FRDGBuilder::AddPass(
 
 	FlushAccessModeQueue();
 
-	LambdaPassType* Pass = Passes.Allocate<LambdaPassType>(Allocator, MoveTemp(Name), Flags, MoveTemp(ExecuteLambda));
+	LambdaPassType* Pass = Passes.Allocate<LambdaPassType>(Allocators.Root, MoveTemp(Name), Flags, MoveTemp(ExecuteLambda));
 	SetupEmptyPass(Pass);
 	return Pass;
 #else
@@ -267,7 +273,7 @@ FRDGPassRef FRDGBuilder::AddPassInternal(
 
 	FlushAccessModeQueue();
 
-	FRDGPass* Pass = Allocator.AllocNoDestruct<LambdaPassType>(
+	FRDGPass* Pass = Allocators.Root.AllocNoDestruct<LambdaPassType>(
 		MoveTemp(Name),
 		ParametersMetadata,
 		ParameterStruct,
@@ -330,7 +336,6 @@ inline void FRDGBuilder::QueueBufferUpload(FRDGBufferRef Buffer, const void* Ini
 
 	UploadedBuffers.Emplace(Buffer, InitialData, InitialDataSize);
 	Buffer->bQueuedForUpload = 1;
-	Buffer->bForceNonTransient = 1;
 }
 
 inline void FRDGBuilder::QueueBufferUpload(FRDGBufferRef Buffer, const void* InitialData, uint64 InitialDataSize, FRDGBufferInitialDataFreeCallback&& InitialDataFreeCallback)
@@ -344,7 +349,6 @@ inline void FRDGBuilder::QueueBufferUpload(FRDGBufferRef Buffer, const void* Ini
 
 	UploadedBuffers.Emplace(Buffer, InitialData, InitialDataSize, MoveTemp(InitialDataFreeCallback));
 	Buffer->bQueuedForUpload = 1;
-	Buffer->bForceNonTransient = 1;
 }
 
 inline void FRDGBuilder::QueueBufferUpload(FRDGBufferRef Buffer, FRDGBufferInitialDataFillCallback&& InitialDataFillCallback)
@@ -353,7 +357,6 @@ inline void FRDGBuilder::QueueBufferUpload(FRDGBufferRef Buffer, FRDGBufferIniti
 
 	UploadedBuffers.Emplace(Buffer, MoveTemp(InitialDataFillCallback));
 	Buffer->bQueuedForUpload = 1;
-	Buffer->bForceNonTransient = 1;
 }
 
 inline void FRDGBuilder::QueueBufferUpload(FRDGBufferRef Buffer, FRDGBufferInitialDataCallback&& InitialDataCallback, FRDGBufferInitialDataSizeCallback&& InitialDataSizeCallback)
@@ -362,7 +365,6 @@ inline void FRDGBuilder::QueueBufferUpload(FRDGBufferRef Buffer, FRDGBufferIniti
 
 	UploadedBuffers.Emplace(Buffer, MoveTemp(InitialDataCallback), MoveTemp(InitialDataSizeCallback));
 	Buffer->bQueuedForUpload = 1;
-	Buffer->bForceNonTransient = 1;
 }
 
 inline void FRDGBuilder::QueueBufferUpload(FRDGBufferRef Buffer, FRDGBufferInitialDataCallback&& InitialDataCallback, FRDGBufferInitialDataSizeCallback&& InitialDataSizeCallback, FRDGBufferInitialDataFreeCallback&& InitialDataFreeCallback)
@@ -371,7 +373,6 @@ inline void FRDGBuilder::QueueBufferUpload(FRDGBufferRef Buffer, FRDGBufferIniti
 
 	UploadedBuffers.Emplace(Buffer, MoveTemp(InitialDataCallback), MoveTemp(InitialDataSizeCallback), MoveTemp(InitialDataFreeCallback));
 	Buffer->bQueuedForUpload = 1;
-	Buffer->bForceNonTransient = 1;
 }
 
 inline void FRDGBuilder::QueueCommitReservedBuffer(FRDGBufferRef Buffer, uint64 CommitSizeInBytes)
@@ -380,6 +381,7 @@ inline void FRDGBuilder::QueueCommitReservedBuffer(FRDGBufferRef Buffer, uint64
 
 	Buffer->PendingCommitSize = FMath::Max<uint64>(CommitSizeInBytes, Buffer->PendingCommitSize);
 	Buffer->PooledBuffer->SetCommittedSize(Buffer->PendingCommitSize);
+	Buffer->bQueuedForReservedCommit = 1;
 }
 
 inline void FRDGBuilder::QueueTextureExtraction(FRDGTextureRef Texture, TRefCountPtr<IPooledRenderTarget>* OutTexturePtr, ERHIAccess AccessFinal, ERDGResourceExtractionFlags Flags)
@@ -394,6 +396,8 @@ inline void FRDGBuilder::QueueTextureExtraction(FRDGTextureRef Texture, TRefCoun
 
 	*OutTexturePtr = nullptr;
 
+	const bool bWasExtracted = Texture->bExtracted;
+
 	Texture->bExtracted = true;
 
 	if (EnumHasAnyFlags(Flags, ERDGResourceExtractionFlags::AllowTransient))
@@ -409,6 +413,11 @@ inline void FRDGBuilder::QueueTextureExtraction(FRDGTextureRef Texture, TRefCoun
 	}
 
 	ExtractedTextures.Emplace(Texture, OutTexturePtr);
+
+	if (!bWasExtracted)
+	{
+		AsyncSetupQueue.Push(FAsyncSetupOp::CullRootTexture(Texture));
+	}
 }
 
 inline void FRDGBuilder::QueueBufferExtraction(FRDGBufferRef Buffer, TRefCountPtr<FRDGPooledBuffer>* OutBufferPtr)
@@ -417,9 +426,16 @@ inline void FRDGBuilder::QueueBufferExtraction(FRDGBufferRef Buffer, TRefCountPt
 
 	*OutBufferPtr = nullptr;
 
+	const bool bWasExtracted = Buffer->bExtracted;
+
 	Buffer->bExtracted = true;
 	Buffer->bForceNonTransient = true;
 	ExtractedBuffers.Emplace(Buffer, OutBufferPtr);
+
+	if (!bWasExtracted)
+	{
+		AsyncSetupQueue.Push(FAsyncSetupOp::CullRootBuffer(Buffer));
+	}
 }
 
 inline void FRDGBuilder::QueueBufferExtraction(FRDGBufferRef Buffer, TRefCountPtr<FRDGPooledBuffer>* OutBufferPtr, ERHIAccess AccessFinal)
@@ -500,7 +516,10 @@ namespace UE::RDG
 	template<typename TaskCollectionType, decltype(std::declval<TaskCollectionType>().begin())* = nullptr>
 	inline void Wait(const TaskCollectionType& Tasks)
 	{
-		UE::Tasks::Wait(Tasks);
+		if (!Tasks.IsEmpty())
+		{
+			UE::Tasks::Wait(Tasks);
+		}
 	}
 
 	template<typename TaskType, decltype(std::declval<TaskType>().Wait())* = nullptr>
@@ -520,7 +539,10 @@ UE::Tasks::FTask FRDGBuilder::AddSetupTask(
 {
 	UE::Tasks::FTask Task;
 
-	bCondition &= bParallelSetupEnabled;
+	if (!bCondition || IsImmediateMode())
+	{
+		UE::RDG::Wait(Prerequisites);
+	}
 
 	auto OuterLambda = [TaskLambda = MoveTemp(TaskLambda)]() mutable
 	{
@@ -528,11 +550,10 @@ UE::Tasks::FTask FRDGBuilder::AddSetupTask(
 		TaskLambda();
 	};
 
-	const UE::Tasks::EExtendedTaskPriority ExtendedTaskPriority = bCondition ? UE::Tasks::EExtendedTaskPriority::None : UE::Tasks::EExtendedTaskPriority::Inline;
+	const UE::Tasks::EExtendedTaskPriority ExtendedTaskPriority = ParallelSetup.bEnabled ? UE::Tasks::EExtendedTaskPriority::None : UE::Tasks::EExtendedTaskPriority::Inline;
 
-	if (IsImmediateMode())
+	if (!bCondition || (!ParallelSetup.bEnabled && UE::RDG::IsCompleted(Prerequisites)))
 	{
-		UE::RDG::Wait(Prerequisites);
 		OuterLambda();
 	}
 	else if (Pipe)
@@ -544,9 +565,9 @@ UE::Tasks::FTask FRDGBuilder::AddSetupTask(
 		Task = UE::Tasks::Launch(TEXT("FRDGBuilder::AddSetupTask"), MoveTemp(OuterLambda), Forward<PrerequisitesCollectionType&&>(Prerequisites), Priority, ExtendedTaskPriority);
 	}
 
-	if (bCondition)
+	if (Task.IsValid())
 	{
-		ParallelSetupEvents.Emplace(Task);
+		ParallelSetup.Tasks.Emplace(Task);
 	}
 
 	return Task;
@@ -594,44 +615,41 @@ UE::Tasks::FTask FRDGBuilder::AddCommandListSetupTask(
 {
 	UE::Tasks::FTask Task;
 
-	bCondition &= bParallelSetupEnabled;
-
-	FRHICommandList* RHICmdListTask = nullptr;
-
-	const bool bUseSeparateCommandList = bCondition || !UE::RDG::IsCompleted(Prerequisites);
-
-	// When using prerequisites we have to 
-	if (bUseSeparateCommandList)
-	{
-		RHICmdListTask = new FRHICommandList(RHICmdList.GetGPUMask());
-	}
-	else
+	if (!bCondition || IsImmediateMode())
 	{
-		RHICmdListTask = &RHICmdList;
+		UE::RDG::Wait(Prerequisites);
 	}
 
-	auto OuterLambda = [TaskLambda = MoveTemp(TaskLambda), RHICmdListTask]() mutable
+	// Need a separate command list with inline tasks when prerequisites are not complete yet.
+	const bool bAllocateCommandListForTask = bCondition && (ParallelSetup.bEnabled || !UE::RDG::IsCompleted(Prerequisites));
+
+	auto OuterLambda = [this, TaskLambda = MoveTemp(TaskLambda), bAllocateCommandListForTask]() mutable
 	{
 		FOptionalTaskTagScope Scope(ETaskTag::EParallelRenderingThread);
 
-		if (!RHICmdListTask->IsImmediate())
+		FRHICommandList* RHICmdListTask = &RHICmdList;
+		
+		if (bAllocateCommandListForTask)
 		{
+			RHICmdListTask = new FRHICommandList(RHICmdList.GetGPUMask());
 			RHICmdListTask->SwitchPipeline(ERHIPipeline::Graphics);
 		}
 
 		TaskLambda(*RHICmdListTask);
 
-		if (!RHICmdListTask->IsImmediate())
+		if (bAllocateCommandListForTask)
 		{
 			RHICmdListTask->FinishRecording();
+
+			UE::TScopeLock Lock(ParallelSetup.CommandListsMutex);
+			ParallelSetup.CommandLists.Emplace(RHICmdListTask);
 		}
 	};
 
 	const UE::Tasks::EExtendedTaskPriority ExtendedTaskPriority = bCondition ? UE::Tasks::EExtendedTaskPriority::None : UE::Tasks::EExtendedTaskPriority::Inline;
 
-	if (IsImmediateMode())
+	if (!bAllocateCommandListForTask)
 	{
-		UE::RDG::Wait(Prerequisites);
 		OuterLambda();
 	}
 	else if (Pipe)
@@ -643,10 +661,9 @@ UE::Tasks::FTask FRDGBuilder::AddCommandListSetupTask(
 		Task = UE::Tasks::Launch(TEXT("FRDGBuilder::AddCommandListSetupTask"), MoveTemp(OuterLambda), Forward<PrerequisitesCollectionType&&>(Prerequisites), Priority, ExtendedTaskPriority);
 	}
 
-	if (!RHICmdListTask->IsImmediate())
+	if (Task.IsValid())
 	{
-		ParallelSetupEvents.Emplace(Task);
-		ParallelSetupCommandLists.Emplace(RHICmdListTask);
+		ParallelSetup.Tasks.Emplace(Task);
 	}
 
 	return Task;
@@ -689,19 +706,13 @@ inline void FRDGBuilder::RemoveUnusedBufferWarning(FRDGBufferRef Buffer)
 inline void FRDGBuilder::BeginEventScope(FRDGEventName&& ScopeName)
 {
 #if RDG_GPU_DEBUG_SCOPES
-	if (!bFinalEventScopeActive)
-	{
-		GPUScopeStacks.BeginEventScope(MoveTemp(ScopeName), RHICmdList.GetGPUMask(), ERDGEventScopeFlags::None);
-	}
+	GPUScopeStacks.BeginEventScope(MoveTemp(ScopeName), RHICmdList.GetGPUMask(), ERDGEventScopeFlags::None);
 #endif
 }
 
 inline void FRDGBuilder::EndEventScope()
 {
 #if RDG_GPU_DEBUG_SCOPES
-	if (!bFinalEventScopeActive)
-	{
-		GPUScopeStacks.EndEventScope();
-	}
+	GPUScopeStacks.EndEventScope();
 #endif
 }
\ No newline at end of file
diff --git a/Engine/Source/Runtime/RenderCore/Public/RenderGraphDefinitions.h b/Engine/Source/Runtime/RenderCore/Public/RenderGraphDefinitions.h
--- a/Engine/Source/Runtime/RenderCore/Public/RenderGraphDefinitions.h
+++ b/Engine/Source/Runtime/RenderCore/Public/RenderGraphDefinitions.h
@@ -711,6 +711,7 @@ using FRDGUniformBufferBitArray = TRDGHandleBitArray<FRDGUniformBufferHandle>;
 using FRDGViewHandle = TRDGHandle<FRDGView, uint16>;
 using FRDGViewRegistry = TRDGHandleRegistry<FRDGViewHandle, ERDGHandleRegistryDestructPolicy::Never>;
 using FRDGViewUniqueFilter = TRDGHandleUniqueFilter<FRDGViewHandle>;
+using FRDGViewBitArray = TRDGHandleBitArray<FRDGViewHandle>;
 
 using FRDGTextureHandle = TRDGHandle<FRDGTexture, uint16>;
 using FRDGTextureRegistry = TRDGHandleRegistry<FRDGTextureHandle, ERDGHandleRegistryDestructPolicy::Never>;
diff --git a/Engine/Source/Runtime/RenderCore/Public/RenderGraphEvent.h b/Engine/Source/Runtime/RenderCore/Public/RenderGraphEvent.h
--- a/Engine/Source/Runtime/RenderCore/Public/RenderGraphEvent.h
+++ b/Engine/Source/Runtime/RenderCore/Public/RenderGraphEvent.h
@@ -215,9 +215,7 @@ class TRDGScopeStack final
 public:
 	using ScopeType = typename ScopeOpType::ScopeType;
 
-	TRDGScopeStack(FRDGAllocator& InAllocator)
-		: Allocator(InAllocator)
-	{}
+	TRDGScopeStack() = default;
 
 	~TRDGScopeStack()
 	{
@@ -231,7 +229,7 @@ public:
 	template <typename... TScopeConstructArgs>
 	inline void BeginScope(TScopeConstructArgs&&... ScopeConstructArgs)
 	{
-		auto Scope = Allocator.AllocNoDestruct<ScopeType>(CurrentScope, Forward<TScopeConstructArgs>(ScopeConstructArgs)...);
+		auto Scope = FRDGAllocator::GetTLS().AllocNoDestruct<ScopeType>(CurrentScope, Forward<TScopeConstructArgs>(ScopeConstructArgs)...);
 		Scopes.Add(Scope);
 		CurrentScope = Scope;
 	}
@@ -269,7 +267,6 @@ public:
 	}
 
 private:
-	FRDGAllocator& Allocator;
 
 	/** The top of the scope stack during setup. */
 	const ScopeType* CurrentScope = nullptr;
@@ -323,9 +320,7 @@ private:
 class FRDGTimingScopeStack final
 {
 public:
-	FRDGTimingScopeStack(FRDGAllocator& Allocator)
-		: ScopeStack(Allocator)
-	{}
+	FRDGTimingScopeStack() = default;
 
 	inline void BeginScope(const DynamicRenderScaling::FBudget& Budget)
 	{
@@ -490,9 +485,8 @@ public:
 class FRDGEventScopeStack final
 {
 public:
-	FRDGEventScopeStack(FRDGAllocator& Allocator)
-		: ScopeStack(Allocator)
-		, bRDGEvents(GetEmitRDGEvents())
+	FRDGEventScopeStack()
+		: bRDGEvents(GetEmitRDGEvents())
 	{}
 
 	inline void BeginScope(FRDGEventName&& EventName, FRHIGPUMask GPUMask, ERDGEventScopeFlags Flags)
@@ -642,10 +636,9 @@ public:
 class FRDGGPUStatScopeStack final
 {
 public:
-	FRDGGPUStatScopeStack(FRDGAllocator& Allocator)
-		: ScopeStack(Allocator)
+	FRDGGPUStatScopeStack()
 #if HAS_GPU_STATS
-		, bGPUStats(AreGPUStatsEnabled())
+		: bGPUStats(AreGPUStatsEnabled())
 #endif
 	{}
 
@@ -750,13 +743,7 @@ struct FRDGGPUScopeOpArrays
 /** The complete set of scope stack implementations. */
 struct FRDGGPUScopeStacks
 {
-	FRDGGPUScopeStacks(FRDGAllocator& Allocator)
-		: Timing(Allocator)
-#if RDG_GPU_DEBUG_SCOPES
-		, Event(Allocator)
-		, Stat(Allocator)
-#endif
-	{}
+	FRDGGPUScopeStacks() = default;
 
 	inline void ReserveOps(int32 PassCount)
 	{
@@ -805,9 +792,7 @@ struct FRDGGPUScopeStacks
 
 struct FRDGGPUScopeStacksByPipeline
 {
-	FRDGGPUScopeStacksByPipeline(FRDGAllocator& Allocator)
-		: Graphics(Allocator)
-		, AsyncCompute(Allocator)
+	FRDGGPUScopeStacksByPipeline()
 	{
 		IsTimingIsEnabled.SetAll(false);
 	}
@@ -836,15 +821,21 @@ struct FRDGGPUScopeStacksByPipeline
 #if RDG_GPU_DEBUG_SCOPES
 	inline void BeginEventScope(FRDGEventName&& ScopeName, FRHIGPUMask GPUMask, ERDGEventScopeFlags Flags)
 	{
-		FRDGEventName ScopeNameCopy = ScopeName;
-		Graphics.Event.BeginScope(MoveTemp(ScopeNameCopy), GPUMask, Flags);
-		AsyncCompute.Event.BeginScope(MoveTemp(ScopeName), GPUMask, Flags);
+		if (!bFinalEventScopeActive)
+		{
+			FRDGEventName ScopeNameCopy = ScopeName;
+			Graphics.Event.BeginScope(MoveTemp(ScopeNameCopy), GPUMask, Flags);
+			AsyncCompute.Event.BeginScope(MoveTemp(ScopeName), GPUMask, Flags);
+		}
 	}
 
 	inline void EndEventScope()
 	{
-		Graphics.Event.EndScope();
-		AsyncCompute.Event.EndScope();
+		if (!bFinalEventScopeActive)
+		{
+			Graphics.Event.EndScope();
+			AsyncCompute.Event.EndScope();
+		}
 	}
 
 	inline void BeginStatScope(const FName& Name, const FName& StatName, const TCHAR* Description, FDrawCallCategoryName& Category)
@@ -879,6 +870,7 @@ struct FRDGGPUScopeStacksByPipeline
 
 	FRDGGPUScopeStacks Graphics;
 	FRDGGPUScopeStacks AsyncCompute;
+	bool bFinalEventScopeActive = false;
 
 private:
 	DynamicRenderScaling::TMap<bool> IsTimingIsEnabled;
@@ -922,9 +914,7 @@ public:
 class FRDGCSVStatScopeStack final
 {
 public:
-	FRDGCSVStatScopeStack(FRDGAllocator& Allocator)
-		: ScopeStack(Allocator)
-	{}
+	FRDGCSVStatScopeStack() = default;
 
 	void BeginScope(const char* StatName)
 	{
@@ -1018,9 +1008,7 @@ struct FRDGCPUScopeOpArrays
 
 struct FRDGCPUScopeStacks
 {
-	FRDGCPUScopeStacks(FRDGAllocator& Allocator)
-		: CSV(Allocator)
-	{}
+	FRDGCPUScopeStacks() = default;
 
 	inline void ReserveOps()
 	{
diff --git a/Engine/Source/Runtime/RenderCore/Public/RenderGraphPass.h b/Engine/Source/Runtime/RenderCore/Public/RenderGraphPass.h
--- a/Engine/Source/Runtime/RenderCore/Public/RenderGraphPass.h
+++ b/Engine/Source/Runtime/RenderCore/Public/RenderGraphPass.h
@@ -67,6 +67,21 @@ enum class ERDGBarrierLocation : uint8
 	Epilogue
 };
 
+struct FRDGTransitionInfo
+{
+	FRDGTransitionInfo() = default;
+
+	ERHIAccess AccessBefore;
+	ERHIAccess AccessAfter;
+	uint16 Handle;
+	ERDGViewableResourceType Type;
+	EResourceTransitionFlags Flags;
+	uint32 ArraySlice : 16;
+	uint32 MipIndex   : 8;
+	uint32 PlaneSlice : 2;
+	uint32 bReservedCommit : 1;
+};
+
 struct FRDGBarrierBatchEndId
 {
 	FRDGBarrierBatchEndId() = default;
@@ -95,7 +110,7 @@ public:
 	RENDERCORE_API FRDGBarrierBatchBegin(ERHIPipeline PipelinesToBegin, ERHIPipeline PipelinesToEnd, const TCHAR* DebugName, FRDGPass* DebugPass);
 	RENDERCORE_API FRDGBarrierBatchBegin(ERHIPipeline PipelinesToBegin, ERHIPipeline PipelinesToEnd, const TCHAR* DebugName, FRDGPassesByPipeline DebugPasses);
 
-	RENDERCORE_API void AddTransition(FRDGViewableResource* Resource, const FRHITransitionInfo& Info);
+	RENDERCORE_API void AddTransition(FRDGViewableResource* Resource, FRDGTransitionInfo Info);
 
 	RENDERCORE_API void AddAlias(FRDGViewableResource* Resource, const FRHITransientAliasingInfo& Info);
 
@@ -105,7 +120,7 @@ public:
 		bTransitionNeeded = true;
 	}
 
-	RENDERCORE_API void CreateTransition();
+	RENDERCORE_API void CreateTransition(TConstArrayView<FRHITransitionInfo> TransitionsRHI);
 
 	RENDERCORE_API void Submit(FRHIComputeCommandList& RHICmdList, ERHIPipeline Pipeline);
 	RENDERCORE_API void Submit(FRHIComputeCommandList& RHICmdList, ERHIPipeline Pipeline, FRDGTransitionQueue& TransitionsToBegin);
@@ -113,6 +128,7 @@ public:
 	void Reserve(uint32 TransitionCount)
 	{
 		Transitions.Reserve(TransitionCount);
+		Aliases.Reserve(TransitionCount);
 	}
 
 	bool IsTransitionNeeded() const
@@ -122,7 +138,7 @@ public:
 
 private:
 	const FRHITransition* Transition = nullptr;
-	TArray<FRHITransitionInfo, FRDGArrayAllocator> Transitions;
+	TArray<FRDGTransitionInfo, FRDGArrayAllocator> Transitions;
 	TArray<FRHITransientAliasingInfo, FRDGArrayAllocator> Aliases;
 	ERHITransitionCreateFlags TransitionFlags = ERHITransitionCreateFlags::NoFence;
 	ERHIPipeline PipelinesToBegin;
@@ -141,6 +157,7 @@ private:
 
 	friend class FRDGBarrierBatchEnd;
 	friend class FRDGBarrierValidation;
+	friend class FRDGBuilder;
 };
 
 using FRDGTransitionCreateQueue = TArray<FRDGBarrierBatchBegin*, FRDGArrayAllocator>;
@@ -283,11 +300,6 @@ public:
 		return bSentinel;
 	}
 
-	const FRDGPassHandleArray& GetProducers() const
-	{
-		return Producers;
-	}
-
 	/** Returns the graphics pass responsible for forking the async interval this pass is in. */
 	FRDGPassHandle GetGraphicsForkPass() const
 	{
@@ -432,8 +444,8 @@ protected:
 	FRDGPassHandle EpilogueBarrierPass;
 
 	/** Lists of producer passes and the full list of cross-pipeline consumer passes. */
-	FRDGPassHandleArray CrossPipelineConsumers;
-	FRDGPassHandleArray Producers;
+	TArray<FRDGPassHandle, FRDGArrayAllocator> CrossPipelineConsumers;
+	TArray<FRDGPass*, FRDGArrayAllocator> Producers;
 
 	struct FTextureState
 	{
@@ -443,15 +455,13 @@ protected:
 			: Texture(InTexture)
 		{
 			const uint32 SubresourceCount = Texture->GetSubresourceCount();
-			State.Reserve(SubresourceCount);
 			State.SetNum(SubresourceCount);
-			MergeState.Reserve(SubresourceCount);
 			MergeState.SetNum(SubresourceCount);
 		}
 
 		FRDGTextureRef Texture = nullptr;
 		FRDGTextureSubresourceState State;
-		FRDGTextureSubresourceStateIndirect MergeState;
+		FRDGTextureSubresourceState MergeState;
 		uint16 ReferenceCount = 0;
 	};
 
diff --git a/Engine/Source/Runtime/RenderCore/Public/RenderGraphResources.h b/Engine/Source/Runtime/RenderCore/Public/RenderGraphResources.h
--- a/Engine/Source/Runtime/RenderCore/Public/RenderGraphResources.h
+++ b/Engine/Source/Runtime/RenderCore/Public/RenderGraphResources.h
@@ -45,8 +45,6 @@ struct FRDGViewableResourceDebugData;
 /** Used for tracking pass producer / consumer edges in the graph for culling and pipe fencing. */
 struct FRDGProducerState
 {
-	FRDGProducerState() = default;
-
 	FRDGPass* Pass = nullptr;
 	FRDGPass* PassIfSkipUAVBarrier = nullptr;
 	ERHIAccess Access = ERHIAccess::Unknown;
@@ -64,18 +62,24 @@ struct FRDGSubresourceState
 	/** Given a before and after state, returns whether they can be merged into a single state. */
 	static bool IsMergeAllowed(ERDGViewableResourceType ResourceType, const FRDGSubresourceState& Previous, const FRDGSubresourceState& Next);
 
-	FRDGSubresourceState() = default;
+	FRDGSubresourceState()
+		: bReservedCommit(0)
+	{}
 
 	explicit FRDGSubresourceState(ERHIAccess InAccess)
 		: Access(InAccess)
+		, bReservedCommit(0)
 	{}
 
+	explicit FRDGSubresourceState(ERHIPipeline Pipeline, FRDGPassHandle PassHandle)
+		: bReservedCommit(0)
+	{
+		SetPass(Pipeline, PassHandle);
+	}
+
 	/** Initializes the first and last pass and the pipeline. Clears any other pass state. */
 	void SetPass(ERHIPipeline Pipeline, FRDGPassHandle PassHandle);
 
-	/** Finalizes the state at the end of the transition chain; keeps access intact. */
-	void Finalize();
-
 	/** Validates that the state is in a correct configuration for use. */
 	void Validate();
 
@@ -94,9 +98,6 @@ struct FRDGSubresourceState
 	/** The last used access on the pass. */
 	ERHIAccess Access = ERHIAccess::Unknown;
 
-	/** The last used transition flags on the pass. */
-	EResourceTransitionFlags Flags = EResourceTransitionFlags::None;
-
 	/** The first pass in this state. */
 	FRDGPassHandlesByPipeline FirstPass;
 
@@ -105,10 +106,15 @@ struct FRDGSubresourceState
 
 	/** The last no-UAV barrier to be used by this subresource. */
 	FRDGViewUniqueFilter NoUAVBarrierFilter;
+
+	/** The last used transition flags on the pass. */
+	EResourceTransitionFlags Flags = EResourceTransitionFlags::None;
+
+	/** Whether this subresource state represents a commit operation for a reserved resource. */
+	uint8 bReservedCommit : 1;
 };
 
-using FRDGTextureSubresourceState = TRDGTextureSubresourceArray<FRDGSubresourceState, FRDGArrayAllocator>;
-using FRDGTextureSubresourceStateIndirect = TRDGTextureSubresourceArray<FRDGSubresourceState*, FRDGArrayAllocator>;
+using FRDGTextureSubresourceState = TRDGTextureSubresourceArray<FRDGSubresourceState*, FRDGArrayAllocator>;
 
 /** Generic graph resource. */
 class FRDGResource
@@ -232,7 +238,6 @@ private:
 	const FRDGParameterStruct ParameterStruct;
 	TRefCountPtr<FRHIUniformBuffer> UniformBufferRHI;
 	FRDGUniformBufferHandle Handle;
-	bool bQueuedForCreate = false;
 	bool bExternal = false;
 
 	friend FRDGBuilder;
@@ -293,12 +298,6 @@ public:
 		return bExternal;
 	}
 
-	/** Whether this resource is allocated through the transient resource allocator. */
-	bool IsTransient() const
-	{
-		return bTransient;
-	}
-
 	/** Whether this resource is has been queued for extraction at the end of graph execution. */
 	bool IsExtracted() const
 	{
@@ -321,6 +320,11 @@ public:
 protected:
 	RENDERCORE_API FRDGViewableResource(const TCHAR* InName, ERDGViewableResourceType InType, bool bSkipTracking, bool bImmediateFirstBarrier);
 
+	bool IsCullRoot() const
+	{
+		return bExternal || bExtracted;
+	}
+
 	static const ERHIAccess DefaultEpilogueAccess = ERHIAccess::SRVMask;
 
 	enum class ETransientExtractionHint : uint8
@@ -330,13 +334,6 @@ protected:
 		Enable
 	};
 
-	enum class EFirstBarrier : uint8
-	{
-		Split,
-		ImmediateRequested,
-		ImmediateConfirmed
-	};
-
 	enum class EAccessMode : uint8
 	{
 		Internal,
@@ -380,26 +377,29 @@ protected:
 	/** Whether this resource cannot be made transient. */
 	uint8 bForceNonTransient : 1;
 
-	/** Whether this resource is allowed to be both transient and extracted. */
-	ETransientExtractionHint TransientExtractionHint : 2;
+	/** If true, the graph will skip the last graph transition. Used for aliased pooled resources. */
+	uint8 bSkipLastTransition : 1;
 
-	/** Whether this resource is the last owner of its allocation (i.e. nothing aliases the allocation later in the execution timeline). */
-	uint8 bLastOwner : 1;
+	/** If true, this resource should skip the first split barrier and perform transition right away. */
+	uint8 bSplitFirstTransition : 1;
 
 	/** If true, the resource has been queued for an upload operation. */
 	uint8 bQueuedForUpload : 1;
 
-	/** If true, this resource should skip the prologue split barrier and perform transition right away. */
-	EFirstBarrier FirstBarrier : 2;
+	/** If false, the resource needs to be collected. */
+	uint8 bCollectForAllocate : 1;
 
-	FRDGPassHandle FirstPass;
-	FRDGPassHandle LastPass;
+	/** If true, the reserved resource is having tiles committed. */
+	uint8 bQueuedForReservedCommit : 1;
 
-	/** The state of the resource at the graph epilogue. */
-	ERHIAccess EpilogueAccess = DefaultEpilogueAccess;
+	/** Whether this resource is allowed to be both transient and extracted. */
+	ETransientExtractionHint TransientExtractionHint;
 
-private:
-	static const uint16 DeallocatedReferenceCount = ~0;
+	FRDGPassHandle FirstPass;
+	FRDGPassHandle LastPass;
+	FRDGPassHandle MinAcquirePass;
+	FRDGPassHandle MinDiscardPass;
+	FRDGPassHandle MaxDiscardPass;
 
 	/** Number of references in passes and deferred queries. */
 	uint16 ReferenceCount;
@@ -407,6 +407,15 @@ private:
 	/** Scratch index allocated for the resource in the pass being setup. */
 	uint16 PassStateIndex = 0;
 
+	/** Set of aliasing overlaps to apply to the acquire transition if transient. */
+	TArrayView<const FRHITransientAliasingOverlap> AliasingOverlaps;
+
+	/** The state of the resource at the graph epilogue. */
+	ERHIAccess EpilogueAccess = DefaultEpilogueAccess;
+
+private:
+	static const uint16 DeallocatedReferenceCount = ~0;
+
 	void SetExternalAccessMode(ERHIAccess InReadOnlyAccess, ERHIPipeline InPipelines)
 	{
 		check(!AccessModeState.bLocked);
@@ -556,24 +565,7 @@ public:
 	RENDERCORE_API FRDGTextureSubresourceRange GetSubresourceRangeSRV() const;
 
 private:
-	FRDGTexture(const TCHAR* InName, const FRDGTextureDesc& InDesc, ERDGTextureFlags InFlags)
-		: FRDGViewableResource(InName, ERDGViewableResourceType::Texture, EnumHasAnyFlags(InFlags, ERDGTextureFlags::SkipTracking), EnumHasAnyFlags(InFlags, ERDGTextureFlags::ForceImmediateFirstBarrier) || EnumHasAnyFlags(InDesc.Flags, ETextureCreateFlags::Presentable))
-		, Desc(InDesc)
-		, Flags(InFlags)
-		, Layout(InDesc)
-		, WholeRange(Layout)
-		, SubresourceCount(Layout.GetSubresourceCount())
-	{
-		MergeState.Reserve(SubresourceCount);
-		MergeState.SetNum(SubresourceCount);
-		LastProducers.Reserve(SubresourceCount);
-		LastProducers.SetNum(SubresourceCount);
-
-		if (EnumHasAnyFlags(Desc.Flags, ETextureCreateFlags::Foveation))
-		{
-			EpilogueAccess = ERHIAccess::ShadingRateSource;
-		}
-	}
+	FRDGTexture(const TCHAR* InName, const FRDGTextureDesc& InDesc, ERDGTextureFlags InFlags);
 
 	/** Returns RHI texture without access checks. */
 	FRHITexture* GetRHIUnchecked() const
@@ -581,35 +573,38 @@ private:
 		return static_cast<FRHITexture*>(FRDGResource::GetRHIUnchecked());
 	}
 
-	/** Returns the current texture state. Only valid to call after SetRHI. */
-	FRDGTextureSubresourceState& GetState()
-	{
-		check(State);
-		return *State;
-	}
-
-	/** The next texture to own the PooledTexture allocation during execution. */
-	FRDGTextureHandle NextOwner;
-
 	/** The handle registered with the builder. */
 	FRDGTextureHandle Handle;
 
+	/** The previous / next texture to own the PooledTexture allocation during execution. */
+	FRDGTextureHandle PreviousOwner;
+	FRDGTextureHandle NextOwner;
+
 	/** The layout used to facilitate subresource transitions. */
 	FRDGTextureSubresourceLayout Layout;
 	FRDGTextureSubresourceRange  WholeRange;
-	const uint32 SubresourceCount;
+	const uint16 SubresourceCount;
+
+	/** Tracks subresource states as the graph is built. */
+	FRDGTextureSubresourceState State;
+
+	/** Tracks the first state in the graph for each subresource. */
+	FRDGTextureSubresourceState FirstState;
+
+	/** Tracks merged subresource states as the graph is built. */
+	FRDGTextureSubresourceState MergeState;
+
+	/** Tracks pass producers for each subresource as the graph is built. */
+	TRDGTextureSubresourceArray<FRDGProducerStatesByPipeline, FRDGArrayAllocator> LastProducers;
 
 	/** The assigned render target to use during execution. Never reset. */
 	IPooledRenderTarget* RenderTarget = nullptr;
 
-	union
-	{
-		/** The assigned pooled texture to use during execution. Never reset. */
-		FRDGPooledTexture* PooledTexture = nullptr;
+	/** The assigned pooled texture to use during execution. Never reset. */
+	FRDGPooledTexture* PooledTexture = nullptr;
 
-		/** The assigned transient texture to use during execution. Never reset. */
-		FRHITransientTexture* TransientTexture;
-	};
+	/** The assigned transient texture to use during execution. Never reset. */
+	FRHITransientTexture* TransientTexture = nullptr;
 
 	/** The assigned view cache for this texture (sourced from transient / pooled texture). Never reset. */
 	FRHITextureViewCache* ViewCache = nullptr;
@@ -617,15 +612,6 @@ private:
 	/** Valid strictly when holding a strong reference; use PooledRenderTarget instead. */
 	TRefCountPtr<IPooledRenderTarget> Allocation;
 
-	/** Tracks subresource states as the graph is built. */
-	FRDGTextureSubresourceState* State = nullptr;
-
-	/** Tracks merged subresource states as the graph is built. */
-	FRDGTextureSubresourceStateIndirect MergeState;
-
-	/** Tracks pass producers for each subresource as the graph is built. */
-	TRDGTextureSubresourceArray<FRDGProducerStatesByPipeline, FRDGArrayAllocator> LastProducers;
-
 #if RDG_ENABLE_DEBUG
 	struct FRDGTextureDebugData* TextureDebugData = nullptr;
 	RENDERCORE_API FRDGTextureDebugData& GetTextureDebugData() const;
@@ -1292,11 +1278,7 @@ public:
 	}
 
 private:
-	FRDGBuffer(const TCHAR* InName, const FRDGBufferDesc& InDesc, ERDGBufferFlags InFlags)
-		: FRDGViewableResource(InName, ERDGViewableResourceType::Buffer, EnumHasAnyFlags(InFlags, ERDGBufferFlags::SkipTracking), EnumHasAnyFlags(InFlags, ERDGBufferFlags::ForceImmediateFirstBarrier))
-		, Desc(InDesc)
-		, Flags(InFlags)
-	{}
+	FRDGBuffer(const TCHAR* InName, const FRDGBufferDesc& InDesc, ERDGBufferFlags InFlags);
 
 	FRDGBuffer(const TCHAR* InName, const FRDGBufferDesc& InDesc, ERDGBufferFlags InFlags, FRDGBufferNumElementsCallback&& InNumElementsCallback)
 		: FRDGBuffer(InName, InDesc, InFlags)
@@ -1305,45 +1287,25 @@ private:
 	}
 
 	/** Finalizes any pending field of the buffer descriptor. */
-	void FinalizeDesc()
-	{
-		if (NumElementsCallback)
-		{
-			Desc.NumElements = FMath::Max(NumElementsCallback(), 1u);
-		}
-	}
+	void FinalizeDesc();
 
 	FRHIBuffer* GetRHIUnchecked() const
 	{
 		return static_cast<FRHIBuffer*>(FRDGResource::GetRHIUnchecked());
 	}
 
-	bool IsCulled() const
-	{
-		return FRDGViewableResource::IsCulled() && (PendingCommitSize == 0);
-	}
-
-	/** Returns the current buffer state. Only valid to call after SetRHI. */
-	FRDGSubresourceState& GetState() const
-	{
-		check(State);
-		return *State;
-	}
-
 	/** Registered handle set by the builder. */
 	FRDGBufferHandle Handle;
 
-	/** The next buffer to own the PooledBuffer allocation during execution. */
+	/** The previous / next buffer to own the PooledBuffer allocation during execution. */
+	FRDGBufferHandle PreviousOwner;
 	FRDGBufferHandle NextOwner;
 
-	union
-	{
-		/** Assigned pooled buffer pointer. Never reset once assigned. */
-		FRDGPooledBuffer* PooledBuffer = nullptr;
+	/** Assigned pooled buffer pointer. Never reset once assigned. */
+	FRDGPooledBuffer* PooledBuffer = nullptr;
 
-		/** Assigned transient buffer pointer. Never reset once assigned. */
-		FRHITransientBuffer* TransientBuffer;
-	};
+	/** Assigned transient buffer pointer. Never reset once assigned. */
+	FRHITransientBuffer* TransientBuffer = nullptr;
 
 	/** The assigned buffer view cache (sourced from the pooled / transient buffer. Never reset. */
 	FRHIBufferViewCache* ViewCache = nullptr;
@@ -1351,12 +1313,6 @@ private:
 	/** Valid strictly when holding a strong reference; use PooledBuffer instead. */
 	TRefCountPtr<FRDGPooledBuffer> Allocation;
 
-	/** Cached state pointer from the pooled / transient buffer. */
-	FRDGSubresourceState* State = nullptr;
-
-	/** Tracks the merged subresource state as the graph is built. */
-	FRDGSubresourceState* MergeState = nullptr;
-
 	/** Tracks the last pass that produced this resource as the graph is built. */
 	FRDGProducerStatesByPipeline LastProducer;
 
@@ -1366,6 +1322,15 @@ private:
 	/** Optional reserved resource commit size to apply on the first resource transition. */
 	uint64 PendingCommitSize = 0;
 
+	/** Cached state pointer from the pooled / transient buffer. */
+	FRDGSubresourceState* State = nullptr;
+
+	/** Tracks the first state in the graph for this buffer. */
+	FRDGSubresourceState* FirstState = nullptr;
+
+	/** Tracks the merged subresource state as the graph is built. */
+	FRDGSubresourceState* MergeState = nullptr;
+
 #if RDG_ENABLE_DEBUG
 	struct FRDGBufferDebugData* BufferDebugData = nullptr;
 	RENDERCORE_API FRDGBufferDebugData& GetBufferDebugData() const;
@@ -1377,12 +1342,6 @@ private:
 	friend FRDGBufferRegistry;
 	friend FRDGAllocator;
 	friend FRDGTrace;
-
-#if RDG_ENABLE_DEBUG
-	// Access to IsCulled() is needed
-	template <typename ResourceRegistryType, typename FunctionType>
-	friend void EnumerateExtendedLifetimeResources(ResourceRegistryType& Registry, FunctionType Function);
-#endif //RDG_ENABLE_DEBUG
 };
 
 /** Render graph tracked buffer SRV. */
diff --git a/Engine/Source/Runtime/RenderCore/Public/RenderGraphResources.inl b/Engine/Source/Runtime/RenderCore/Public/RenderGraphResources.inl
--- a/Engine/Source/Runtime/RenderCore/Public/RenderGraphResources.inl
+++ b/Engine/Source/Runtime/RenderCore/Public/RenderGraphResources.inl
@@ -4,15 +4,6 @@
 
 PRAGMA_DISABLE_BUFFER_OVERRUN_WARNING
 
-inline void FRDGSubresourceState::Finalize()
-{
-	ensureMsgf(!EnumHasAnyFlags(GetPipelines(), ERHIPipeline::AsyncCompute), TEXT("Resource should not be on the async compute pipeline!"));
-
-	const ERHIAccess LocalAccess = Access;
-	*this = {};
-	Access = LocalAccess;
-}
-
 inline void FRDGSubresourceState::SetPass(ERHIPipeline Pipeline, FRDGPassHandle PassHandle)
 {
 	FirstPass = {};
diff --git a/Engine/Source/Runtime/RenderCore/Public/RenderGraphTextureSubresource.h b/Engine/Source/Runtime/RenderCore/Public/RenderGraphTextureSubresource.h
--- a/Engine/Source/Runtime/RenderCore/Public/RenderGraphTextureSubresource.h
+++ b/Engine/Source/Runtime/RenderCore/Public/RenderGraphTextureSubresource.h
@@ -5,7 +5,7 @@
 #include "RHI.h"
 
 template <typename InElementType, typename InAllocatorType = FDefaultAllocator>
-using TRDGTextureSubresourceArray = TArray<InElementType, TInlineAllocator<1, InAllocatorType>>;
+using TRDGTextureSubresourceArray = TArray<InElementType, InAllocatorType>;
 
 struct FRDGTextureSubresource
 {
diff --git a/Engine/Source/Runtime/RenderCore/Public/RenderGraphValidation.h b/Engine/Source/Runtime/RenderCore/Public/RenderGraphValidation.h
--- a/Engine/Source/Runtime/RenderCore/Public/RenderGraphValidation.h
+++ b/Engine/Source/Runtime/RenderCore/Public/RenderGraphValidation.h
@@ -167,8 +167,8 @@ public:
 private:
 	struct FResourceMap
 	{
-		TMap<FRDGTextureRef, TArray<FRHITransitionInfo>> Textures;
-		TMap<FRDGBufferRef, FRHITransitionInfo> Buffers;
+		TMap<FRDGTextureRef, TArray<FRDGTransitionInfo>> Textures;
+		TMap<FRDGBufferRef, FRDGTransitionInfo> Buffers;
 		TMap<FRDGViewableResource*, FRHITransientAliasingInfo> Aliases;
 	};
 
diff --git a/Engine/Source/Runtime/Renderer/Private/MeshDrawCommands.cpp b/Engine/Source/Runtime/Renderer/Private/MeshDrawCommands.cpp
--- a/Engine/Source/Runtime/Renderer/Private/MeshDrawCommands.cpp
+++ b/Engine/Source/Runtime/Renderer/Private/MeshDrawCommands.cpp
@@ -1456,7 +1456,7 @@ void FParallelMeshDrawCommandPass::DispatchPassSetup(
 		// This work needs to be deferred until at least BuildRenderingCommands (to ensure the DynamicPrimitiveCollector is uploaded), so we use the async mechanism either way 
 		auto FinalizeInstanceCullingSetup = [this, Scene](FInstanceCullingContext& InstanceCullingContext)
 		{
-			WaitForMeshPassSetupTask();
+			WaitForMeshPassSetupTask(IsInActualRenderingThread() ? EWaitThread::Render : EWaitThread::Task);
 
 #if DO_CHECK
 			for (const FVisibleMeshDrawCommand& VisibleMeshDrawCommand : TaskContext.MeshDrawCommands)
