Fixed build break in RDG due to missing DLL export.

diff --git a/Engine/Source/Runtime/RenderCore/Public/RenderGraphResources.h b/Engine/Source/Runtime/RenderCore/Public/RenderGraphResources.h
--- a/Engine/Source/Runtime/RenderCore/Public/RenderGraphResources.h
+++ b/Engine/Source/Runtime/RenderCore/Public/RenderGraphResources.h
@@ -565,7 +565,7 @@ public:
 	RENDERCORE_API FRDGTextureSubresourceRange GetSubresourceRangeSRV() const;
 
 private:
-	FRDGTexture(const TCHAR* InName, const FRDGTextureDesc& InDesc, ERDGTextureFlags InFlags);
+	RENDERCORE_API FRDGTexture(const TCHAR* InName, const FRDGTextureDesc& InDesc, ERDGTextureFlags InFlags);
 
 	/** Returns RHI texture without access checks. */
 	FRHITexture* GetRHIUnchecked() const
@@ -1278,7 +1278,7 @@ public:
 	}
 
 private:
-	FRDGBuffer(const TCHAR* InName, const FRDGBufferDesc& InDesc, ERDGBufferFlags InFlags);
+	RENDERCORE_API FRDGBuffer(const TCHAR* InName, const FRDGBufferDesc& InDesc, ERDGBufferFlags InFlags);
 
 	FRDGBuffer(const TCHAR* InName, const FRDGBufferDesc& InDesc, ERDGBufferFlags InFlags, FRDGBufferNumElementsCallback&& InNumElementsCallback)
 		: FRDGBuffer(InName, InDesc, InFlags)
