Render Graph:  A couple simple development build RDG optimizations.
* Skip contents of Begin/EndUAVOverlap when GRHIValidationEnabled is false, avoiding expensive GatherPassUAVsForOverlapValidation call, saving around 1.4% frame wide.  Lower level function does nothing when validation isn't enabled.
* Optimize trivial RDG event names with no format specifiers.  Detects at compile time that there are no variadic parameters and calls a non-variadic version of the constructor, which references the string instead of generating a dynamically allocated copy through FCString::GetVarArgs.  Covers around 40% of RDG event names in my test case, saving around 1.2% frame wide.

#rnx

diff --git a/Engine/Source/Runtime/RenderCore/Private/RenderGraphBuilder.cpp b/Engine/Source/Runtime/RenderCore/Private/RenderGraphBuilder.cpp
--- a/Engine/Source/Runtime/RenderCore/Private/RenderGraphBuilder.cpp
+++ b/Engine/Source/Runtime/RenderCore/Private/RenderGraphBuilder.cpp
@@ -59,12 +59,15 @@ struct FParallelPassSet : public FRHICommandListImmediate::FQueuedCommandList
 inline void BeginUAVOverlap(const FRDGPass* Pass, FRHIComputeCommandList& RHICmdList)
 {
 #if ENABLE_RHI_VALIDATION
-	TArray<FRHIUnorderedAccessView*, TInlineAllocator<MaxSimultaneousUAVs, FRDGArrayAllocator>> UAVs;
-	GatherPassUAVsForOverlapValidation(Pass, UAVs);
-
-	if (UAVs.Num())
+	if (GRHIValidationEnabled)
 	{
-		RHICmdList.BeginUAVOverlap(UAVs);
+		TArray<FRHIUnorderedAccessView*, TInlineAllocator<MaxSimultaneousUAVs, FRDGArrayAllocator>> UAVs;
+		GatherPassUAVsForOverlapValidation(Pass, UAVs);
+
+		if (UAVs.Num())
+		{
+			RHICmdList.BeginUAVOverlap(UAVs);
+		}
 	}
 #endif
 }
@@ -72,12 +75,15 @@ inline void BeginUAVOverlap(const FRDGPass* Pass, FRHIComputeCommandList& RHICmd
 inline void EndUAVOverlap(const FRDGPass* Pass, FRHIComputeCommandList& RHICmdList)
 {
 #if ENABLE_RHI_VALIDATION
-	TArray<FRHIUnorderedAccessView*, TInlineAllocator<MaxSimultaneousUAVs, FRDGArrayAllocator>> UAVs;
-	GatherPassUAVsForOverlapValidation(Pass, UAVs);
-
-	if (UAVs.Num())
+	if (GRHIValidationEnabled)
 	{
-		RHICmdList.EndUAVOverlap(UAVs);
+		TArray<FRHIUnorderedAccessView*, TInlineAllocator<MaxSimultaneousUAVs, FRDGArrayAllocator>> UAVs;
+		GatherPassUAVsForOverlapValidation(Pass, UAVs);
+
+		if (UAVs.Num())
+		{
+			RHICmdList.EndUAVOverlap(UAVs);
+		}
 	}
 #endif
 }
diff --git a/Engine/Source/Runtime/RenderCore/Public/RenderGraphEvent.h b/Engine/Source/Runtime/RenderCore/Public/RenderGraphEvent.h
--- a/Engine/Source/Runtime/RenderCore/Public/RenderGraphEvent.h
+++ b/Engine/Source/Runtime/RenderCore/Public/RenderGraphEvent.h
@@ -30,7 +30,8 @@
  *  RDG_EVENT_SCOPE(GraphBuilder, "MyProcessing %sx%s", ViewRect.Width(), ViewRect.Height());
  */
 #if RDG_EVENTS == RDG_EVENTS_STRING_REF || RDG_EVENTS == RDG_EVENTS_STRING_COPY
-	#define RDG_EVENT_NAME(Format, ...) FRDGEventName(TEXT(Format), ##__VA_ARGS__)
+	// Skip expensive string formatting for the relatively common case of no varargs.  We detect this by stringizing the varargs and checking if the string is non-empty (more than just a null terminator).
+	#define RDG_EVENT_NAME(Format, ...) (sizeof(#__VA_ARGS__ "") > 1 ? FRDGEventName(TEXT(Format), ##__VA_ARGS__) : FRDGEventName(1, TEXT(Format)))
 	#define RDG_EVENT_SCOPE(GraphBuilder, Format, ...) FRDGEventScopeGuard PREPROCESSOR_JOIN(__RDG_ScopeRef_,__LINE__) ((GraphBuilder), RDG_EVENT_NAME(Format, ##__VA_ARGS__))
 	#define RDG_EVENT_SCOPE_CONDITIONAL(GraphBuilder, Condition, Format, ...) FRDGEventScopeGuard PREPROCESSOR_JOIN(__RDG_ScopeRef_,__LINE__) ((GraphBuilder), RDG_EVENT_NAME(Format, ##__VA_ARGS__), Condition)
 
@@ -386,7 +387,9 @@ class FRDGEventName final
 public:
 	FRDGEventName() = default;
 
+	// Constructors require a string that matches the RDG builder lifetime, as copies are not made in all configurations.
 	RENDERCORE_API explicit FRDGEventName(const TCHAR* EventFormat, ...);
+	FRDGEventName(int32 NonVariadic, const TCHAR* EventName);
 
 	FRDGEventName(const FRDGEventName& Other);
 	FRDGEventName(FRDGEventName&& Other);
diff --git a/Engine/Source/Runtime/RenderCore/Public/RenderGraphEvent.inl b/Engine/Source/Runtime/RenderCore/Public/RenderGraphEvent.inl
--- a/Engine/Source/Runtime/RenderCore/Public/RenderGraphEvent.inl
+++ b/Engine/Source/Runtime/RenderCore/Public/RenderGraphEvent.inl
@@ -93,6 +93,16 @@ TRDGScopeOpArray<ScopeOpType> TRDGScopeStackHelper<ScopeOpType>::EndCompile()
 	return TRDGScopeOpArray<ScopeOpType>(Ops, OffsetIndex, Ops.Num() - OffsetIndex);
 }
 
+// Lower overhead non-variadic version of constructor with arbitrary integer first argument to avoid overload resolution ambiguity.
+// Avoids dynamic allocation of the formatted string and other overhead.
+inline FRDGEventName::FRDGEventName(int32 NonVariadic, const TCHAR* InEventName)
+#if RDG_EVENTS == RDG_EVENTS_STRING_REF || RDG_EVENTS == RDG_EVENTS_STRING_COPY
+	: EventFormat(InEventName)
+#endif
+{
+	check(InEventName != nullptr);
+}
+
 #if RDG_EVENTS != RDG_EVENTS_STRING_COPY
 inline FRDGEventName::FRDGEventName(const TCHAR* InEventFormat, ...)
 #if RDG_EVENTS == RDG_EVENTS_STRING_REF
