Enable DXC by default for remaining SPIR-V backends.

diff --git a/Engine/Source/Runtime/RenderCore/Private/Shader.cpp b/Engine/Source/Runtime/RenderCore/Private/Shader.cpp
--- a/Engine/Source/Runtime/RenderCore/Private/Shader.cpp
+++ b/Engine/Source/Runtime/RenderCore/Private/Shader.cpp
@@ -108,10 +108,10 @@ static TAutoConsoleVariable<int32> CVarAllowCompilingThroughWorkers(
 
 static TAutoConsoleVariable<int32> CVarShadersForceDXC(
 	TEXT("r.Shaders.ForceDXC"),
-	0,
+	1,
 	TEXT("Forces DirectX Shader Compiler (DXC) to be used for all shaders instead of HLSLcc if supported.\n")
-	TEXT(" 0: Disable (default)\n")
-	TEXT(" 1: Force new compiler for all shaders"),
+	TEXT(" 1: Force new compiler for all shaders (default)\n")
+	TEXT(" 0: Disable"),
 	ECVF_ReadOnly);
 
 static TLinkedList<FShaderType*>*			GShaderTypeList = nullptr;
