Fixed PSO cache refcount check. Flush RDG command list setup tasks to the RHI thread prior to flushing the PSO cache.

diff --git a/Engine/Source/Runtime/RenderCore/Private/RenderGraphBuilder.cpp b/Engine/Source/Runtime/RenderCore/Private/RenderGraphBuilder.cpp
--- a/Engine/Source/Runtime/RenderCore/Private/RenderGraphBuilder.cpp
+++ b/Engine/Source/Runtime/RenderCore/Private/RenderGraphBuilder.cpp
@@ -1621,6 +1621,12 @@ void FRDGBuilder::WaitForParallelSetupTasks()
 		UE::Tasks::Wait(ParallelSetup.Tasks);
 		ParallelSetup.Tasks.Reset();
 	}
+
+	if (!ParallelSetup.CommandLists.IsEmpty())
+	{
+		RHICmdList.QueueAsyncCommandListSubmit(ParallelSetup.CommandLists);
+		ParallelSetup.CommandLists.Reset();
+	}
 }
 
 ///////////////////////////////////////////////////////////////////////////////////////////////////
@@ -1894,12 +1900,6 @@ void FRDGBuilder::Execute()
 		UE::Tasks::Launch(UE_SOURCE_LOCATION, [this] { ParallelExecute.DispatchTaskEvent->Trigger(); });
 	}
 
-	if (!ParallelSetup.CommandLists.IsEmpty())
-	{
-		RHICmdList.QueueAsyncCommandListSubmit(ParallelSetup.CommandLists);
-		ParallelSetup.CommandLists.Empty();
-	}
-
 	IF_RDG_ENABLE_DEBUG(GRDGAllowRHIAccess = ParallelExecute.bEnabled);
 	IF_RDG_ENABLE_TRACE(Trace.OutputGraphBegin());
 
