Fixed build break.

diff --git a/Engine/Source/Runtime/RenderCore/Public/RenderGraphAllocator.h b/Engine/Source/Runtime/RenderCore/Public/RenderGraphAllocator.h
--- a/Engine/Source/Runtime/RenderCore/Public/RenderGraphAllocator.h
+++ b/Engine/Source/Runtime/RenderCore/Public/RenderGraphAllocator.h
@@ -124,7 +124,7 @@ public:
 private:
 #if RDG_USE_MALLOC
 	TArray<void*> Mallocs;
-	uint32 NumMallocBytes = 0;
+	uint64 NumMallocBytes = 0;
 #else
 	FMemStackBase MemStack;
 #endif
