Fixed build break in RDG allocator.

diff --git a/Engine/Source/Runtime/RenderCore/Public/RenderGraphAllocator.h b/Engine/Source/Runtime/RenderCore/Public/RenderGraphAllocator.h
--- a/Engine/Source/Runtime/RenderCore/Public/RenderGraphAllocator.h
+++ b/Engine/Source/Runtime/RenderCore/Public/RenderGraphAllocator.h
@@ -113,7 +113,7 @@ public:
 	FORCEINLINE int32 GetByteCount() const
 	{
 #if RDG_USE_MALLOC
-		return NumMallocBytes;
+		return static_cast<int32>(NumMallocBytes);
 #else
 		return MemStack.GetByteCount();
 #endif
