Submit RDG setup command lists before flushing the RHI thread to avoid PSO refcounting issue.

diff --git a/Engine/Source/Runtime/RenderCore/Private/RenderGraphBuilder.cpp b/Engine/Source/Runtime/RenderCore/Private/RenderGraphBuilder.cpp
--- a/Engine/Source/Runtime/RenderCore/Private/RenderGraphBuilder.cpp
+++ b/Engine/Source/Runtime/RenderCore/Private/RenderGraphBuilder.cpp
@@ -1621,11 +1621,14 @@ void FRDGBuilder::WaitForParallelSetupTasks()
 		UE::Tasks::Wait(ParallelSetup.Tasks);
 		ParallelSetup.Tasks.Reset();
 	}
+}
 
+void FRDGBuilder::SubmitParallelSetupTasks()
+{
 	if (!ParallelSetup.CommandLists.IsEmpty())
 	{
 		RHICmdList.QueueAsyncCommandListSubmit(ParallelSetup.CommandLists);
-		ParallelSetup.CommandLists.Reset();
+		ParallelSetup.CommandLists.Empty();
 	}
 }
 
@@ -1663,9 +1666,9 @@ void FRDGBuilder::Execute()
 
 	if (!IsImmediateMode())
 	{
-		WaitForParallelSetupTasks();
-
+		SubmitParallelSetupTasks();
 		BeginFlushResourcesRHI();
+		WaitForParallelSetupTasks();
 
 		if (ParallelSetup.bEnabled)
 		{
@@ -1889,9 +1892,8 @@ void FRDGBuilder::Execute()
 		FinalizeResources();
 	}
 
-	// Process RHI thread flush before helping with barrier compilation on the render thread.
+	SubmitParallelSetupTasks();
 	EndFlushResourcesRHI();
-
 	WaitForParallelSetupTasks();
 
 	if (ParallelExecute.DispatchTaskEvent)
diff --git a/Engine/Source/Runtime/RenderCore/Public/RenderGraphBuilder.h b/Engine/Source/Runtime/RenderCore/Public/RenderGraphBuilder.h
--- a/Engine/Source/Runtime/RenderCore/Public/RenderGraphBuilder.h
+++ b/Engine/Source/Runtime/RenderCore/Public/RenderGraphBuilder.h
@@ -892,7 +892,6 @@ private:
 		TArray<UE::Tasks::FTask, FRDGArrayAllocator> Tasks;
 
 		/** Array of all command lists to submit for AddCommandListSetupTask. */
-		UE::FMutex CommandListsMutex;
 		TArray<FRHICommandListImmediate::FQueuedCommandList, FConcurrentLinearArrayAllocator> CommandLists;
 
 		bool bEnabled = false;
@@ -900,6 +899,7 @@ private:
 	} ParallelSetup;
 
 	void WaitForParallelSetupTasks();
+	void SubmitParallelSetupTasks();
 
 	/////////////////////////////////////////////////////////////////////////////
 	// Parallel Execution
diff --git a/Engine/Source/Runtime/RenderCore/Public/RenderGraphBuilder.inl b/Engine/Source/Runtime/RenderCore/Public/RenderGraphBuilder.inl
--- a/Engine/Source/Runtime/RenderCore/Public/RenderGraphBuilder.inl
+++ b/Engine/Source/Runtime/RenderCore/Public/RenderGraphBuilder.inl
@@ -623,15 +623,21 @@ UE::Tasks::FTask FRDGBuilder::AddCommandListSetupTask(
 	// Need a separate command list with inline tasks when prerequisites are not complete yet.
 	const bool bAllocateCommandListForTask = bCondition && (ParallelSetup.bEnabled || !UE::RDG::IsCompleted(Prerequisites));
 
-	auto OuterLambda = [this, TaskLambda = MoveTemp(TaskLambda), bAllocateCommandListForTask]() mutable
+	FRHICommandList* RHICmdListTask = &RHICmdList;
+
+	if (bAllocateCommandListForTask)
+	{
+		SCOPED_NAMED_EVENT(CreateCommandList, FColor::Emerald);
+		RHICmdListTask = new FRHICommandList(RHICmdList.GetGPUMask());
+		ParallelSetup.CommandLists.Emplace(RHICmdListTask);
+	}
+
+	auto OuterLambda = [this, TaskLambda = MoveTemp(TaskLambda), RHICmdListTask, bAllocateCommandListForTask]() mutable
 	{
 		FOptionalTaskTagScope Scope(ETaskTag::EParallelRenderingThread);
 
-		FRHICommandList* RHICmdListTask = &RHICmdList;
-		
 		if (bAllocateCommandListForTask)
 		{
-			RHICmdListTask = new FRHICommandList(RHICmdList.GetGPUMask());
 			RHICmdListTask->SwitchPipeline(ERHIPipeline::Graphics);
 		}
 
@@ -640,9 +646,6 @@ UE::Tasks::FTask FRDGBuilder::AddCommandListSetupTask(
 		if (bAllocateCommandListForTask)
 		{
 			RHICmdListTask->FinishRecording();
-
-			UE::TScopeLock Lock(ParallelSetup.CommandListsMutex);
-			ParallelSetup.CommandLists.Emplace(RHICmdListTask);
 		}
 	};
 
