Various changes for cross-platform implementation of FDispatchShaderBundleCS

* Use explicit FShaderParameter-s instead of shader parameter struct macros, as this allows the shader to be more easily dispatched from low-level RHI code
* Add total record count parameter to ProcessShaderBundleRecord
* Add CFLAG_ForceBindful to the shader when required by a specific platform

diff --git a/Engine/Shaders/Private/ShaderBundleDispatch.usf b/Engine/Shaders/Private/ShaderBundleDispatch.usf
--- a/Engine/Shaders/Private/ShaderBundleDispatch.usf
+++ b/Engine/Shaders/Private/ShaderBundleDispatch.usf
@@ -23,6 +23,7 @@ void DispatchShaderBundleEntry(uint RecordIndex : SV_DispatchThreadID)
 		ProcessShaderBundleRecord(
 			PlatformData,
 			RecordIndex,
+			RecordCount,
 			RecordArgBuffer,
 			RecordDataBuffer,
 			RWExecutionBuffer);
diff --git a/Engine/Source/Runtime/RenderCore/Private/ShaderBundles.cpp b/Engine/Source/Runtime/RenderCore/Private/ShaderBundles.cpp
--- a/Engine/Source/Runtime/RenderCore/Private/ShaderBundles.cpp
+++ b/Engine/Source/Runtime/RenderCore/Private/ShaderBundles.cpp
@@ -6,6 +6,7 @@
 #include "RenderGraphDefinitions.h"
 #include "ShaderParameterMacros.h"
 #include "RenderGraphFwd.h"
+#include "ShaderCompilerCore.h"
 
 IMPLEMENT_GLOBAL_SHADER(FDispatchShaderBundleCS, "/Engine/Private/ShaderBundleDispatch.usf", "DispatchShaderBundleEntry", SF_Compute);
 
@@ -16,6 +17,11 @@ bool FDispatchShaderBundleCS::ShouldCompilePermutation(const FGlobalShaderPermut
 
 void FDispatchShaderBundleCS::ModifyCompilationEnvironment(const FGlobalShaderPermutationParameters& Parameters, FShaderCompilerEnvironment& OutEnvironment)
 {
+	if (FDataDrivenShaderPlatformInfo::GetRequiresBindfulUtilityShaders(Parameters.Platform))
+	{
+		OutEnvironment.CompilerFlags.Add(CFLAG_ForceBindful);
+	}
+
 	OutEnvironment.SetDefine(TEXT("THREADGROUP_SIZEX"), ThreadGroupSizeX);
 	OutEnvironment.SetDefine(TEXT("USE_SHADER_ROOT_CONSTANTS"), RHISupportsShaderRootConstants(Parameters.Platform) ? 1 : 0);
 	FGlobalShader::ModifyCompilationEnvironment(Parameters, OutEnvironment);
@@ -37,19 +43,24 @@ void FDispatchShaderBundle::Dispatch(
 
 	auto ComputeShader = GetGlobalShaderMap(GMaxRHIFeatureLevel)->GetShader<FDispatchShaderBundleCS>();
 
-	FDispatchShaderBundleCS::FParameters Parameters;
-
-	Parameters.RecordCount = RecordCount;
-	Parameters.PlatformData = ShaderBundle->GetPlatformData();
-	Parameters.RecordArgBuffer = RecordArgBufferSRV;
-	Parameters.RecordDataBuffer = RecordDataBufferSRV;
-	Parameters.RWExecutionBuffer = ExecutionBufferUAV;
+	FRHIBatchedShaderParameters& Parameters = RHICmdList.GetScratchShaderParameters();
 
-	const FIntVector GroupCount = FComputeShaderUtils::GetGroupCount(RecordCount, FDispatchShaderBundleCS::ThreadGroupSizeX);
+	SetShaderValue(Parameters, ComputeShader->RecordCountParam, RecordCount);
+	SetShaderValue(Parameters, ComputeShader->PlatformDataParam, ShaderBundle->GetPlatformData());
+	SetSRVParameter(Parameters, ComputeShader->RecordArgBufferParam, RecordArgBufferSRV);
+	SetSRVParameter(Parameters, ComputeShader->RecordDataBufferParam, RecordDataBufferSRV);
+	SetUAVParameter(Parameters, ComputeShader->RWExecutionBufferParam, ExecutionBufferUAV);
 
 	FRHIComputeShader* ComputeShaderRHI = ComputeShader.GetComputeShader();
 	SetComputePipelineState(RHICmdList, ComputeShaderRHI);
-	SetShaderParameters(RHICmdList, ComputeShader, ComputeShaderRHI, Parameters);
+
+	RHICmdList.SetBatchedShaderParameters(ComputeShaderRHI, Parameters);
+
+	const FIntVector GroupCount = FComputeShaderUtils::GetGroupCount(RecordCount, FDispatchShaderBundleCS::ThreadGroupSizeX);
 	DispatchComputeShader(RHICmdList, ComputeShader, GroupCount.X, GroupCount.Y, GroupCount.Z);
-	UnsetShaderUAVs(RHICmdList, ComputeShader, ComputeShaderRHI);
+
+	FRHIBatchedShaderUnbinds& Unbinds = RHICmdList.GetScratchShaderUnbinds();
+
+	UnsetUAVParameter(Unbinds, ComputeShader->RWExecutionBufferParam);
+	RHICmdList.SetBatchedShaderUnbinds(ComputeShaderRHI, Unbinds);
 }
\ No newline at end of file
diff --git a/Engine/Source/Runtime/RenderCore/Public/ShaderBundles.h b/Engine/Source/Runtime/RenderCore/Public/ShaderBundles.h
--- a/Engine/Source/Runtime/RenderCore/Public/ShaderBundles.h
+++ b/Engine/Source/Runtime/RenderCore/Public/ShaderBundles.h
@@ -24,17 +24,25 @@ class FDispatchShaderBundleCS : public FGlobalShader
 	DECLARE_EXPORTED_GLOBAL_SHADER(FDispatchShaderBundleCS, RENDERCORE_API);
 
 public:
-	SHADER_USE_PARAMETER_STRUCT(FDispatchShaderBundleCS, FGlobalShader)
+
+	FDispatchShaderBundleCS() = default;
+	FDispatchShaderBundleCS(const ShaderMetaType::CompiledShaderInitializerType& Initializer)
+		: FGlobalShader(Initializer)
+	{
+		RecordCountParam.Bind(Initializer.ParameterMap, TEXT("RecordCount"), SPF_Mandatory);
+		PlatformDataParam.Bind(Initializer.ParameterMap, TEXT("PlatformData"), SPF_Mandatory);
+		RecordArgBufferParam.Bind(Initializer.ParameterMap, TEXT("RecordArgBuffer"), SPF_Mandatory);
+		RecordDataBufferParam.Bind(Initializer.ParameterMap, TEXT("RecordDataBuffer"), SPF_Mandatory);
+		RWExecutionBufferParam.Bind(Initializer.ParameterMap, TEXT("RWExecutionBuffer"), SPF_Mandatory);
+	}
 
 	static const uint32 ThreadGroupSizeX = 64;
 
-	BEGIN_SHADER_PARAMETER_STRUCT(FParameters, RENDERCORE_API)
-		SHADER_PARAMETER(uint32, RecordCount)
-		SHADER_PARAMETER(FUintVector4, PlatformData)
-		SHADER_PARAMETER_SRV(ByteAddressBuffer, RecordArgBuffer)
-		SHADER_PARAMETER_SRV(ByteAddressBuffer, RecordDataBuffer)
-		SHADER_PARAMETER_UAV(RWByteAddressBuffer, RWExecutionBuffer)
-	END_SHADER_PARAMETER_STRUCT()
+	LAYOUT_FIELD(FShaderParameter, RecordCountParam);
+	LAYOUT_FIELD(FShaderParameter, PlatformDataParam);
+	LAYOUT_FIELD(FShaderResourceParameter, RecordArgBufferParam);
+	LAYOUT_FIELD(FShaderResourceParameter, RecordDataBufferParam);
+	LAYOUT_FIELD(FShaderResourceParameter, RWExecutionBufferParam);
 
 	static bool ShouldCompilePermutation(const FGlobalShaderPermutationParameters& Parameters);
 	static void ModifyCompilationEnvironment(const FGlobalShaderPermutationParameters& Parameters, FShaderCompilerEnvironment& OutEnvironment);
