[shaders] Optimize the parameter metadata portion of DDC key construction. This is done by pre-computing a layout signature hash on initialization of each FShaderParameterMetadata struct and store it for use in DDC keys (rather than re-writing the layout to a memory buffer every time it's needed). This removes a massive amount of temporary allocations incurred by the unnecessarily complex FShaderSaveArchive mechanism, and has an outsized impact on editor startup time due to the usage of a portion of the DDC key string in constructing the unique name for materials (which is used in various places). In addition usage of the 32-bit LayoutHash has been removed from other places it was used in the shader DDC keys; this is a weak hash constructed by GetTypeHash which is not intended to be persisted. As part of this change the FShaderSaveArchive/SerializationHistory mechanism previously used here has been deprecated as this metadata hashing was the last place it was used.

diff --git a/Engine/Source/Runtime/RenderCore/Private/ShaderCore.cpp b/Engine/Source/Runtime/RenderCore/Private/ShaderCore.cpp
--- a/Engine/Source/Runtime/RenderCore/Private/ShaderCore.cpp
+++ b/Engine/Source/Runtime/RenderCore/Private/ShaderCore.cpp
@@ -3125,7 +3125,7 @@ void AppendKeyStringShaderDependencies(
 
 		if (const FShaderParametersMetadata* ParameterStructMetadata = ShaderType->GetRootParametersMetadata())
 		{
-			OutKeyString.Appendf(TEXT("%08x"), ParameterStructMetadata->GetLayoutHash());
+			ParameterStructMetadata->AppendKeyString(OutKeyString);
 		}
 
 		const FSHAHash LayoutHash = GetShaderTypeLayoutHash(ShaderType->GetLayout(), LayoutParams);
@@ -3155,7 +3155,7 @@ void AppendKeyStringShaderDependencies(
 		{
 			if (const FShaderParametersMetadata* ParameterStructMetadata = ShaderType->GetRootParametersMetadata())
 			{
-				OutKeyString.Appendf(TEXT("%08x"), ParameterStructMetadata->GetLayoutHash());
+				ParameterStructMetadata->AppendKeyString(OutKeyString);
 			}
 
 			for (const TCHAR* UniformBufferName : ShaderType->GetReferencedUniformBufferNames())
@@ -3195,11 +3195,6 @@ void AppendKeyStringShaderDependencies(
 	}
 
 	{
-		TArray<uint8> TempData;
-		FSerializationHistory SerializationHistory;
-		FMemoryWriter Ar(TempData, true);
-		FShaderSaveArchive SaveArchive(Ar, SerializationHistory);
-
 		TArray<const TCHAR*> SortedUniformBufferNames = ReferencedUniformBufferNames.Array();
 		Algo::Sort(SortedUniformBufferNames, FUniformBufferNameSortOrder());
 
@@ -3207,10 +3202,8 @@ void AppendKeyStringShaderDependencies(
 		for (const TCHAR* UniformBufferName : SortedUniformBufferNames)
 		{
 			FShaderParametersMetadata* UniformBufferMetadata = FindUniformBufferStructByName(UniformBufferName);
-			UniformBufferMetadata->SerializeLayout(SaveArchive);
+			UniformBufferMetadata->AppendKeyString(OutKeyString);
 		}
-
-		SerializationHistory.AppendKeyString(OutKeyString);
 	}
 }
 
@@ -3611,7 +3604,8 @@ FShaderCommonCompileJob::FInputHash FShaderCompileJob::GetInputHash()
 
 		if (Input.RootParametersStructure)
 		{
-			const_cast<FShaderParametersMetadata*>(Input.RootParametersStructure)->SerializeLayout(Hasher);
+			FBlake3Hash LayoutSignature = Input.RootParametersStructure->GetLayoutSignature();
+			Hasher << LayoutSignature;
 		}
 
 		InputHash = Hasher.Finalize();
@@ -3636,7 +3630,8 @@ FShaderCommonCompileJob::FInputHash FShaderCompileJob::GetInputHash()
 
 			if (Input.RootParametersStructure)
 			{
-				const_cast<FShaderParametersMetadata*>(Input.RootParametersStructure)->SerializeLayout(Archive);
+				FBlake3Hash LayoutSignature = Input.RootParametersStructure->GetLayoutSignature();
+				Archive << LayoutSignature;
 			}
 
 			// hash the source file so changes to files during the development are picked up
diff --git a/Engine/Source/Runtime/RenderCore/Private/ShaderParameterMetadata.cpp b/Engine/Source/Runtime/RenderCore/Private/ShaderParameterMetadata.cpp
--- a/Engine/Source/Runtime/RenderCore/Private/ShaderParameterMetadata.cpp
+++ b/Engine/Source/Runtime/RenderCore/Private/ShaderParameterMetadata.cpp
@@ -5,15 +5,17 @@
 =============================================================================*/
 
 #include "ShaderParameterMetadata.h"
+
+#include "DataDrivenShaderPlatformInfo.h"
+#include "Interfaces/IPluginManager.h"
+#include "Misc/CommandLine.h"
 #include "RenderCore.h"
 #include "RHIUniformBufferLayoutInitializer.h"
+#include "Serialization/MemoryHasher.h"
 #include "ShaderCore.h"
 #include "ShaderCompilerCore.h"
 #include "ShaderParameters.h"
-#include "DataDrivenShaderPlatformInfo.h"
 #include "ShaderParameterMacros.h"
-#include "Interfaces/IPluginManager.h"
-#include "Misc/CommandLine.h"
 
 bool SupportShaderPrecisionModifier(EShaderPlatform Platform)
 {
@@ -456,16 +458,15 @@ void FShaderParametersMetadata::InitializeAllUniformBufferStructs()
 	}
 }
 
-void FShaderParametersMetadata::FMember::SerializeLayout(FArchive& Ar)
+#if WITH_EDITOR
+
+void FShaderParametersMetadata::FMember::HashLayout(FMemoryHasherBlake3& Hasher)
 {
-	// This is only used at the moment for writing to a buffer for hashing; deserialize not supported
-	// due to various const members, but also not needed.
-	check(Ar.IsSaving());
+	Hasher << Offset;
+	Hasher << reinterpret_cast<uint8&>(BaseType);
 
-	Ar << Offset;
-	Ar << reinterpret_cast<uint8&>(BaseType);
-	Ar.Serialize(const_cast<TCHAR*>(Name), FCString::Strlen(Name));
-	Ar << NumElements;
+	Hasher.Serialize(const_cast<TCHAR*>(Name), FCString::Strlen(Name));
+	Hasher << NumElements;
 
 	const bool bIsRHIResource = (
 		BaseType == UBMT_TEXTURE ||
@@ -477,33 +478,28 @@ void FShaderParametersMetadata::FMember::SerializeLayout(FArchive& Ar)
 		BaseType == UBMT_UINT32 ||
 		BaseType == UBMT_FLOAT32)
 	{
-		Ar << reinterpret_cast<uint8&>(Precision);
-		Ar << NumRows;
-		Ar << NumColumns;
+		Hasher << reinterpret_cast<uint8&>(Precision);
+		Hasher << NumRows;
+		Hasher << NumColumns;
 	}
 	else if (BaseType == UBMT_INCLUDED_STRUCT || BaseType == UBMT_NESTED_STRUCT)
 	{
-		const_cast<FShaderParametersMetadata*>(Struct)->SerializeLayout(Ar);
+		const_cast<FShaderParametersMetadata*>(Struct)->HashLayout(Hasher);
 	}
 	else if (bIsRHIResource || bIsRDGResource)
 	{
-		Ar.Serialize(const_cast<TCHAR*>(ShaderType), FCString::Strlen(ShaderType));
+		Hasher.Serialize(const_cast<TCHAR*>(ShaderType), FCString::Strlen(ShaderType));
 	}
 }
 
-void FShaderParametersMetadata::SerializeLayout(FArchive& Ar) 
+void FShaderParametersMetadata::HashLayout(FMemoryHasherBlake3& SignatureData) 
 {
-	// This is only used at the moment for writing to a buffer for hashing; deserialize not supported
-	check(Ar.IsSaving());
-
-	Ar << const_cast<uint32&>(Size);
-
 	for (FMember& CurrentMember : Members)
 	{
-		CurrentMember.SerializeLayout(Ar);
+		CurrentMember.HashLayout(SignatureData);
 	}
 }
-
+#endif // WITH_EDITOR
 
 void FShaderParametersMetadata::InitializeLayout(FRHIUniformBufferLayoutInitializer* OutLayoutInitializer)
 {
@@ -856,6 +852,12 @@ void FShaderParametersMetadata::InitializeLayout(FRHIUniformBufferLayoutInitiali
 	}
 
 	Layout = RHICreateUniformBufferLayout(LayoutInitializer);
+
+#if WITH_EDITOR
+	FMemoryHasherBlake3 Hasher;
+	HashLayout(Hasher);
+	LayoutSignature = Hasher.Finalize();
+#endif
 }
 
 
diff --git a/Engine/Source/Runtime/RenderCore/Public/Shader.h b/Engine/Source/Runtime/RenderCore/Public/Shader.h
--- a/Engine/Source/Runtime/RenderCore/Public/Shader.h
+++ b/Engine/Source/Runtime/RenderCore/Public/Shader.h
@@ -567,7 +567,8 @@ public:
 };
 
 /** Encapsulates information about a shader's serialization behavior, used to detect when C++ serialization changes to auto-recompile. */
-class FSerializationHistory
+
+class UE_DEPRECATED(5.4, "FSerializationHistory is no longer used and will be removed") FSerializationHistory 
 {
 public: 
 
@@ -2426,164 +2427,48 @@ public:
 };
 
 /** Tracks state when traversing a FSerializationHistory. */
-class FSerializationHistoryTraversalState
+class UE_DEPRECATED(5.4, "FSerializationHistoryTraversalState is no longer used and will be removed") FSerializationHistoryTraversalState
 {
 public:
-
-	const FSerializationHistory& History;
-	int32 NextTokenIndex;
-	int32 NextFullLengthIndex;
-
-	FSerializationHistoryTraversalState(const FSerializationHistory& InHistory) :
-		History(InHistory),
-		NextTokenIndex(0),
-		NextFullLengthIndex(0)
+	PRAGMA_DISABLE_DEPRECATION_WARNINGS
+	FSerializationHistoryTraversalState(const FSerializationHistory& InHistory)
+	PRAGMA_ENABLE_DEPRECATION_WARNINGS
 	{}
 
 	/** Gets the length value from NextTokenIndex + Offset into history. */
 	uint32 GetValue(int32 Offset)
 	{
-		int32 CurrentOffset = Offset;
-
-		// Move to the desired offset
-		while (CurrentOffset > 0)
-		{
-			StepForward();
-			CurrentOffset--;
-		}
-
-		while (CurrentOffset < 0)
-		{
-			StepBackward();
-			CurrentOffset++;
-		}
-		check(CurrentOffset == 0);
-
-		// Decode
-		const int8 Token = History.GetToken(NextTokenIndex);
-		const uint32 Value = Token == 0 ? History.FullLengths[NextFullLengthIndex] : (int32)Token;
-
-		// Restore state
-		while (CurrentOffset < Offset)
-		{
-			StepBackward();
-			CurrentOffset++;
-		}
-
-		while (CurrentOffset > Offset)
-		{
-			StepForward();
-			CurrentOffset--;
-		}
-		check(CurrentOffset == Offset);
-
-		return Value;
+		return 0;
 	}
 
 	FORCEINLINE void StepForward()
 	{
-		const int8 Token = History.GetToken(NextTokenIndex);
-
-		if (Token == 0)
-		{
-			checkSlow(NextFullLengthIndex - 1 < History.FullLengths.Num());
-			NextFullLengthIndex++;
-		}
-
-		// Not supporting seeking past the front most serialization in the history
-		checkSlow(NextTokenIndex - 1 < History.NumTokens);
-		NextTokenIndex++;
 	}
 
 	void StepBackward()
 	{
-		// Not supporting seeking outside of the history tracked
-		check(NextTokenIndex > 0);
-		NextTokenIndex--;
-
-		const int8 Token = History.GetToken(NextTokenIndex);
-
-		if (Token == 0)
-		{
-			check(NextFullLengthIndex > 0);
-			NextFullLengthIndex--;
-		}
 	}
 };
 
 /** Archive used when saving shaders, which generates data used to detect serialization mismatches on load. */
-class FShaderSaveArchive final : public FArchiveProxy
+class UE_DEPRECATED(5.4, "FShaderSaveArchive is no longer used and will be removed") FShaderSaveArchive final : public FArchiveProxy
 {
 public:
 
+	PRAGMA_DISABLE_DEPRECATION_WARNINGS
 	FShaderSaveArchive(FArchive& Archive, FSerializationHistory& InHistory) : 
-		FArchiveProxy(Archive),
-		HistoryTraversalState(InHistory),
-		History(InHistory)
+		FArchiveProxy(Archive)
+	PRAGMA_ENABLE_DEPRECATION_WARNINGS
 	{
-		OriginalPosition = Archive.Tell();
 	}
 
 	virtual ~FShaderSaveArchive()
 	{
-		// Seek back to the original archive position so we can undo any serializations that went through this archive
-		InnerArchive.Seek(OriginalPosition);
 	}
 
 	virtual void Serialize( void* V, int64 Length )
 	{
-		if (HistoryTraversalState.NextTokenIndex < HistoryTraversalState.History.NumTokens)
-		{
-			// We are no longer appending (due to a seek), make sure writes match up in size with what's already been written
-			check(Length == HistoryTraversalState.GetValue(0));
-		}
-		else
-		{
-			// Appending to the archive, track the size of this serialization
-			check(Length >= 0 && Length <= TNumericLimits<uint32>::Max());
-			History.AddValue((uint32)Length);
-		}
-		HistoryTraversalState.StepForward();
-		
-		if (V)
-		{
-			FArchiveProxy::Serialize(V, Length);
-		}
 	}
-
-	virtual void Seek( int64 InPos )
-	{
-		int64 Offset = InPos - Tell();
-		if (Offset <= 0)
-		{
-			// We're seeking backward, walk backward through the serialization history while updating NextSerialization
-			while (Offset < 0)
-			{
-				Offset += HistoryTraversalState.GetValue(-1);
-				HistoryTraversalState.StepBackward();
-			}
-		}
-		else
-		{
-			// We're seeking forward, walk forward through the serialization history while updating NextSerialization
-			while (Offset > 0)
-			{
-				Offset -= HistoryTraversalState.GetValue(-1);
-				HistoryTraversalState.StepForward();
-			}
-			HistoryTraversalState.StepForward();
-		}
-		check(Offset == 0);
-		
-		FArchiveProxy::Seek(InPos);
-	}
-
-	FSerializationHistoryTraversalState HistoryTraversalState;
-	FSerializationHistory& History;
-
-private:
-	/** Stored off position of the original archive we are wrapping. */
-	int64 OriginalPosition;
 };
 
 /**
diff --git a/Engine/Source/Runtime/RenderCore/Public/ShaderParameterMetadata.h b/Engine/Source/Runtime/RenderCore/Public/ShaderParameterMetadata.h
--- a/Engine/Source/Runtime/RenderCore/Public/ShaderParameterMetadata.h
+++ b/Engine/Source/Runtime/RenderCore/Public/ShaderParameterMetadata.h
@@ -19,6 +19,7 @@
 #include "Misc/StringBuilder.h"
 #include "RHI.h"
 #include "RHIDefinitions.h"
+#include "Serialization/MemoryHasher.h"
 #include "Serialization/MemoryImage.h"
 #include "Serialization/MemoryLayout.h"
 #include "Templates/AlignmentTemplates.h"
@@ -251,8 +252,6 @@ public:
 			return ElementSize;
 		}
 
-		void SerializeLayout(FArchive& Ar);
-
 		static RENDERCORE_API void GenerateShaderParameterType(
 			FString& Result,
 			bool bSupportsPrecisionModifier,
@@ -265,6 +264,10 @@ public:
 		RENDERCORE_API void GenerateShaderParameterType(FString& Result, EShaderPlatform ShaderPlatform) const;
 
 	private:
+		friend class FShaderParametersMetadata;
+#if WITH_EDITOR
+		void HashLayout(FMemoryHasherBlake3& SignatureData);
+#endif
 
 		const TCHAR* Name;
 		const TCHAR* ShaderType;
@@ -391,7 +394,30 @@ public:
 		return LayoutHash;	
 	}
 
-	void SerializeLayout(FArchive& Ar);
+#if WITH_EDITOR
+	inline void AppendKeyString(FString& OutKeyString) const
+	{
+		TStringBuilder<sizeof(TCHAR) * (sizeof(FBlake3Hash::ByteArray) * 2 + 4)> StrBuilder;
+		StrBuilder << "SPM_";
+		StrBuilder << LayoutSignature;
+		OutKeyString.Append(StrBuilder.ToView());
+	}
+#endif
+
+	inline const FBlake3Hash& GetLayoutSignature() const
+	{
+#if WITH_EDITOR
+		check(IsLayoutInitialized());
+		return LayoutSignature;
+#else
+		// shader compilation types & WITH_EDITOR is a massive mess upstream; this should never actually be called outside of the editor
+		// but actually compiling the function out is a headache, so we instead just assert if it's called 
+		checkNoEntry();
+		static FBlake3Hash Dummy;
+		return Dummy;
+#endif
+	}
+
 
 	/** Iterate recursively over all FShaderParametersMetadata. */
 	template<typename TParameterFunction>
@@ -469,6 +495,13 @@ private:
 	/** Hash about the entire memory layout of the structure. */
 	uint32 LayoutHash = 0;
 
+#if WITH_EDITOR
+	void HashLayout(FMemoryHasherBlake3& SignatureData);
+	
+	/** Strong persistable hash representing the binary layout of the entire parameter structure */
+	FBlake3Hash LayoutSignature;
+#endif
+
 	/** Additional flags for how to use the buffer */
 	uint32 UsageFlags = 0;
 
