Fixed null pointer deref crash on Android. DepthBeforeState was NULL.

diff --git a/Engine/Source/Runtime/RenderCore/Private/RenderGraphBuilder.cpp b/Engine/Source/Runtime/RenderCore/Private/RenderGraphBuilder.cpp
--- a/Engine/Source/Runtime/RenderCore/Private/RenderGraphBuilder.cpp
+++ b/Engine/Source/Runtime/RenderCore/Private/RenderGraphBuilder.cpp
@@ -3736,7 +3736,7 @@ void FRDGBuilder::AddTextureTransition(FRDGTexture* Texture, FRDGTextureSubresou
 				StencilStateBefore->LastPass[GraphicsPipe] = MaxPassHandle;
 			}
 			// Case 2: transitioning out of a fused copy state.
-			else if (EnumHasAnyFlags(DepthStateBefore->Access, ERHIAccess::CopySrc | ERHIAccess::CopyDest))
+			else if (DepthStateBefore && EnumHasAnyFlags(DepthStateBefore->Access, ERHIAccess::CopySrc | ERHIAccess::CopyDest))
 			{
 				check(StencilStateBefore->Access        == DepthStateBefore->Access);
 				check(StencilStateBefore->GetLastPass() == DepthStateBefore->GetLastPass());
