Merge CL 30972041 from //UE5/Dev-ParallelRendering to //UE5/Main

Update engine version in UE_DEPRECATED macros for functions that were deprecated in //UE5/Dev-ParallelRendering
 - This work was originally planned for UE 5.4 but has been delayed to UE 5.5.

diff --git a/Engine/Source/Runtime/Core/Public/HAL/PlatformMisc.h b/Engine/Source/Runtime/Core/Public/HAL/PlatformMisc.h
--- a/Engine/Source/Runtime/Core/Public/HAL/PlatformMisc.h
+++ b/Engine/Source/Runtime/Core/Public/HAL/PlatformMisc.h
@@ -196,7 +196,7 @@ public:
 
 #endif
 
-#define SCOPED_PROFILER_COLOR(...) UE_DEPRECATED_MACRO(5.4, "SCOPED_PROFILER_COLOR is deprecated and there is no replacement.")
+#define SCOPED_PROFILER_COLOR(...) UE_DEPRECATED_MACRO(5.5, "SCOPED_PROFILER_COLOR is deprecated and there is no replacement.")
 
 // For timing OnEnterBackground tasks. This can be time sensitive on some platforms
 class FScopedEnterBackgroundEvent
diff --git a/Engine/Source/Runtime/RHI/Public/RHICommandList.h b/Engine/Source/Runtime/RHI/Public/RHICommandList.h
--- a/Engine/Source/Runtime/RHI/Public/RHICommandList.h
+++ b/Engine/Source/Runtime/RHI/Public/RHICommandList.h
@@ -434,7 +434,7 @@ public:
 	//
 	RHI_API void FinishRecording();
 
-	UE_DEPRECATED(5.4, "SetCurrentStat is deprecated and there is no replacement. Consider marking up rendering code with RDG event scopes or RHI breadcrumbs.")
+	UE_DEPRECATED(5.5, "SetCurrentStat is deprecated and there is no replacement. Consider marking up rendering code with RDG event scopes or RHI breadcrumbs.")
 	inline void SetCurrentStat(TStatId Stat) {}
 
 	FORCEINLINE_DEBUGGABLE void* Alloc(int64 AllocSize, int64 Alignment)
@@ -2961,12 +2961,12 @@ PRAGMA_ENABLE_DEPRECATION_WARNINGS
 	}
 #endif // WITH_RHI_BREADCRUMBS
 
-	UE_DEPRECATED(5.4, "RHIPushEvent is deprecated. All events and markers now use the RHI breadcrumb system. Use RDG_EVENT_SCOPE or SCOPED_DRAW_EVENT macros to mark up rendering code, rather than calling this function directly.")
+	UE_DEPRECATED(5.5, "RHIPushEvent is deprecated. All events and markers now use the RHI breadcrumb system. Use RDG_EVENT_SCOPE or SCOPED_DRAW_EVENT macros to mark up rendering code, rather than calling this function directly.")
 	FORCEINLINE_DEBUGGABLE void PushEvent(const TCHAR* Name, FColor Color)
 	{
 	}
 
-	UE_DEPRECATED(5.4, "RHIPopEvent is deprecated. All events and markers now use the RHI breadcrumb system. Use RDG_EVENT_SCOPE or SCOPED_DRAW_EVENT macros to mark up rendering code, rather than calling this function directly.")
+	UE_DEPRECATED(5.5, "RHIPopEvent is deprecated. All events and markers now use the RHI breadcrumb system. Use RDG_EVENT_SCOPE or SCOPED_DRAW_EVENT macros to mark up rendering code, rather than calling this function directly.")
 	FORCEINLINE_DEBUGGABLE void PopEvent()
 	{
 	}
@@ -4441,11 +4441,11 @@ public:
 		GDynamicRHI->RHIRead3DSurfaceFloatData(Texture,Rect,ZMinMax,OutData,Flags);
 	}
 	
-	UE_DEPRECATED(5.4, "FRHICommandListImmediate::AcquireThreadOwnership() is deprecated. Thread ownership of the RHI is automatic, so this function is redundant.")
+	UE_DEPRECATED(5.5, "FRHICommandListImmediate::AcquireThreadOwnership() is deprecated. Thread ownership of the RHI is automatic, so this function is redundant.")
 	FORCEINLINE void AcquireThreadOwnership()
 	{}
 	
-	UE_DEPRECATED(5.4, "FRHICommandListImmediate::ReleaseThreadOwnership() is deprecated. Thread ownership of the RHI is automatic, so this function is redundant.")
+	UE_DEPRECATED(5.5, "FRHICommandListImmediate::ReleaseThreadOwnership() is deprecated. Thread ownership of the RHI is automatic, so this function is redundant.")
 	FORCEINLINE void ReleaseThreadOwnership()
 	{}
 	
@@ -4457,7 +4457,7 @@ public:
 		return GDynamicRHI->RHIFlushResources();
 	}
 
-	UE_DEPRECATED(5.4, "FlushPendingDeletes is deprecated and removed. Use RHICmdList.ImmediateFlush(EImmediateFlushType::FlushRHIThreadFlushResources)")
+	UE_DEPRECATED(5.5, "FlushPendingDeletes is deprecated and removed. Use RHICmdList.ImmediateFlush(EImmediateFlushType::FlushRHIThreadFlushResources)")
 	inline int32 FlushPendingDeletes() { return 0; }
 	
 	FORCEINLINE uint32 GetGPUFrameCycles()
@@ -4550,7 +4550,7 @@ public:
 		return GDynamicRHI->RHIGetNativeCommandBuffer();
 	}
 
-	UE_DEPRECATED(5.4, "RHIPollRenderQueryResults is deprecated. Platform RHIs that require query polling now do this automatically as part of RHI command list submission.")
+	UE_DEPRECATED(5.5, "RHIPollRenderQueryResults is deprecated. Platform RHIs that require query polling now do this automatically as part of RHI command list submission.")
 	FORCEINLINE void PollRenderQueryResults() {}
 
 	/**
@@ -4832,7 +4832,7 @@ public:
 	//
 	// Blocks the calling thread until the RHI thread is idle.
 	//
-	UE_DEPRECATED(5.4, "FRHICommandListExecutor::WaitForRHIThreadTasks() is deprecated. Call FRHICommandListImmediate::ImmediateFlush(EImmediateFlushType::FlushRHIThread) instead.")
+	UE_DEPRECATED(5.5, "FRHICommandListExecutor::WaitForRHIThreadTasks() is deprecated. Call FRHICommandListImmediate::ImmediateFlush(EImmediateFlushType::FlushRHIThread) instead.")
 	inline void WaitForRHIThreadTasks()
 	{
 		FRHICommandListImmediate::Get().ImmediateFlush(EImmediateFlushType::FlushRHIThread);
@@ -4867,13 +4867,13 @@ public:
 	//
 	RHI_API static bool AreRHITasksActive();
 
-	UE_DEPRECATED(5.4, "FRHICommandListExecutor::IsRHIThreadActive() is deprecated. Use FRHICommandListExecutor::AreRHITasksActive() instead.")
+	UE_DEPRECATED(5.5, "FRHICommandListExecutor::IsRHIThreadActive() is deprecated. Use FRHICommandListExecutor::AreRHITasksActive() instead.")
 	static inline bool IsRHIThreadActive()
 	{
 		return AreRHITasksActive();
 	}
 
-	UE_DEPRECATED(5.4, "FRHICommandListExecutor::IsRHIThreadCompletelyFlushed() is deprecated. Use FRHICommandListExecutor::AreRHITasksActive() instead.")
+	UE_DEPRECATED(5.5, "FRHICommandListExecutor::IsRHIThreadCompletelyFlushed() is deprecated. Use FRHICommandListExecutor::AreRHITasksActive() instead.")
 	static inline bool IsRHIThreadCompletelyFlushed()
 	{
 		return !AreRHITasksActive();
@@ -5297,11 +5297,11 @@ FORCEINLINE void RHIUnlockTextureCubeFace(FRHITextureCube* Texture, uint32 FaceI
 	 FRHICommandListExecutor::GetImmediateCommandList().UnlockTextureCubeFace(Texture, FaceIndex, ArrayIndex, MipIndex, bLockWithinMiptail);
 }
 
-UE_DEPRECATED(5.4, "RHIAcquireThreadOwnership() is deprecated. Thread ownership of the RHI is automatic, so this function is redundant.")
+UE_DEPRECATED(5.5, "RHIAcquireThreadOwnership() is deprecated. Thread ownership of the RHI is automatic, so this function is redundant.")
 FORCEINLINE void RHIAcquireThreadOwnership()
 {}
 
-UE_DEPRECATED(5.4, "RHIReleaseThreadOwnership() is deprecated. Thread ownership of the RHI is automatic, so this function is redundant.")
+UE_DEPRECATED(5.5, "RHIReleaseThreadOwnership() is deprecated. Thread ownership of the RHI is automatic, so this function is redundant.")
 FORCEINLINE void RHIReleaseThreadOwnership()
 {}
 
diff --git a/Engine/Source/Runtime/RHI/Public/RHIResources.h b/Engine/Source/Runtime/RHI/Public/RHIResources.h
--- a/Engine/Source/Runtime/RHI/Public/RHIResources.h
+++ b/Engine/Source/Runtime/RHI/Public/RHIResources.h
@@ -88,7 +88,7 @@ public:
 	UE_DEPRECATED(5.3, "FlushPendingDeletes is deprecated, please use FRHICommandListExecutor::GetImmediateCommandList().ImmediateFlush(EImmediateFlushType::FlushRHIThreadFlushResources)")
 	static inline int32 FlushPendingDeletes(FRHICommandListImmediate& RHICmdList) { return 0; }
 
-	UE_DEPRECATED(5.4, "Don't call Bypass() on an FRHIResource. Use the Bypass() function on an FRHICommandList instance, or the FRHICommmandListExecutor.")
+	UE_DEPRECATED(5.5, "Don't call Bypass() on an FRHIResource. Use the Bypass() function on an FRHICommandList instance, or the FRHICommmandListExecutor.")
 	RHI_API static bool Bypass();
 
 	bool IsValid() const
diff --git a/Engine/Source/Runtime/RenderCore/Public/ProfilingDebugging/RealtimeGPUProfiler.h b/Engine/Source/Runtime/RenderCore/Public/ProfilingDebugging/RealtimeGPUProfiler.h
--- a/Engine/Source/Runtime/RenderCore/Public/ProfilingDebugging/RealtimeGPUProfiler.h
+++ b/Engine/Source/Runtime/RenderCore/Public/ProfilingDebugging/RealtimeGPUProfiler.h
@@ -110,47 +110,47 @@ class FRenderQueryPool;
 #define SCOPED_CONDITIONAL_DRAW_EVENTF_GAMETHREAD(Name, Condition, Format, ...)         RHI_BREADCRUMB_EVENTF_CONDITIONAL_GAMETHREAD(Name, Condition, Format, ##__VA_ARGS__)
 
 // Deprecated macros
-#define BEGIN_DRAW_EVENTF(RHICmdList, Name, Event, Format, ...)                                      UE_DEPRECATED_MACRO(5.4, "BEGIN_DRAW_EVENTF has been deprecated. Equivalent functionality can be implemented by constructing / destructing an instance of FRHIBreadcrumbEventManual."                      )
-#define STOP_DRAW_EVENT(Event)                                                                       UE_DEPRECATED_MACRO(5.4, "STOP_DRAW_EVENT has been deprecated. Equivalent functionality can be implemented by constructing / destructing an instance of FRHIBreadcrumbEventManual."                        )
-#define STOP_DRAW_EVENT_GAMETHREAD(...)                                                              UE_DEPRECATED_MACRO(5.4, "STOP_DRAW_EVENT_GAMETHREAD has been deprecated. Equivalent functionality can be implemented by constructing / destructing an instance of FRHIBreadcrumbEvent_GameThread."        )
-#define BEGIN_DRAW_EVENTF_GAMETHREAD(...)                                                            UE_DEPRECATED_MACRO(5.4, "BEGIN_DRAW_EVENTF_GAMETHREAD has been deprecated. Equivalent functionality can be implemented by constructing / destructing an instance of FRHIBreadcrumbEvent_GameThread."      )
-#define BEGIN_DRAW_EVENTF_COLOR_GAMETHREAD(...)                                                      UE_DEPRECATED_MACRO(5.4, "BEGIN_DRAW_EVENTF_COLOR_GAMETHREAD has been deprecated. Equivalent functionality can be implemented by constructing / destructing an instance of FRHIBreadcrumbEvent_GameThread.")
-#define SCOPED_DRAW_EVENT_COLOR(RHICmdList, Color, Name)                                             UE_DEPRECATED_MACRO(5.4, "SCOPED_DRAW_EVENT_COLOR has been deprecated. Use SCOPED_DRAW_EVENT instead."                                                ) SCOPED_DRAW_EVENT(RHICmdList, Name)
-#define SCOPED_GPU_EVENT(RHICmdList, Name)                                                           UE_DEPRECATED_MACRO(5.4, "SCOPED_GPU_EVENT has been deprecated. Use SCOPED_DRAW_EVENT instead."                                                       ) SCOPED_DRAW_EVENT(RHICmdList, Name)
-#define SCOPED_GPU_EVENT_COLOR(RHICmdList, Color, Name)                                              UE_DEPRECATED_MACRO(5.4, "SCOPED_GPU_EVENT_COLOR has been deprecated. Use SCOPED_DRAW_EVENT instead."                                                 ) SCOPED_DRAW_EVENT(RHICmdList, Name)
-#define SCOPED_COMPUTE_EVENT(RHICmdList, Name)                                                       UE_DEPRECATED_MACRO(5.4, "SCOPED_COMPUTE_EVENT has been deprecated. Use SCOPED_DRAW_EVENT instead."                                                   ) SCOPED_DRAW_EVENT(RHICmdList, Name)
-#define SCOPED_COMPUTE_EVENT_COLOR(RHICmdList, Color, Name)                                          UE_DEPRECATED_MACRO(5.4, "SCOPED_COMPUTE_EVENT_COLOR has been deprecated. Use SCOPED_DRAW_EVENT instead."                                             ) SCOPED_DRAW_EVENT(RHICmdList, Name)
-#define SCOPED_DRAW_EVENTF_COLOR(RHICmdList, Color, Name, Format, ...)                               UE_DEPRECATED_MACRO(5.4, "SCOPED_DRAW_EVENTF_COLOR has been deprecated. Use SCOPED_DRAW_EVENTF instead."                                              ) SCOPED_DRAW_EVENTF(RHICmdList, Name, Format, ##__VA_ARGS__)
-#define SCOPED_GPU_EVENTF(RHICmdList, Name, Format, ...)                                             UE_DEPRECATED_MACRO(5.4, "SCOPED_GPU_EVENTF has been deprecated. Use SCOPED_DRAW_EVENTF instead."                                                     ) SCOPED_DRAW_EVENTF(RHICmdList, Name, Format, ##__VA_ARGS__)
-#define SCOPED_GPU_EVENTF_COLOR(RHICmdList, Color, Name, Format, ...)                                UE_DEPRECATED_MACRO(5.4, "SCOPED_GPU_EVENTF_COLOR has been deprecated. Use SCOPED_DRAW_EVENTF instead."                                               ) SCOPED_DRAW_EVENTF(RHICmdList, Name, Format, ##__VA_ARGS__)
-#define SCOPED_COMPUTE_EVENTF(RHICmdList, Name, Format, ...)                                         UE_DEPRECATED_MACRO(5.4, "SCOPED_COMPUTE_EVENTF has been deprecated. Use SCOPED_DRAW_EVENTF instead."                                                 ) SCOPED_DRAW_EVENTF(RHICmdList, Name, Format, ##__VA_ARGS__)
-#define SCOPED_COMPUTE_EVENTF_COLOR(RHICmdList, Color, Name, Format, ...)                            UE_DEPRECATED_MACRO(5.4, "SCOPED_COMPUTE_EVENTF_COLOR has been deprecated. Use SCOPED_DRAW_EVENTF instead."                                           ) SCOPED_DRAW_EVENTF(RHICmdList, Name, Format, ##__VA_ARGS__)
-#define SCOPED_CONDITIONAL_DRAW_EVENT_COLOR(RHICmdList, Name, Color, Condition)                      UE_DEPRECATED_MACRO(5.4, "SCOPED_CONDITIONAL_DRAW_EVENT_COLOR has been deprecated. Use SCOPED_CONDITIONAL_DRAW_EVENT instead."                        ) SCOPED_CONDITIONAL_DRAW_EVENT(RHICmdList, Name, Condition)
-#define SCOPED_CONDITIONAL_GPU_EVENT(RHICmdList, Name, Condition)                                    UE_DEPRECATED_MACRO(5.4, "SCOPED_CONDITIONAL_GPU_EVENT has been deprecated. Use SCOPED_CONDITIONAL_DRAW_EVENT instead."                               ) SCOPED_CONDITIONAL_DRAW_EVENT(RHICmdList, Name, Condition)
-#define SCOPED_CONDITIONAL_GPU_EVENT_COLOR(RHICmdList, Name, Color, Condition)                       UE_DEPRECATED_MACRO(5.4, "SCOPED_CONDITIONAL_GPU_EVENT_COLOR has been deprecated. Use SCOPED_CONDITIONAL_DRAW_EVENT instead."                         ) SCOPED_CONDITIONAL_DRAW_EVENT(RHICmdList, Name, Condition)
-#define SCOPED_CONDITIONAL_COMPUTE_EVENT(RHICmdList, Name, Condition)                                UE_DEPRECATED_MACRO(5.4, "SCOPED_CONDITIONAL_COMPUTE_EVENT has been deprecated. Use SCOPED_CONDITIONAL_DRAW_EVENT instead."                           ) SCOPED_CONDITIONAL_DRAW_EVENT(RHICmdList, Name, Condition)
-#define SCOPED_CONDITIONAL_COMPUTE_EVENT_COLOR(RHICmdList, Name, Color, Condition)                   UE_DEPRECATED_MACRO(5.4, "SCOPED_CONDITIONAL_COMPUTE_EVENT_COLOR has been deprecated. Use SCOPED_CONDITIONAL_DRAW_EVENT instead."                     ) SCOPED_CONDITIONAL_DRAW_EVENT(RHICmdList, Name, Condition)
-#define SCOPED_CONDITIONAL_DRAW_EVENTF_COLOR(RHICmdList, Color, Name, Condition, Format, ...)        UE_DEPRECATED_MACRO(5.4, "SCOPED_CONDITIONAL_DRAW_EVENTF_COLOR has been deprecated. Use SCOPED_CONDITIONAL_DRAW_EVENTF instead."                      ) SCOPED_CONDITIONAL_DRAW_EVENTF(RHICmdList, Name, Condition, Format, ##__VA_ARGS__)
-#define SCOPED_CONDITIONAL_GPU_EVENTF(RHICmdList, Name, Condition, Format, ...)                      UE_DEPRECATED_MACRO(5.4, "SCOPED_CONDITIONAL_GPU_EVENTF has been deprecated. Use SCOPED_CONDITIONAL_DRAW_EVENTF instead."                             ) SCOPED_CONDITIONAL_DRAW_EVENTF(RHICmdList, Name, Condition, Format, ##__VA_ARGS__)
-#define SCOPED_CONDITIONAL_GPU_EVENTF_COLOR(RHICmdList, Color, Name, Condition, Format, ...)         UE_DEPRECATED_MACRO(5.4, "SCOPED_CONDITIONAL_GPU_EVENTF_COLOR has been deprecated. Use SCOPED_CONDITIONAL_DRAW_EVENTF instead."                       ) SCOPED_CONDITIONAL_DRAW_EVENTF(RHICmdList, Name, Condition, Format, ##__VA_ARGS__)
-#define SCOPED_CONDITIONAL_COMPUTE_EVENTF(RHICmdList, Name, Condition, Format, ...)                  UE_DEPRECATED_MACRO(5.4, "SCOPED_CONDITIONAL_COMPUTE_EVENTF has been deprecated. Use SCOPED_CONDITIONAL_DRAW_EVENTF instead."                         ) SCOPED_CONDITIONAL_DRAW_EVENTF(RHICmdList, Name, Condition, Format, ##__VA_ARGS__)
-#define SCOPED_CONDITIONAL_COMPUTE_EVENTF_COLOR(RHICmdList, Color, Name, Condition, Format, ...)     UE_DEPRECATED_MACRO(5.4, "SCOPED_CONDITIONAL_COMPUTE_EVENTF_COLOR has been deprecated. Use SCOPED_CONDITIONAL_DRAW_EVENTF instead."                   ) SCOPED_CONDITIONAL_DRAW_EVENTF(RHICmdList, Name, Condition, Format, ##__VA_ARGS__)
-#define BEGIN_DRAW_EVENTF_COLOR(RHICmdList, Color, Name, Event, Format, ...)                         UE_DEPRECATED_MACRO(5.4, "BEGIN_DRAW_EVENTF_COLOR has been deprecated. Use BEGIN_DRAW_EVENTF instead."                                                ) BEGIN_DRAW_EVENTF(RHICmdList, Name, Event, Format, ##__VA_ARGS__)
-#define BEGIN_GPU_EVENTF(RHICmdList, Name, Event, Format, ...)                                       UE_DEPRECATED_MACRO(5.4, "BEGIN_GPU_EVENTF has been deprecated. Use BEGIN_DRAW_EVENTF instead."                                                       ) BEGIN_DRAW_EVENTF(RHICmdList, Name, Event, Format, ##__VA_ARGS__)
-#define BEGIN_GPU_EVENTF_COLOR(RHICmdList, Color, Name, Event, Format, ...)                          UE_DEPRECATED_MACRO(5.4, "BEGIN_GPU_EVENTF_COLOR has been deprecated. Use BEGIN_DRAW_EVENTF instead."                                                 ) BEGIN_DRAW_EVENTF(RHICmdList, Name, Event, Format, ##__VA_ARGS__)
-#define STOP_GPU_EVENT(Event)                                                                        UE_DEPRECATED_MACRO(5.4, "STOP_GPU_EVENT has been deprecated. Use STOP_DRAW_EVENT instead."                                                           ) STOP_DRAW_EVENT(Event)
-#define SCOPED_DRAW_EVENT_COLOR_GAMETHREAD(Color, Name)                                              UE_DEPRECATED_MACRO(5.4, "SCOPED_DRAW_EVENT_COLOR_GAMETHREAD has been deprecated. Use SCOPED_DRAW_EVENT_GAMETHREAD instead."                          ) SCOPED_DRAW_EVENT_GAMETHREAD(Name)
-#define SCOPED_DRAW_EVENTF_COLOR_GAMETHREAD(Color, Name, Format, ...)                                UE_DEPRECATED_MACRO(5.4, "SCOPED_DRAW_EVENTF_COLOR_GAMETHREAD has been deprecated. Use SCOPED_DRAW_EVENTF_GAMETHREAD instead."                        ) SCOPED_DRAW_EVENTF_GAMETHREAD(Name, Format, ##__VA_ARGS__)
-#define SCOPED_CONDITIONAL_DRAW_EVENT_COLOR_GAMETHREAD(Name, Color, Condition)                       UE_DEPRECATED_MACRO(5.4, "SCOPED_CONDITIONAL_DRAW_EVENT_COLOR_GAMETHREAD has been deprecated. Use SCOPED_CONDITIONAL_DRAW_EVENT_GAMETHREAD instead."  ) SCOPED_CONDITIONAL_DRAW_EVENT_GAMETHREAD(Name, Condition)
-#define SCOPED_CONDITIONAL_DRAW_EVENTF_COLOR_GAMETHREAD(Color, Name, Condition, Format, ...)         UE_DEPRECATED_MACRO(5.4, "SCOPED_CONDITIONAL_DRAW_EVENTF_COLOR_GAMETHREAD has been deprecated. Use SCOPED_CONDITIONAL_DRAW_EVENTF_GAMETHREAD instead.") SCOPED_CONDITIONAL_DRAW_EVENTF_GAMETHREAD(Name, Condition, Format, ##__VA_ARGS__)
-#define SCOPED_RHI_DRAW_EVENT(RHICmdContext, Name)                                                   UE_DEPRECATED_MACRO(5.4, "SCOPED_RHI_DRAW_EVENT has been deprecated. Use standard RHI breadcrumb events instead."                   )
-#define SCOPED_RHI_DRAW_EVENTF(RHICmdContext, Name, Format, ...)                                     UE_DEPRECATED_MACRO(5.4, "SCOPED_RHI_DRAW_EVENTF has been deprecated. Use standard RHI breadcrumb events instead."                  )
-#define SCOPED_RHI_CONDITIONAL_DRAW_EVENT(RHICmdContext, Name, Condition)                            UE_DEPRECATED_MACRO(5.4, "SCOPED_RHI_CONDITIONAL_DRAW_EVENT has been deprecated. Use standard RHI breadcrumb events instead."       )
-#define SCOPED_RHI_CONDITIONAL_DRAW_EVENTF(RHICmdContext, Name, Condition, Format, ...)              UE_DEPRECATED_MACRO(5.4, "SCOPED_RHI_CONDITIONAL_DRAW_EVENTF has been deprecated. Use standard RHI breadcrumb events instead."      )
-#define SCOPED_RHI_DRAW_EVENT_COLOR(RHICmdContext, Color, Name)                                      UE_DEPRECATED_MACRO(5.4, "SCOPED_RHI_DRAW_EVENT_COLOR has been deprecated. Use standard RHI breadcrumb events instead."             )
-#define SCOPED_RHI_DRAW_EVENTF_COLOR(RHICmdContext, Color, Name, Format, ...)                        UE_DEPRECATED_MACRO(5.4, "SCOPED_RHI_DRAW_EVENTF_COLOR has been deprecated. Use standard RHI breadcrumb events instead."            )
-#define SCOPED_RHI_CONDITIONAL_DRAW_EVENT_COLOR(RHICmdContext, Color, Name, Condition)               UE_DEPRECATED_MACRO(5.4, "SCOPED_RHI_CONDITIONAL_DRAW_EVENT_COLOR has been deprecated. Use standard RHI breadcrumb events instead." )
-#define SCOPED_RHI_CONDITIONAL_DRAW_EVENTF_COLOR(RHICmdContext, Color, Name, Condition, Format, ...) UE_DEPRECATED_MACRO(5.4, "SCOPED_RHI_CONDITIONAL_DRAW_EVENTF_COLOR has been deprecated. Use standard RHI breadcrumb events instead.")
+#define BEGIN_DRAW_EVENTF(RHICmdList, Name, Event, Format, ...)                                      UE_DEPRECATED_MACRO(5.5, "BEGIN_DRAW_EVENTF has been deprecated. Equivalent functionality can be implemented by constructing / destructing an instance of FRHIBreadcrumbEventManual."                      )
+#define STOP_DRAW_EVENT(Event)                                                                       UE_DEPRECATED_MACRO(5.5, "STOP_DRAW_EVENT has been deprecated. Equivalent functionality can be implemented by constructing / destructing an instance of FRHIBreadcrumbEventManual."                        )
+#define STOP_DRAW_EVENT_GAMETHREAD(...)                                                              UE_DEPRECATED_MACRO(5.5, "STOP_DRAW_EVENT_GAMETHREAD has been deprecated. Equivalent functionality can be implemented by constructing / destructing an instance of FRHIBreadcrumbEvent_GameThread."        )
+#define BEGIN_DRAW_EVENTF_GAMETHREAD(...)                                                            UE_DEPRECATED_MACRO(5.5, "BEGIN_DRAW_EVENTF_GAMETHREAD has been deprecated. Equivalent functionality can be implemented by constructing / destructing an instance of FRHIBreadcrumbEvent_GameThread."      )
+#define BEGIN_DRAW_EVENTF_COLOR_GAMETHREAD(...)                                                      UE_DEPRECATED_MACRO(5.5, "BEGIN_DRAW_EVENTF_COLOR_GAMETHREAD has been deprecated. Equivalent functionality can be implemented by constructing / destructing an instance of FRHIBreadcrumbEvent_GameThread.")
+#define SCOPED_DRAW_EVENT_COLOR(RHICmdList, Color, Name)                                             UE_DEPRECATED_MACRO(5.5, "SCOPED_DRAW_EVENT_COLOR has been deprecated. Use SCOPED_DRAW_EVENT instead."                                                ) SCOPED_DRAW_EVENT(RHICmdList, Name)
+#define SCOPED_GPU_EVENT(RHICmdList, Name)                                                           UE_DEPRECATED_MACRO(5.5, "SCOPED_GPU_EVENT has been deprecated. Use SCOPED_DRAW_EVENT instead."                                                       ) SCOPED_DRAW_EVENT(RHICmdList, Name)
+#define SCOPED_GPU_EVENT_COLOR(RHICmdList, Color, Name)                                              UE_DEPRECATED_MACRO(5.5, "SCOPED_GPU_EVENT_COLOR has been deprecated. Use SCOPED_DRAW_EVENT instead."                                                 ) SCOPED_DRAW_EVENT(RHICmdList, Name)
+#define SCOPED_COMPUTE_EVENT(RHICmdList, Name)                                                       UE_DEPRECATED_MACRO(5.5, "SCOPED_COMPUTE_EVENT has been deprecated. Use SCOPED_DRAW_EVENT instead."                                                   ) SCOPED_DRAW_EVENT(RHICmdList, Name)
+#define SCOPED_COMPUTE_EVENT_COLOR(RHICmdList, Color, Name)                                          UE_DEPRECATED_MACRO(5.5, "SCOPED_COMPUTE_EVENT_COLOR has been deprecated. Use SCOPED_DRAW_EVENT instead."                                             ) SCOPED_DRAW_EVENT(RHICmdList, Name)
+#define SCOPED_DRAW_EVENTF_COLOR(RHICmdList, Color, Name, Format, ...)                               UE_DEPRECATED_MACRO(5.5, "SCOPED_DRAW_EVENTF_COLOR has been deprecated. Use SCOPED_DRAW_EVENTF instead."                                              ) SCOPED_DRAW_EVENTF(RHICmdList, Name, Format, ##__VA_ARGS__)
+#define SCOPED_GPU_EVENTF(RHICmdList, Name, Format, ...)                                             UE_DEPRECATED_MACRO(5.5, "SCOPED_GPU_EVENTF has been deprecated. Use SCOPED_DRAW_EVENTF instead."                                                     ) SCOPED_DRAW_EVENTF(RHICmdList, Name, Format, ##__VA_ARGS__)
+#define SCOPED_GPU_EVENTF_COLOR(RHICmdList, Color, Name, Format, ...)                                UE_DEPRECATED_MACRO(5.5, "SCOPED_GPU_EVENTF_COLOR has been deprecated. Use SCOPED_DRAW_EVENTF instead."                                               ) SCOPED_DRAW_EVENTF(RHICmdList, Name, Format, ##__VA_ARGS__)
+#define SCOPED_COMPUTE_EVENTF(RHICmdList, Name, Format, ...)                                         UE_DEPRECATED_MACRO(5.5, "SCOPED_COMPUTE_EVENTF has been deprecated. Use SCOPED_DRAW_EVENTF instead."                                                 ) SCOPED_DRAW_EVENTF(RHICmdList, Name, Format, ##__VA_ARGS__)
+#define SCOPED_COMPUTE_EVENTF_COLOR(RHICmdList, Color, Name, Format, ...)                            UE_DEPRECATED_MACRO(5.5, "SCOPED_COMPUTE_EVENTF_COLOR has been deprecated. Use SCOPED_DRAW_EVENTF instead."                                           ) SCOPED_DRAW_EVENTF(RHICmdList, Name, Format, ##__VA_ARGS__)
+#define SCOPED_CONDITIONAL_DRAW_EVENT_COLOR(RHICmdList, Name, Color, Condition)                      UE_DEPRECATED_MACRO(5.5, "SCOPED_CONDITIONAL_DRAW_EVENT_COLOR has been deprecated. Use SCOPED_CONDITIONAL_DRAW_EVENT instead."                        ) SCOPED_CONDITIONAL_DRAW_EVENT(RHICmdList, Name, Condition)
+#define SCOPED_CONDITIONAL_GPU_EVENT(RHICmdList, Name, Condition)                                    UE_DEPRECATED_MACRO(5.5, "SCOPED_CONDITIONAL_GPU_EVENT has been deprecated. Use SCOPED_CONDITIONAL_DRAW_EVENT instead."                               ) SCOPED_CONDITIONAL_DRAW_EVENT(RHICmdList, Name, Condition)
+#define SCOPED_CONDITIONAL_GPU_EVENT_COLOR(RHICmdList, Name, Color, Condition)                       UE_DEPRECATED_MACRO(5.5, "SCOPED_CONDITIONAL_GPU_EVENT_COLOR has been deprecated. Use SCOPED_CONDITIONAL_DRAW_EVENT instead."                         ) SCOPED_CONDITIONAL_DRAW_EVENT(RHICmdList, Name, Condition)
+#define SCOPED_CONDITIONAL_COMPUTE_EVENT(RHICmdList, Name, Condition)                                UE_DEPRECATED_MACRO(5.5, "SCOPED_CONDITIONAL_COMPUTE_EVENT has been deprecated. Use SCOPED_CONDITIONAL_DRAW_EVENT instead."                           ) SCOPED_CONDITIONAL_DRAW_EVENT(RHICmdList, Name, Condition)
+#define SCOPED_CONDITIONAL_COMPUTE_EVENT_COLOR(RHICmdList, Name, Color, Condition)                   UE_DEPRECATED_MACRO(5.5, "SCOPED_CONDITIONAL_COMPUTE_EVENT_COLOR has been deprecated. Use SCOPED_CONDITIONAL_DRAW_EVENT instead."                     ) SCOPED_CONDITIONAL_DRAW_EVENT(RHICmdList, Name, Condition)
+#define SCOPED_CONDITIONAL_DRAW_EVENTF_COLOR(RHICmdList, Color, Name, Condition, Format, ...)        UE_DEPRECATED_MACRO(5.5, "SCOPED_CONDITIONAL_DRAW_EVENTF_COLOR has been deprecated. Use SCOPED_CONDITIONAL_DRAW_EVENTF instead."                      ) SCOPED_CONDITIONAL_DRAW_EVENTF(RHICmdList, Name, Condition, Format, ##__VA_ARGS__)
+#define SCOPED_CONDITIONAL_GPU_EVENTF(RHICmdList, Name, Condition, Format, ...)                      UE_DEPRECATED_MACRO(5.5, "SCOPED_CONDITIONAL_GPU_EVENTF has been deprecated. Use SCOPED_CONDITIONAL_DRAW_EVENTF instead."                             ) SCOPED_CONDITIONAL_DRAW_EVENTF(RHICmdList, Name, Condition, Format, ##__VA_ARGS__)
+#define SCOPED_CONDITIONAL_GPU_EVENTF_COLOR(RHICmdList, Color, Name, Condition, Format, ...)         UE_DEPRECATED_MACRO(5.5, "SCOPED_CONDITIONAL_GPU_EVENTF_COLOR has been deprecated. Use SCOPED_CONDITIONAL_DRAW_EVENTF instead."                       ) SCOPED_CONDITIONAL_DRAW_EVENTF(RHICmdList, Name, Condition, Format, ##__VA_ARGS__)
+#define SCOPED_CONDITIONAL_COMPUTE_EVENTF(RHICmdList, Name, Condition, Format, ...)                  UE_DEPRECATED_MACRO(5.5, "SCOPED_CONDITIONAL_COMPUTE_EVENTF has been deprecated. Use SCOPED_CONDITIONAL_DRAW_EVENTF instead."                         ) SCOPED_CONDITIONAL_DRAW_EVENTF(RHICmdList, Name, Condition, Format, ##__VA_ARGS__)
+#define SCOPED_CONDITIONAL_COMPUTE_EVENTF_COLOR(RHICmdList, Color, Name, Condition, Format, ...)     UE_DEPRECATED_MACRO(5.5, "SCOPED_CONDITIONAL_COMPUTE_EVENTF_COLOR has been deprecated. Use SCOPED_CONDITIONAL_DRAW_EVENTF instead."                   ) SCOPED_CONDITIONAL_DRAW_EVENTF(RHICmdList, Name, Condition, Format, ##__VA_ARGS__)
+#define BEGIN_DRAW_EVENTF_COLOR(RHICmdList, Color, Name, Event, Format, ...)                         UE_DEPRECATED_MACRO(5.5, "BEGIN_DRAW_EVENTF_COLOR has been deprecated. Use BEGIN_DRAW_EVENTF instead."                                                ) BEGIN_DRAW_EVENTF(RHICmdList, Name, Event, Format, ##__VA_ARGS__)
+#define BEGIN_GPU_EVENTF(RHICmdList, Name, Event, Format, ...)                                       UE_DEPRECATED_MACRO(5.5, "BEGIN_GPU_EVENTF has been deprecated. Use BEGIN_DRAW_EVENTF instead."                                                       ) BEGIN_DRAW_EVENTF(RHICmdList, Name, Event, Format, ##__VA_ARGS__)
+#define BEGIN_GPU_EVENTF_COLOR(RHICmdList, Color, Name, Event, Format, ...)                          UE_DEPRECATED_MACRO(5.5, "BEGIN_GPU_EVENTF_COLOR has been deprecated. Use BEGIN_DRAW_EVENTF instead."                                                 ) BEGIN_DRAW_EVENTF(RHICmdList, Name, Event, Format, ##__VA_ARGS__)
+#define STOP_GPU_EVENT(Event)                                                                        UE_DEPRECATED_MACRO(5.5, "STOP_GPU_EVENT has been deprecated. Use STOP_DRAW_EVENT instead."                                                           ) STOP_DRAW_EVENT(Event)
+#define SCOPED_DRAW_EVENT_COLOR_GAMETHREAD(Color, Name)                                              UE_DEPRECATED_MACRO(5.5, "SCOPED_DRAW_EVENT_COLOR_GAMETHREAD has been deprecated. Use SCOPED_DRAW_EVENT_GAMETHREAD instead."                          ) SCOPED_DRAW_EVENT_GAMETHREAD(Name)
+#define SCOPED_DRAW_EVENTF_COLOR_GAMETHREAD(Color, Name, Format, ...)                                UE_DEPRECATED_MACRO(5.5, "SCOPED_DRAW_EVENTF_COLOR_GAMETHREAD has been deprecated. Use SCOPED_DRAW_EVENTF_GAMETHREAD instead."                        ) SCOPED_DRAW_EVENTF_GAMETHREAD(Name, Format, ##__VA_ARGS__)
+#define SCOPED_CONDITIONAL_DRAW_EVENT_COLOR_GAMETHREAD(Name, Color, Condition)                       UE_DEPRECATED_MACRO(5.5, "SCOPED_CONDITIONAL_DRAW_EVENT_COLOR_GAMETHREAD has been deprecated. Use SCOPED_CONDITIONAL_DRAW_EVENT_GAMETHREAD instead."  ) SCOPED_CONDITIONAL_DRAW_EVENT_GAMETHREAD(Name, Condition)
+#define SCOPED_CONDITIONAL_DRAW_EVENTF_COLOR_GAMETHREAD(Color, Name, Condition, Format, ...)         UE_DEPRECATED_MACRO(5.5, "SCOPED_CONDITIONAL_DRAW_EVENTF_COLOR_GAMETHREAD has been deprecated. Use SCOPED_CONDITIONAL_DRAW_EVENTF_GAMETHREAD instead.") SCOPED_CONDITIONAL_DRAW_EVENTF_GAMETHREAD(Name, Condition, Format, ##__VA_ARGS__)
+#define SCOPED_RHI_DRAW_EVENT(RHICmdContext, Name)                                                   UE_DEPRECATED_MACRO(5.5, "SCOPED_RHI_DRAW_EVENT has been deprecated. Use standard RHI breadcrumb events instead."                   )
+#define SCOPED_RHI_DRAW_EVENTF(RHICmdContext, Name, Format, ...)                                     UE_DEPRECATED_MACRO(5.5, "SCOPED_RHI_DRAW_EVENTF has been deprecated. Use standard RHI breadcrumb events instead."                  )
+#define SCOPED_RHI_CONDITIONAL_DRAW_EVENT(RHICmdContext, Name, Condition)                            UE_DEPRECATED_MACRO(5.5, "SCOPED_RHI_CONDITIONAL_DRAW_EVENT has been deprecated. Use standard RHI breadcrumb events instead."       )
+#define SCOPED_RHI_CONDITIONAL_DRAW_EVENTF(RHICmdContext, Name, Condition, Format, ...)              UE_DEPRECATED_MACRO(5.5, "SCOPED_RHI_CONDITIONAL_DRAW_EVENTF has been deprecated. Use standard RHI breadcrumb events instead."      )
+#define SCOPED_RHI_DRAW_EVENT_COLOR(RHICmdContext, Color, Name)                                      UE_DEPRECATED_MACRO(5.5, "SCOPED_RHI_DRAW_EVENT_COLOR has been deprecated. Use standard RHI breadcrumb events instead."             )
+#define SCOPED_RHI_DRAW_EVENTF_COLOR(RHICmdContext, Color, Name, Format, ...)                        UE_DEPRECATED_MACRO(5.5, "SCOPED_RHI_DRAW_EVENTF_COLOR has been deprecated. Use standard RHI breadcrumb events instead."            )
+#define SCOPED_RHI_CONDITIONAL_DRAW_EVENT_COLOR(RHICmdContext, Color, Name, Condition)               UE_DEPRECATED_MACRO(5.5, "SCOPED_RHI_CONDITIONAL_DRAW_EVENT_COLOR has been deprecated. Use standard RHI breadcrumb events instead." )
+#define SCOPED_RHI_CONDITIONAL_DRAW_EVENTF_COLOR(RHICmdContext, Color, Name, Condition, Format, ...) UE_DEPRECATED_MACRO(5.5, "SCOPED_RHI_CONDITIONAL_DRAW_EVENTF_COLOR has been deprecated. Use standard RHI breadcrumb events instead.")
 
 #if HAS_GPU_STATS
 
diff --git a/Engine/Source/Runtime/RenderCore/Public/RenderGraphBuilder.h b/Engine/Source/Runtime/RenderCore/Public/RenderGraphBuilder.h
--- a/Engine/Source/Runtime/RenderCore/Public/RenderGraphBuilder.h
+++ b/Engine/Source/Runtime/RenderCore/Public/RenderGraphBuilder.h
@@ -236,7 +236,7 @@ public:
 	RENDERCORE_API void AddPassDependency(FRDGPass* Producer, FRDGPass* Consumer);
 
 	/** Sets the current command list stat for all subsequent passes. */
-	UE_DEPRECATED(5.4, "SetCommandListStat is deprecated. The underlying stats have been removed. Consider marking up rendering code with RDG event scopes.")
+	UE_DEPRECATED(5.5, "SetCommandListStat is deprecated. The underlying stats have been removed. Consider marking up rendering code with RDG event scopes.")
 	inline void SetCommandListStat(TStatId StatId) {}
 
 	/** A hint to the builder to flush work to the RHI thread after the last queued pass on the execution timeline. */
diff --git a/Engine/Source/Runtime/Renderer/Private/SceneRendering.h b/Engine/Source/Runtime/Renderer/Private/SceneRendering.h
--- a/Engine/Source/Runtime/Renderer/Private/SceneRendering.h
+++ b/Engine/Source/Runtime/Renderer/Private/SceneRendering.h
@@ -574,7 +574,7 @@ protected:
 public:
 	FParallelCommandListSet(const FRDGPass* InPass, const FViewInfo& InView, FRHICommandListImmediate& InParentCmdList, bool bHasRenderPasses = true);
 
-	UE_DEPRECATED(5.4, "The 'ExecuteStat' has been removed from FParallelCommandListSet. Use the constructor that does not take a TStatId.")
+	UE_DEPRECATED(5.5, "The 'ExecuteStat' has been removed from FParallelCommandListSet. Use the constructor that does not take a TStatId.")
 	FParallelCommandListSet(const FRDGPass* InPass, TStatId InExecuteStat, const FViewInfo& InView, FRHICommandListImmediate& InParentCmdList, bool bHasRenderPasses = true)
 		: FParallelCommandListSet(InPass, InView, InParentCmdList, bHasRenderPasses)
 	{}
@@ -612,7 +612,7 @@ public:
 		, ViewportScale(InViewportScale)
 	{}
 
-	UE_DEPRECATED(5.4, "The 'ExecuteStat' has been removed from FParallelCommandListSet. Use the constructor that does not take a TStatId.")
+	UE_DEPRECATED(5.5, "The 'ExecuteStat' has been removed from FParallelCommandListSet. Use the constructor that does not take a TStatId.")
 	FRDGParallelCommandListSet(const FRDGPass* InPass, FRHICommandListImmediate& InParentCmdList, TStatId InStatId, const FViewInfo& InView, const FParallelCommandListBindings& InBindings, float InViewportScale = 1.0f)
 		: FRDGParallelCommandListSet(InPass, InParentCmdList, InView, InBindings, InViewportScale)
 	{}
