Temporarily disable the realtime GPU profiler
 - It's causing crashes when loading large levels in the editor: "Assertion failed: !OverflowEventCount".

diff --git a/Engine/Source/Runtime/RenderCore/Private/ProfilingDebugging/RealtimeGPUProfiler.cpp b/Engine/Source/Runtime/RenderCore/Private/ProfilingDebugging/RealtimeGPUProfiler.cpp
--- a/Engine/Source/Runtime/RenderCore/Private/ProfilingDebugging/RealtimeGPUProfiler.cpp
+++ b/Engine/Source/Runtime/RenderCore/Private/ProfilingDebugging/RealtimeGPUProfiler.cpp
@@ -19,7 +19,7 @@ CSV_DEFINE_CATEGORY_MODULE(RENDERCORE_API, GPU, true);
 
 static TAutoConsoleVariable<int> CVarGPUStatsEnabled(
 	TEXT("r.GPUStatsEnabled"),
-	1,
+	0,
 	TEXT("Enables or disables GPU stat recording"));
 
 
