Fixed race condition between RDG builder and async destruction task. Added a sync point with scene render clean-up to fix race condition with scene extension context.

diff --git a/Engine/Source/Runtime/RenderCore/Private/RenderGraphBuilder.cpp b/Engine/Source/Runtime/RenderCore/Private/RenderGraphBuilder.cpp
--- a/Engine/Source/Runtime/RenderCore/Private/RenderGraphBuilder.cpp
+++ b/Engine/Source/Runtime/RenderCore/Private/RenderGraphBuilder.cpp
@@ -385,6 +385,7 @@ void FRDGBuilder::EndFlushResourcesRHI()
 	}
 
 	CSV_SCOPED_TIMING_STAT_EXCLUSIVE(STAT_RDG_FlushResourcesRHI);
+	CSV_SCOPED_SET_WAIT_STAT(FlushResourcesRHI);
 	SCOPED_NAMED_EVENT(EndFlushResourcesRHI, FColor::Emerald);
 	RHICmdList.ImmediateFlush(EImmediateFlushType::FlushRHIThreadFlushResources);
 	PipelineStateCache::FlushResources();
@@ -606,11 +607,27 @@ FRDGBuilder::FRDGBuilder(FRHICommandListImmediate& InRHICmdList, FRDGEventName I
 #endif
 }
 
+UE::Tasks::FTask FRDGBuilder::FAsyncDeleter::LastTask;
+
+FRDGBuilder::FAsyncDeleter::~FAsyncDeleter()
+{
+	if (Function)
+	{
+		// Launch the task with a prerequisite on any previously launched RDG async delete task.
+		LastTask = UE::Tasks::Launch(UE_SOURCE_LOCATION, [Function = MoveTemp(Function)]() mutable {}, LastTask);
+	}
+}
+
+void FRDGBuilder::WaitForAsyncDeleteTask()
+{
+	FAsyncDeleter::LastTask.Wait();
+}
+
 FRDGBuilder::~FRDGBuilder()
 {
 	if (ParallelExecute.bEnabled && GRDGParallelDestruction > 0)
 	{
-		UE::Tasks::Launch(UE_SOURCE_LOCATION, [
+		AsyncDeleter.Function = [
 			Allocators				= MoveTemp(Allocators),
 			Passes					= MoveTemp(Passes),
 			Textures				= MoveTemp(Textures),
@@ -621,7 +638,7 @@ FRDGBuilder::~FRDGBuilder()
 			ActivePooledTextures	= MoveTemp(ActivePooledTextures),
 			ActivePooledBuffers		= MoveTemp(ActivePooledBuffers),
 			UploadedBuffers			= MoveTemp(UploadedBuffers)
-		] () mutable {});
+		] () mutable {};
 	}
 }
 
@@ -1639,6 +1656,7 @@ void FRDGBuilder::Execute()
 {
 	CSV_SCOPED_TIMING_STAT_EXCLUSIVE(RDG);
 	SCOPED_DRAW_EVENTF(RHICmdList, FRDGBuilder_Execute, TEXT("FRDGBuilder::Execute"));
+	CSV_SCOPED_SET_WAIT_STAT(RDG);
 
 #if WITH_RHI_BREADCRUMBS
 	check(LocalCurrentBreadcrumb == FRHIBreadcrumbNode::Sentinel);
@@ -4580,4 +4598,4 @@ void FRDGBuilder::ForceCopyCrossGPU()
 		RHICmdList.TransferResources(Transfers);
 	}
 }
-#endif  // WITH_MGPU
+#endif  // WITH_MGPU
\ No newline at end of file
diff --git a/Engine/Source/Runtime/RenderCore/Public/RenderGraphBuilder.h b/Engine/Source/Runtime/RenderCore/Public/RenderGraphBuilder.h
--- a/Engine/Source/Runtime/RenderCore/Public/RenderGraphBuilder.h
+++ b/Engine/Source/Runtime/RenderCore/Public/RenderGraphBuilder.h
@@ -47,6 +47,15 @@ struct TStatId;
 class FRDGBuilder
 	: public FRDGScopeState
 {
+	struct FAsyncDeleter
+	{
+		TUniqueFunction<void()> Function;
+		static UE::Tasks::FTask LastTask;
+
+		RENDERCORE_API ~FAsyncDeleter();
+
+	} AsyncDeleter;
+
 	struct
 	{
 		// Allocator for all root graph allocations on the graph builder thread.
@@ -402,6 +411,9 @@ public:
 	/** Whether RDG is running in immediate mode. */
 	static RENDERCORE_API bool IsImmediateMode();
 
+	/** Waits for the last RDG async delete task that was launched. */
+	static RENDERCORE_API void WaitForAsyncDeleteTask();
+
 	/** The blackboard used to hold common data tied to the graph lifetime. */
 	FRDGBlackboard Blackboard;
 
diff --git a/Engine/Source/Runtime/Renderer/Private/PlanarReflectionRendering.cpp b/Engine/Source/Runtime/Renderer/Private/PlanarReflectionRendering.cpp
--- a/Engine/Source/Runtime/Renderer/Private/PlanarReflectionRendering.cpp
+++ b/Engine/Source/Runtime/Renderer/Private/PlanarReflectionRendering.cpp
@@ -351,7 +351,7 @@ static void UpdatePlanarReflectionContents_RenderThread(
 	FDeferredUpdateResource::UpdateResources(RHICmdList);
 
 	const ERHIFeatureLevel::Type FeatureLevel = SceneRenderer->FeatureLevel;
-	FRDGBuilder GraphBuilder(RHICmdList, RDG_EVENT_NAME("PlanarReflection"), FSceneRenderer::GetRDGParalelExecuteFlags(FeatureLevel));
+	FRDGBuilder GraphBuilder(RHICmdList, RDG_EVENT_NAME("PlanarReflection"), ERDGBuilderFlags::AllowParallelExecute);
 
 	// Make sure we render to the same set of GPUs as the main scene renderer.
 	if (MainSceneRenderer->ViewFamily.RenderTarget != nullptr)
diff --git a/Engine/Source/Runtime/Renderer/Private/ReflectionEnvironmentCapture.cpp b/Engine/Source/Runtime/Renderer/Private/ReflectionEnvironmentCapture.cpp
--- a/Engine/Source/Runtime/Renderer/Private/ReflectionEnvironmentCapture.cpp
+++ b/Engine/Source/Runtime/Renderer/Private/ReflectionEnvironmentCapture.cpp
@@ -511,7 +511,7 @@ void CaptureSceneToScratchCubemap(
 	// update any resources that needed a deferred update
 	FDeferredUpdateResource::UpdateResources(RHICmdList);
 
-	FRDGBuilder GraphBuilder(RHICmdList, RDG_EVENT_NAME("CubeMapCapture"), FSceneRenderer::GetRDGParalelExecuteFlags(FeatureLevel));
+	FRDGBuilder GraphBuilder(RHICmdList, RDG_EVENT_NAME("CubeMapCapture"), ERDGBuilderFlags::AllowParallelExecute);
 
 	{
 		RDG_EVENT_SCOPE(GraphBuilder, "CubeMapCapture");
diff --git a/Engine/Source/Runtime/Renderer/Private/SceneCaptureRendering.cpp b/Engine/Source/Runtime/Renderer/Private/SceneCaptureRendering.cpp
--- a/Engine/Source/Runtime/Renderer/Private/SceneCaptureRendering.cpp
+++ b/Engine/Source/Runtime/Renderer/Private/SceneCaptureRendering.cpp
@@ -327,10 +327,10 @@ static void UpdateSceneCaptureContentDeferred_RenderThread(
 
 #if WANTS_DRAW_MESH_EVENTS
 	SCOPED_DRAW_EVENTF(RHICmdList, SceneCapture, TEXT("SceneCapture %s"), EventName);
-	FRDGBuilder GraphBuilder(RHICmdList, RDG_EVENT_NAME("SceneCapture %s", *EventName), FSceneRenderer::GetRDGParalelExecuteFlags(FeatureLevel));
+	FRDGBuilder GraphBuilder(RHICmdList, RDG_EVENT_NAME("SceneCapture %s", *EventName), ERDGBuilderFlags::AllowParallelExecute);
 #else
 	SCOPED_DRAW_EVENT(RHICmdList, UpdateSceneCaptureContent_RenderThread);
-	FRDGBuilder GraphBuilder(RHICmdList, RDG_EVENT_NAME("SceneCapture"), FSceneRenderer::GetRDGParalelExecuteFlags(FeatureLevel));
+	FRDGBuilder GraphBuilder(RHICmdList, RDG_EVENT_NAME("SceneCapture"), ERDGBuilderFlags::AllowParallelExecute);
 #endif
 
 	{
diff --git a/Engine/Source/Runtime/Renderer/Private/SceneRendering.cpp b/Engine/Source/Runtime/Renderer/Private/SceneRendering.cpp
--- a/Engine/Source/Runtime/Renderer/Private/SceneRendering.cpp
+++ b/Engine/Source/Runtime/Renderer/Private/SceneRendering.cpp
@@ -4334,6 +4334,7 @@ static void FinishCleanUp(FRHICommandListImmediate& RHICmdList)
 	GlobalDynamicBuffer::GarbageCollect();
 	GPrimitiveIdVertexBufferPool.DiscardAll();
 	FGraphicsMinimalPipelineStateId::ResetLocalPipelineIdTableSize();
+	FRDGBuilder::WaitForAsyncDeleteTask();
 }
 
 static void DeleteSceneRenderers(const TArray<FSceneRenderer*>& SceneRenderers, FParallelMeshDrawCommandPass::EWaitThread WaitThread)
@@ -4610,7 +4611,7 @@ static void RenderViewFamilies_RenderThread(FRHICommandListImmediate& RHICmdList
 				ViewFamily.EngineShowFlags.HitProxies ? TEXT("RenderHitProxies") : TEXT("Render"),
 				ViewFamily.bResolveScene ? TEXT("Primary") : TEXT("Auxiliary")
 			),
-			FSceneRenderer::GetRDGParalelExecuteFlags(FeatureLevel)
+			ERDGBuilderFlags::AllowParallelExecute
 		);
 
 #if WITH_GPUDEBUGCRASH
diff --git a/Engine/Source/Runtime/Renderer/Private/SceneRendering.h b/Engine/Source/Runtime/Renderer/Private/SceneRendering.h
--- a/Engine/Source/Runtime/Renderer/Private/SceneRendering.h
+++ b/Engine/Source/Runtime/Renderer/Private/SceneRendering.h
@@ -2280,12 +2280,6 @@ public:
 	/** Gets a readable light name for use with a draw event. */
 	static void GetLightNameForDrawEvent(const FLightSceneProxy* LightProxy, FString& LightNameWithLevel);
 
-	/** Get the ParallelExecuteFlags depending on FeatureLevel*/
-	static ERDGBuilderFlags GetRDGParalelExecuteFlags(ERHIFeatureLevel::Type FeatureLevel)
-	{
-		return ERDGBuilderFlags::AllowParallelExecute;
-	}
-
 	FORCEINLINE FSceneTextures& GetActiveSceneTextures() { return ViewFamily.GetSceneTextures(); }
 	FORCEINLINE FSceneTexturesConfig& GetActiveSceneTexturesConfig() { return ViewFamily.SceneTexturesConfig; }
 
