Fixed race condition between RDG builder and async destruction task. Added a sync point with scene render clean-up to fix race condition with scene extension context.

diff --git a/Engine/Source/Runtime/RenderCore/Private/RenderGraphBuilder.cpp b/Engine/Source/Runtime/RenderCore/Private/RenderGraphBuilder.cpp
--- a/Engine/Source/Runtime/RenderCore/Private/RenderGraphBuilder.cpp
+++ b/Engine/Source/Runtime/RenderCore/Private/RenderGraphBuilder.cpp
@@ -602,6 +602,15 @@ FRDGBuilder::FRDGBuilder(FRHICommandListImmediate& InRHICmdList, FRDGEventName I
 		bSupportsTransientBuffers  = TransientResourceAllocator->SupportsResourceType(ERHITransientResourceType::Buffer);
 	}
 
+#endif
+
+#if RHI_WANT_BREADCRUMB_EVENTS
+	if (ParallelExecute.bEnabled)
+	{
+		BreadcrumbState = FRDGBreadcrumbState::Create(Allocators.Root);
+	}
+#endif
+
 #if RDG_DUMP_RESOURCES
 	DumpNewGraphBuilder();
 #endif
@@ -4598,4 +4607,4 @@ void FRDGBuilder::ForceCopyCrossGPU()
 		RHICmdList.TransferResources(Transfers);
 	}
 }
-#endif  // WITH_MGPU
\ No newline at end of file
+#endif  // WITH_MGPU
