Fixed build break.

diff --git a/Engine/Source/Runtime/RenderCore/Private/RenderGraphBuilder.cpp b/Engine/Source/Runtime/RenderCore/Private/RenderGraphBuilder.cpp
--- a/Engine/Source/Runtime/RenderCore/Private/RenderGraphBuilder.cpp
+++ b/Engine/Source/Runtime/RenderCore/Private/RenderGraphBuilder.cpp
@@ -1667,11 +1667,6 @@ void FRDGBuilder::Execute()
 	SCOPED_DRAW_EVENTF(RHICmdList, FRDGBuilder_Execute, TEXT("FRDGBuilder::Execute"));
 	CSV_SCOPED_SET_WAIT_STAT(RDG);
 
-#if WITH_RHI_BREADCRUMBS
-	check(LocalCurrentBreadcrumb == FRHIBreadcrumbNode::Sentinel);
-	LocalCurrentBreadcrumb = RHICmdList.GetCurrentBreadcrumbRef();
-#endif
-
 	GRDGTransientResourceAllocator.ReleasePendingDeallocations();
 
 	{
