Fixed bad robomerge.

diff --git a/Engine/Source/Runtime/RenderCore/Private/RenderGraphBuilder.cpp b/Engine/Source/Runtime/RenderCore/Private/RenderGraphBuilder.cpp
--- a/Engine/Source/Runtime/RenderCore/Private/RenderGraphBuilder.cpp
+++ b/Engine/Source/Runtime/RenderCore/Private/RenderGraphBuilder.cpp
@@ -602,15 +602,6 @@ FRDGBuilder::FRDGBuilder(FRHICommandListImmediate& InRHICmdList, FRDGEventName I
 		bSupportsTransientBuffers  = TransientResourceAllocator->SupportsResourceType(ERHITransientResourceType::Buffer);
 	}
 
-#endif
-
-#if RHI_WANT_BREADCRUMB_EVENTS
-	if (ParallelExecute.bEnabled)
-	{
-		BreadcrumbState = FRDGBreadcrumbState::Create(Allocators.Root);
-	}
-#endif
-
 #if RDG_DUMP_RESOURCES
 	DumpNewGraphBuilder();
 #endif
@@ -1667,6 +1658,11 @@ void FRDGBuilder::Execute()
 	SCOPED_DRAW_EVENTF(RHICmdList, FRDGBuilder_Execute, TEXT("FRDGBuilder::Execute"));
 	CSV_SCOPED_SET_WAIT_STAT(RDG);
 
+#if WITH_RHI_BREADCRUMBS
+	check(LocalCurrentBreadcrumb == FRHIBreadcrumbNode::Sentinel);
+	LocalCurrentBreadcrumb = RHICmdList.GetCurrentBreadcrumbRef();
+#endif
+
 	GRDGTransientResourceAllocator.ReleasePendingDeallocations();
 
 	{
