RHI pool resource debug name optimization:  Opportunistically skip updating the debug name if it's the same as it was last time the pooled resource (Buffer or Texture) was used.  In testing, 82% of buffers and 90% of textures fall under this category, and the optimization saved 2% frame wide.  In the future, we could consider an exhaustive search or hash lookup by debug name, but this works well enough for now.

#rnx

diff --git a/Engine/Source/Runtime/RHI/Public/RHIResources.h b/Engine/Source/Runtime/RHI/Public/RHIResources.h
--- a/Engine/Source/Runtime/RHI/Public/RHIResources.h
+++ b/Engine/Source/Runtime/RHI/Public/RHIResources.h
@@ -4767,6 +4767,10 @@ public:
 #else
 	void SetDebugName(FRHICommandListBase& RHICmdList, const TCHAR* DebugName) {}
 #endif
+	FORCEINLINE int32 NumItems() const
+	{
+		return UAVs.Num() + SRVs.Num();
+	}
 
 private:
 	TArray<TPair<FRHIBufferUAVCreateInfo, FUnorderedAccessViewRHIRef>, TInlineAllocator<1>> UAVs;
diff --git a/Engine/Source/Runtime/RenderCore/Private/RenderGraphResourcePool.cpp b/Engine/Source/Runtime/RenderCore/Private/RenderGraphResourcePool.cpp
--- a/Engine/Source/Runtime/RenderCore/Private/RenderGraphResourcePool.cpp
+++ b/Engine/Source/Runtime/RenderCore/Private/RenderGraphResourcePool.cpp
@@ -102,12 +102,30 @@ TRefCountPtr<FRDGPooledBuffer> FRDGBufferPool::FindFreeBuffer(FRHICommandListBas
 		check(PooledBuffer->GetAlignedDesc() == AlignedDesc);
 
 		PooledBuffer->LastUsedFrame = FrameCounter;
-		PooledBuffer->ViewCache.SetDebugName(RHICmdList, InDebugName);
+
+#if (UE_BUILD_SHIPPING || UE_BUILD_TEST)
 		PooledBuffer->Name = InDebugName;
+#else
+		// For performance, avoid updating name if it happens to be the same (true 80% of the time in testing)
+		bool bNameUpdated = false;
+		if (!PooledBuffer->Name || FCString::Strcmp(PooledBuffer->Name, InDebugName))
+		{
+			// Name changed, need to update it
+			PooledBuffer->Name = InDebugName;
 
-	#if !(UE_BUILD_SHIPPING || UE_BUILD_TEST)
-		RHICmdList.BindDebugLabelName(PooledBuffer->GetRHI(), InDebugName);
-	#endif
+			RHICmdList.BindDebugLabelName(PooledBuffer->GetRHI(), InDebugName);
+
+			bNameUpdated = true;
+		}
+
+		// Propagate the debug name to ViewCache if the name was updated, or if any items were added to ViewCache since the debug name was set
+		int32 ViewCacheNum = PooledBuffer->ViewCache.NumItems();
+		if (bNameUpdated || PooledBuffer->NameUpdatedViewCacheNum != ViewCacheNum)
+		{
+			PooledBuffer->NameUpdatedViewCacheNum = ViewCacheNum;
+			PooledBuffer->ViewCache.SetDebugName(RHICmdList, InDebugName);
+		}
+#endif
 
 		// We need the external-facing desc to match what the user requested.
 		const_cast<FRDGBufferDesc&>(PooledBuffer->Desc).NumElements = Desc.NumElements;
diff --git a/Engine/Source/Runtime/RenderCore/Private/RenderTargetPool.cpp b/Engine/Source/Runtime/RenderCore/Private/RenderTargetPool.cpp
--- a/Engine/Source/Runtime/RenderCore/Private/RenderTargetPool.cpp
+++ b/Engine/Source/Runtime/RenderCore/Private/RenderTargetPool.cpp
@@ -153,16 +153,21 @@ TRefCountPtr<IPooledRenderTarget> FRenderTargetPool::FindFreeElement(FRHICommand
 		TRACE_COUNTER_SET(RenderTargetPoolSize, (int64)AllocationLevelInKB * 1024);
 
 		FoundIndex = PooledRenderTargets.Num() - 1;
-		Found->Desc.DebugName = Name;
 	}
 
+#if (UE_BUILD_SHIPPING || UE_BUILD_TEST)
 	Found->Desc.DebugName = Name;
-	Found->UnusedForNFrames = 0;
-
-#if !(UE_BUILD_SHIPPING || UE_BUILD_TEST)
-	RHICmdList.BindDebugLabelName(Found->GetRHI(), Name);
+#else
+	// For performance, avoid updating name if it happens to be the same (true 90% of the time in testing)
+	if (FCString::Strcmp(Found->Desc.DebugName, Name))
+	{
+		Found->Desc.DebugName = Name;
+		RHICmdList.BindDebugLabelName(Found->GetRHI(), Name);
+	}
 #endif
 
+	Found->UnusedForNFrames = 0;
+
 	return TRefCountPtr<IPooledRenderTarget>(MoveTemp(Found));
 }
 
diff --git a/Engine/Source/Runtime/RenderCore/Public/RenderGraphResources.h b/Engine/Source/Runtime/RenderCore/Public/RenderGraphResources.h
--- a/Engine/Source/Runtime/RenderCore/Public/RenderGraphResources.h
+++ b/Engine/Source/Runtime/RenderCore/Public/RenderGraphResources.h
@@ -1227,6 +1227,11 @@ private:
 	// May be UINT64_MAX for regular (non-reserved) buffers or when the entire resource is committed.
 	uint64 CommittedSizeInBytes = UINT64_MAX;
 
+#if !(UE_BUILD_SHIPPING || UE_BUILD_TEST)
+	// Tracks the number of items in ViewCache when the debug Name was updated, so we know if items were added and we need to propagate the debug name to the new items
+	int32 NameUpdatedViewCacheNum = 0;
+#endif
+
 	const uint32 NumAllocatedElements;
 	uint32 LastUsedFrame = 0;
 
