Fixed RHI validation error when transitioning back to ERHIAccess::Discard.

diff --git a/Engine/Source/Runtime/RenderCore/Private/RenderGraphBuilder.cpp b/Engine/Source/Runtime/RenderCore/Private/RenderGraphBuilder.cpp
--- a/Engine/Source/Runtime/RenderCore/Private/RenderGraphBuilder.cpp
+++ b/Engine/Source/Runtime/RenderCore/Private/RenderGraphBuilder.cpp
@@ -3628,6 +3628,7 @@ void FRDGBuilder::AddLastTextureTransition(FRDGTexture* Texture)
 
 	const FRDGPassHandle EpiloguePassHandle = GetEpiloguePassHandle();
 
+	FRDGSubresourceState* SubresourceStateBefore = nullptr;
 	FRDGSubresourceState& SubresourceStateAfter = *AllocSubresource();
 	SubresourceStateAfter.SetPass(ERHIPipeline::Graphics, EpiloguePassHandle);
 
@@ -3641,21 +3642,25 @@ void FRDGBuilder::AddLastTextureTransition(FRDGTexture* Texture)
 	}
 	else
 	{
-		FRDGSubresourceState& SubresourceStateBefore = *AllocSubresource();
-		SubresourceStateBefore.SetPass(ERHIPipeline::Graphics, Texture->FirstPass);
-
 		SubresourceStateAfter.Access = Texture->EpilogueAccess;
 
-		// Transition any unused (null) sub-resources to the epilogue state since we are assigning a monolithic state across all subresources.
-		for (FRDGSubresourceState*& State : Texture->State)
+		// Transient resources stay in the Discard state.
+		EpilogueResourceAccesses.Emplace(Texture->GetRHI(), SubresourceStateAfter.Access);
+	}
+
+	// Transition any unused (null) sub-resources to the epilogue state since we are assigning a monolithic state across all subresources.
+	for (FRDGSubresourceState*& State : Texture->State)
+	{
+		if (!State)
 		{
-			if (!State)
+			if (!SubresourceStateBefore)
 			{
-				State = &SubresourceStateBefore;
+				SubresourceStateBefore = AllocSubresource();
+				SubresourceStateBefore->SetPass(ERHIPipeline::Graphics, Texture->FirstPass);
 			}
-		}
 
-		EpilogueResourceAccesses.Emplace(Texture->GetRHI(), SubresourceStateAfter.Access);
+			State = SubresourceStateBefore;
+		}
 	}
 
 	InitTextureSubresources(ScratchTextureState, Texture->Layout, &SubresourceStateAfter);
