Re-enable the realtime GPU profiler by default, but disable in editor. This is a more targeted workaround for UE-205211
#tests tested in editor and with a replayrun preflight with the cvar disabled

diff --git a/Engine/Config/BaseEditor.ini b/Engine/Config/BaseEditor.ini
--- a/Engine/Config/BaseEditor.ini
+++ b/Engine/Config/BaseEditor.ini
@@ -322,6 +322,10 @@ MaxNumOfAccessViolation=16
 HighSpawnRateOrBurstThreshold=35
 FarLODDistanceTheshold=3000
 
+[ConsoleVariables]
+; Workaround for crashes when opening large levels (UE-205211)
+r.GPUStatsEnabled=0 
+
 [GatherTextFromAssets]
 MinFreeMemory=1024
 MaxUsedMemory=16384
diff --git a/Engine/Source/Runtime/RenderCore/Private/ProfilingDebugging/RealtimeGPUProfiler.cpp b/Engine/Source/Runtime/RenderCore/Private/ProfilingDebugging/RealtimeGPUProfiler.cpp
--- a/Engine/Source/Runtime/RenderCore/Private/ProfilingDebugging/RealtimeGPUProfiler.cpp
+++ b/Engine/Source/Runtime/RenderCore/Private/ProfilingDebugging/RealtimeGPUProfiler.cpp
@@ -19,7 +19,7 @@ CSV_DEFINE_CATEGORY_MODULE(RENDERCORE_API, GPU, true);
 
 static TAutoConsoleVariable<int> CVarGPUStatsEnabled(
 	TEXT("r.GPUStatsEnabled"),
-	0,
+	1,
 	TEXT("Enables or disables GPU stat recording"));
 
 
