RDG:  Event name fix -- an integration changed the FRDGEventName::GetTCHAR() implementation to not return the raw format string when the dynamic formatted string is missing -- an optimization in a previous CL made it so only the raw string is stored where no formatting is necessary, so this results in missing event names for that case, breaking bread crumbs and GPU captures.

#rnx

diff --git a/Engine/Source/Runtime/RenderCore/Private/RenderGraphEvent.cpp b/Engine/Source/Runtime/RenderCore/Private/RenderGraphEvent.cpp
--- a/Engine/Source/Runtime/RenderCore/Private/RenderGraphEvent.cpp
+++ b/Engine/Source/Runtime/RenderCore/Private/RenderGraphEvent.cpp
@@ -331,7 +331,12 @@ const TCHAR* FRDGEventName::GetTCHAR() const
 {
 #if RDG_EVENTS == RDG_EVENTS_STRING_COPY
 
-	return *FormattedEventName;
+	// Formatted name will be empty in cases where there are no variadic arguments -- EventFormat should be used in that case
+	if (!FormattedEventName.IsEmpty())
+	{
+		return *FormattedEventName;
+	}
+	return EventFormat;
 
 #elif RDG_EVENTS == RDG_EVENTS_STRING_REF
 
