Removed  RDG wait stat to work around crash in CSV profiler

diff --git a/Engine/Source/Runtime/RenderCore/Private/RenderGraphBuilder.cpp b/Engine/Source/Runtime/RenderCore/Private/RenderGraphBuilder.cpp
--- a/Engine/Source/Runtime/RenderCore/Private/RenderGraphBuilder.cpp
+++ b/Engine/Source/Runtime/RenderCore/Private/RenderGraphBuilder.cpp
@@ -1656,7 +1656,6 @@ void FRDGBuilder::Execute()
 {
 	CSV_SCOPED_TIMING_STAT_EXCLUSIVE(RDG);
 	SCOPED_DRAW_EVENTF(RHICmdList, FRDGBuilder_Execute, TEXT("FRDGBuilder::Execute"));
-	CSV_SCOPED_SET_WAIT_STAT(RDG);
 
 #if WITH_RHI_BREADCRUMBS
 	check(LocalCurrentBreadcrumb == FRHIBreadcrumbNode::Sentinel);
