Fix DumpGPU scopes

diff --git a/Engine/Source/Runtime/RenderCore/Private/DumpGPU.cpp b/Engine/Source/Runtime/RenderCore/Private/DumpGPU.cpp
--- a/Engine/Source/Runtime/RenderCore/Private/DumpGPU.cpp
+++ b/Engine/Source/Runtime/RenderCore/Private/DumpGPU.cpp
@@ -2624,7 +2624,13 @@ void FRDGBuilder::DumpResourcePassOutputs(const FRDGPass* Pass)
 		TArray<TSharedPtr<FJsonValue>> ParentEventScopeNames;
 		{
 #if RDG_EVENTS
-			for (FRDGScope const* ParentScope = Pass->GetScope(); ParentScope; ParentScope = ParentScope->Parent)
+			FRDGScope const* ParentScope = Pass->GetScope();
+			if (ensure(ParentScope))
+			{
+				// FRDGScopeState sets ERDGScopeMode::AllEventsAndPassNames when DumpGPU is active.
+				ParentScope = ParentScope->Parent;
+			}
+			for (; ParentScope; ParentScope = ParentScope->Parent)
 			{
 				if (FRDGScope_RHI const* RHIScope = ParentScope->Get<FRDGScope_RHI>())
 				{
diff --git a/Engine/Source/Runtime/RenderCore/Private/RenderGraphPrivate.cpp b/Engine/Source/Runtime/RenderCore/Private/RenderGraphPrivate.cpp
--- a/Engine/Source/Runtime/RenderCore/Private/RenderGraphPrivate.cpp
+++ b/Engine/Source/Runtime/RenderCore/Private/RenderGraphPrivate.cpp
@@ -3,6 +3,7 @@
 #include "RenderGraphPrivate.h"
 #include "RenderGraphEvent.h"
 #include "RenderGraphTrace.h"
+#include "RenderGraphBuilder.h"
 #include "DataDrivenShaderPlatformInfo.h"
 #include "Misc/CommandLine.h"
 #include "RHICommandList.h"
@@ -671,6 +672,12 @@ FRDGScopeState::FState::FState(bool bImmediate)
 		bool bRDGChannelEnabled = false;
 		IF_RDG_ENABLE_TRACE(bRDGChannelEnabled = UE_TRACE_CHANNELEXPR_IS_ENABLED(RDGChannel));
 
+		if (FRDGBuilder::IsDumpingFrame())
+		{
+			// We want all possible scope and pass names in a DumpGPU trace.
+			return ERDGScopeMode::AllEventsAndPassNames;
+		}
+
 		// This is polled once as a workaround for a race condition since the underlying global is not always changed on the render thread.
 		ERDGScopeMode LocalScopeMode = static_cast<ERDGScopeMode>(CVarRDGEvents.GetValueOnRenderThread());
 
