Fixing RDG insights scopes.

diff --git a/Engine/Source/Runtime/RenderCore/Private/RenderGraphTrace.cpp b/Engine/Source/Runtime/RenderCore/Private/RenderGraphTrace.cpp
--- a/Engine/Source/Runtime/RenderCore/Private/RenderGraphTrace.cpp
+++ b/Engine/Source/Runtime/RenderCore/Private/RenderGraphTrace.cpp
@@ -209,11 +209,24 @@ void FRDGTrace::OutputGraphEnd(const FRDGBuilder& GraphBuilder)
 				{
 					const TCHAR* Name = RHIScope->Name.GetTCHAR();
 
+					uint32 Depth = 0;
+
+					for (FRDGScope* Scope = Current; Scope->Parent; Scope = Scope->Parent)
+					{
+						if (Scope->Get<FRDGScope_RHI>())
+						{
+							if (Scope->CPUFirstPass && Scope->CPULastPass)
+							{
+								Depth++;
+							}
+						}
+					}
+
 					UE_TRACE_LOG(RDGTrace, ScopeMessage, RDGChannel)
 						<< ScopeMessage.Name(Name, uint16(FCString::Strlen(Name)))
 						<< ScopeMessage.FirstPass(Current->CPUFirstPass->GetHandle().GetIndexUnchecked())
 						<< ScopeMessage.LastPass(Current->CPULastPass->GetHandle().GetIndexUnchecked())
-						<< ScopeMessage.Depth(Current->Depth);
+						<< ScopeMessage.Depth(Depth);
 				}
 			}
 		};
diff --git a/Engine/Source/Runtime/RenderCore/Public/RenderGraphEvent.h b/Engine/Source/Runtime/RenderCore/Public/RenderGraphEvent.h
--- a/Engine/Source/Runtime/RenderCore/Public/RenderGraphEvent.h
+++ b/Engine/Source/Runtime/RenderCore/Public/RenderGraphEvent.h
@@ -288,7 +288,6 @@ struct FRDGScope
 	FStorage Impl;
 
 #if RDG_ENABLE_TRACE
-	uint32 Depth = 0;
 	bool bVisited = false;
 #endif
 
