Comments an ensure() in DumpGPU to unblock automated tests

diff --git a/Engine/Source/Runtime/RenderCore/Private/DumpGPU.cpp b/Engine/Source/Runtime/RenderCore/Private/DumpGPU.cpp
--- a/Engine/Source/Runtime/RenderCore/Private/DumpGPU.cpp
+++ b/Engine/Source/Runtime/RenderCore/Private/DumpGPU.cpp
@@ -2625,7 +2625,7 @@ void FRDGBuilder::DumpResourcePassOutputs(const FRDGPass* Pass)
 		{
 #if RDG_EVENTS
 			FRDGScope const* ParentScope = Pass->GetScope();
-			if (ensure(ParentScope))
+			if (ParentScope)
 			{
 				// FRDGScopeState sets ERDGScopeMode::AllEventsAndPassNames when DumpGPU is active.
 				ParentScope = ParentScope->Parent;
