Add per-shader-type itemization for shader libraries.

- At the end of the cook, produces stats like

Unique shaders itemization (sorted by compressed size):
SF_Compute: 3202 shaders (65.59%), compressed size: 0.06 MB (19.70 KB avg per shader), uncompressed size: 0.09 MB (29.44 KB avg per shader)
SF_Pixel: 1238 shaders (25.36%), compressed size: 0.02 MB (16.65 KB avg per shader), uncompressed size: 0.03 MB (24.99 KB avg per shader)
SF_RayGen: 296 shaders (6.06%), compressed size: 0.01 MB (37.48 KB avg per shader), uncompressed size: 0.02 MB (55.79 KB avg per shader)
SF_Vertex: 115 shaders (2.36%), compressed size: 0.00 MB (6.97 KB avg per shader), uncompressed size: 0.00 MB (10.78 KB avg per shader)
SF_RayHitGroup: 18 shaders (0.37%), compressed size: 0.00 MB (5.12 KB avg per shader), uncompressed size: 0.00 MB (7.84 KB avg per shader)
SF_RayMiss: 9 shaders (0.18%), compressed size: 0.00 MB (7.63 KB avg per shader), uncompressed size: 0.00 MB (11.81 KB avg per shader)
SF_Geometry: 3 shaders (0.06%), compressed size: 0.00 MB (2.63 KB avg per shader), uncompressed size: 0.00 MB (4.78 KB avg per shader)
SF_RayCallable: 1 shaders (0.02%), compressed size: 0.00 MB (1.68 KB avg per shader), uncompressed size: 0.00 MB (2.71 KB avg per shader)

diff --git a/Engine/Source/Runtime/RenderCore/Private/ShaderCodeArchive.cpp b/Engine/Source/Runtime/RenderCore/Private/ShaderCodeArchive.cpp
--- a/Engine/Source/Runtime/RenderCore/Private/ShaderCodeArchive.cpp
+++ b/Engine/Source/Runtime/RenderCore/Private/ShaderCodeArchive.cpp
@@ -833,6 +833,20 @@ void FSerializedShaderArchive::CollectStatsAndDebugInfo(FDebugStats& OutDebugSta
 				break;
 			}
 		}
+
+		// calculate per-frequency stats
+		FMemory::Memzero(OutExtendedDebugStats->NumShadersPerFrequency);
+		FMemory::Memzero(OutExtendedDebugStats->UncompressedSizePerFrequency);
+		FMemory::Memzero(OutExtendedDebugStats->CompressedSizePerFrequency);
+		for (const FShaderCodeEntry& ShaderEntry : ShaderEntries)
+		{
+			check(ShaderEntry.Frequency < UE_ARRAY_COUNT(OutExtendedDebugStats->NumShadersPerFrequency));
+			++OutExtendedDebugStats->NumShadersPerFrequency[ShaderEntry.Frequency];
+			check(ShaderEntry.Frequency < UE_ARRAY_COUNT(OutExtendedDebugStats->UncompressedSizePerFrequency));
+			OutExtendedDebugStats->UncompressedSizePerFrequency[ShaderEntry.Frequency] += ShaderEntry.UncompressedSize;
+			check(ShaderEntry.Frequency < UE_ARRAY_COUNT(OutExtendedDebugStats->UncompressedSizePerFrequency));
+			OutExtendedDebugStats->CompressedSizePerFrequency[ShaderEntry.Frequency] += ShaderEntry.Size;
+		}
 	}
 
 #if 0 // graph visualization - maybe one day we'll return to this
diff --git a/Engine/Source/Runtime/RenderCore/Private/ShaderCodeLibrary.cpp b/Engine/Source/Runtime/RenderCore/Private/ShaderCodeLibrary.cpp
--- a/Engine/Source/Runtime/RenderCore/Private/ShaderCodeLibrary.cpp
+++ b/Engine/Source/Runtime/RenderCore/Private/ShaderCodeLibrary.cpp
@@ -52,6 +52,7 @@ ShaderCodeLibrary.cpp: Bound shader state cache implementation.
 #include "UObject/Object.h"
 #include "Serialization/CompactBinarySerialization.h"
 #include "Serialization/CompactBinaryWriter.h"
+#include "RHIStrings.h"
 #endif
 
 // allow introspection (e.g. dumping the contents) for easier debugging
@@ -2110,6 +2111,30 @@ struct FEditorShaderCodeArchive
 					UsageString += FString::Printf(TEXT("%d, "), ExtendedStats.TopShaderUsages[IdxUsage]);
 				}
 				UE_LOG(LogShaderLibrary, Display, TEXT("    %s%d"), *UsageString, ExtendedStats.TopShaderUsages[ExtendedStats.TopShaderUsages.Num() - 1]);
+
+				// print per-frequency stats
+				UE_LOG(LogShaderLibrary, Display, TEXT("Unique shaders itemization (sorted by compressed size):"));
+				// sort by compressed size
+				TArray<int32> SortedIndices;
+				for (int32 IdxFreq = 0; IdxFreq < SF_NumFrequencies; ++IdxFreq)
+				{
+					SortedIndices.Add(IdxFreq);
+				}
+				SortedIndices.StableSort([&ExtendedStats](int32 IndexA, int32 IndexB) { return ExtendedStats.CompressedSizePerFrequency[IndexA] >= ExtendedStats.CompressedSizePerFrequency[IndexB]; });
+				for (int32 Freq : SortedIndices)
+				{
+					if (ExtendedStats.NumShadersPerFrequency[Freq] > 0)
+					{
+						UE_LOG(LogShaderLibrary, Display, TEXT("%s: %d shaders (%.2f%%), compressed size: %.2f MB (%.2f KB avg per shader), uncompressed size: %.2f MB (%.2f KB avg per shader)"),
+							GetShaderFrequencyString(static_cast<EShaderFrequency>(Freq)),
+							ExtendedStats.NumShadersPerFrequency[Freq],
+							100.0 * ExtendedStats.NumShadersPerFrequency[Freq] / double(Stats.NumUniqueShaders),
+							double(ExtendedStats.CompressedSizePerFrequency[Freq]) / (1024.0 * 1024.0 * 1024.0), double(ExtendedStats.CompressedSizePerFrequency[Freq]) / (double(ExtendedStats.NumShadersPerFrequency[Freq]) * 1024.0),
+							double(ExtendedStats.UncompressedSizePerFrequency[Freq]) / (1024.0 * 1024.0 * 1024.0), double(ExtendedStats.UncompressedSizePerFrequency[Freq]) / (double(ExtendedStats.NumShadersPerFrequency[Freq]) * 1024.0)
+						);
+					}
+				}
+
 			}
 			else
 			{
diff --git a/Engine/Source/Runtime/RenderCore/Public/ShaderCodeArchive.h b/Engine/Source/Runtime/RenderCore/Public/ShaderCodeArchive.h
--- a/Engine/Source/Runtime/RenderCore/Public/ShaderCodeArchive.h
+++ b/Engine/Source/Runtime/RenderCore/Public/ShaderCodeArchive.h
@@ -136,6 +136,15 @@ public:
 
 		/** For the top shaders (descending), number of shadermaps in which they are used. Expected to be limited to a small number (10) */
 		TArray<int32> TopShaderUsages;
+
+		/** Number of shaers per frequency. */
+		int32 NumShadersPerFrequency[SF_NumFrequencies];
+
+		/** Uncompressed size of all shaders of a given frequency. */
+		uint64 UncompressedSizePerFrequency[SF_NumFrequencies];
+
+		/** Compressed (individually) size of all shaders of a given frequency. */
+		uint64 CompressedSizePerFrequency[SF_NumFrequencies];
 	};
 #endif
 
