Fix shader code stats being output in gigabytes instead of megabytes.

- Thanks to [REDACTED_PERSON] for noticing!

diff --git a/Engine/Source/Runtime/RenderCore/Private/ShaderCodeLibrary.cpp b/Engine/Source/Runtime/RenderCore/Private/ShaderCodeLibrary.cpp
--- a/Engine/Source/Runtime/RenderCore/Private/ShaderCodeLibrary.cpp
+++ b/Engine/Source/Runtime/RenderCore/Private/ShaderCodeLibrary.cpp
@@ -2129,8 +2129,8 @@ struct FEditorShaderCodeArchive
 							GetShaderFrequencyString(static_cast<EShaderFrequency>(Freq)),
 							ExtendedStats.NumShadersPerFrequency[Freq],
 							100.0 * ExtendedStats.NumShadersPerFrequency[Freq] / double(Stats.NumUniqueShaders),
-							double(ExtendedStats.CompressedSizePerFrequency[Freq]) / (1024.0 * 1024.0 * 1024.0), double(ExtendedStats.CompressedSizePerFrequency[Freq]) / (double(ExtendedStats.NumShadersPerFrequency[Freq]) * 1024.0),
-							double(ExtendedStats.UncompressedSizePerFrequency[Freq]) / (1024.0 * 1024.0 * 1024.0), double(ExtendedStats.UncompressedSizePerFrequency[Freq]) / (double(ExtendedStats.NumShadersPerFrequency[Freq]) * 1024.0)
+							double(ExtendedStats.CompressedSizePerFrequency[Freq]) / (1024.0 * 1024.0), double(ExtendedStats.CompressedSizePerFrequency[Freq]) / (double(ExtendedStats.NumShadersPerFrequency[Freq]) * 1024.0),
+							double(ExtendedStats.UncompressedSizePerFrequency[Freq]) / (1024.0 * 1024.0), double(ExtendedStats.UncompressedSizePerFrequency[Freq]) / (double(ExtendedStats.NumShadersPerFrequency[Freq]) * 1024.0)
 						);
 					}
 				}
