Visualize Texture:  Performance and feature upgrades.
* Visualize texture system starts out in an inactive state until a command is issued, avoiding overhead of tracking views and scene textures, saving 1.4% on the render thread.
* Visualization overhead eliminated for views besides the one currently being visualized.
* Support for visualization of textures from scene captures, via "view=N" option (specifying the unique ID of the view), with "view=?" displaying a list of views for reference.
* Improved visualization for cube maps.  PIP uses 2:1 aspect for the longitudinal render to match resource viewer display, and pixel perfect option shows tiled flat cube map faces (actual pixels) rather than running a projection.
* Padding for scene or screen pass textures is removed in the visualization -- the padding otherwise shows up as garbage or blank space.

To remove scene texture padding, it's necessary to add a field to RDG textures to provide an option to track the viewport sizes that were rendered for a given texture.  If not set, the assumption is the whole texture was rendered.  The field is set for FSceneTextures and FScreenPassTexture, covering the vast majority of cases, plus the denoiser was spot fixed -- worst case if any other cases are missed, you still see the padding.  You can tell padding was present when visualizing by contrasting the texture size with the viewport size.

Padding was always a potential issue for the visualizer, but is exacerbated by scene captures, as the padded scene textures are set to a size that's a union of the main view and any scene captures.  Padding is also exacerbated by dynamic resolution scaling, as the buffers will be padded to the maximum resolution.  For example, a cube map rendering at 512x512 will have 93% of the pixel area as padding if the front buffer is at 1440p, or the default dynamic resolution setup will have 70% of the pixels as padding at minimum res.

diff --git a/Engine/Plugins/Media/MediaIOFramework/Source/MediaIOCore/Private/MediaCaptureRenderPass.cpp b/Engine/Plugins/Media/MediaIOFramework/Source/MediaIOCore/Private/MediaCaptureRenderPass.cpp
--- a/Engine/Plugins/Media/MediaIOFramework/Source/MediaIOCore/Private/MediaCaptureRenderPass.cpp
+++ b/Engine/Plugins/Media/MediaIOFramework/Source/MediaIOCore/Private/MediaCaptureRenderPass.cpp
@@ -71,12 +71,14 @@ namespace UE::MediaCapture::Resample
 		Input.Texture = InputTexture.Texture;
 		Input.ViewRect = FIntRect{ 0, 0, InputSize.X, InputSize.Y };
 		Input.LoadAction = ERenderTargetLoadAction::ENoAction;
+		Input.UpdateVisualizeTextureExtent();
 
 		FScreenPassRenderTarget Output;
 		const FIntVector OutputSize = OutputTexture->Desc.GetSize();
 		Output.Texture = OutputTexture;
 		Output.ViewRect = FIntRect{ 0, 0, OutputSize.X, OutputSize.Y };
 		Output.LoadAction = ERenderTargetLoadAction::ENoAction;
+		Output.UpdateVisualizeTextureExtent();
 
 		FMediaCaptureResamplePS::FParameters* PassParameters = FMediaCaptureResamplePS::AllocateAndSetParameters(GraphBuilder, InputTexture.Texture, OutputTexture);
 		
diff --git a/Engine/Shaders/Private/Tools/VisualizeTexture.usf b/Engine/Shaders/Private/Tools/VisualizeTexture.usf
--- a/Engine/Shaders/Private/Tools/VisualizeTexture.usf
+++ b/Engine/Shaders/Private/Tools/VisualizeTexture.usf
@@ -51,6 +51,40 @@ void VisualizeTexturePS(
 		float s = sin(Angles.y);
 		SpherePos = float3(s * sin(Angles.x), -s * cos(Angles.x), -cos(Angles.y));
 	}
+	
+	// Pixel based cube map unwrapping -- assumes 4x3 aspect ratio.  Produces a seamless panorama in the middle, connected to adjacent sky / floor faces.
+	float3 CubePixelPos;
+	{
+		float2 LerpCoords = UV * float2(4.0, 3.0);
+		float FaceCoords[3] =
+		{
+			frac(LerpCoords.x) * 2.0 - 1.0,
+			frac(LerpCoords.y) * 2.0 - 1.0,
+			1.0
+		};
+		
+		// Table selecting a UV coordinate or "1" for each of the XYZ axes using a hex digit, preceded by a "0" or "2" to specify whether the axis should be negated.
+		uint FaceUVMappingTable[12] =
+		{
+			0x010002,		// Positive Z faces (different rotations to line up with panoramic neighbor below)
+			0x200102,
+			0x212002,
+			0x002102,
+			0x020021,		// Positive X
+			0x200221,		// Positive Y
+			0x222021,		// Negative X
+			0x002221,		// Negative Y
+			0x210022,		// Negative Z faces (different rotations to line up with panoramic neighbor above)
+			0x202122,
+			0x012022,
+			0x000122,
+		};
+		uint FaceUVMapping = FaceUVMappingTable[floor(LerpCoords.x) + floor(LerpCoords.y) * 4.0];
+		CubePixelPos = float3(
+			FaceCoords[(FaceUVMapping >> 16) & 0xf] * (1.0 - ((float) ((FaceUVMapping >> 20) & 0xf))),
+			FaceCoords[(FaceUVMapping >>  8) & 0xf] * (1.0 - ((float) ((FaceUVMapping >> 12) & 0xf))),
+			FaceCoords[(FaceUVMapping >>  0) & 0xf] * (1.0 - ((float) ((FaceUVMapping >>  4) & 0xf))));
+	}
 
 	// volume texture unwrapping
 	float3 VolumePos;
@@ -74,7 +108,8 @@ void VisualizeTexturePS(
 
 
 #if TEXTURE_TYPE == 0
-	float4 TexLookup = TextureCubeSampleLevel(VisualizeTextureCube, VisualizeTextureCubeSampler, SpherePos, MipLevel);
+	// Select pixel accurate map coordinates if the visualization is 4x3 aspect ratio
+	float4 TexLookup = TextureCubeSampleLevel(VisualizeTextureCube, VisualizeTextureCubeSampler, TextureExtent.x*0.75 == TextureExtent.y ? CubePixelPos : SpherePos, MipLevel);
 #elif TEXTURE_TYPE == 2 || TEXTURE_TYPE == 6
 	float4 TexLookup = Texture2DSampleLevel(VisualizeTexture2D, VisualizeTexture2DSampler, UV, MipLevel);
 #elif TEXTURE_TYPE == 3
diff --git a/Engine/Source/Runtime/Launch/Private/LaunchEngineLoop.cpp b/Engine/Source/Runtime/Launch/Private/LaunchEngineLoop.cpp
--- a/Engine/Source/Runtime/Launch/Private/LaunchEngineLoop.cpp
+++ b/Engine/Source/Runtime/Launch/Private/LaunchEngineLoop.cpp
@@ -191,6 +191,7 @@
 	#include "DynamicResolutionState.h"
 	#include "EngineModule.h"
 	#include "DumpGPU.h"
+	#include "VisualizeTexture.h"
 
 #if !UE_SERVER
 	#include "AppMediaTimeSource.h"
@@ -5497,6 +5498,8 @@ static inline void BeginFrameRenderThread(FRHICommandListImmediate& RHICmdList,
 	FCsvProfiler::BeginExclusiveStat("RenderThreadOther");
 #endif
 
+	GVisualizeTexture.BeginFrameRenderThread();
+
 	// Waits after this point but before BeginRenderingViewFamilies are not on the RT critical path (since we're waiting for the GT)
 	FThreadIdleStats::EndCriticalPath();
 }
@@ -5522,6 +5525,8 @@ static inline void EndFrameRenderThread(FRHICommandListImmediate& RHICmdList, ui
 	FCoreDelegates::OnEndFrameRT.Broadcast();
 	RHICmdList.EndFrame();
 
+	GVisualizeTexture.EndFrameRenderThread();
+
 	GPU_STATS_ENDFRAME(RHICmdList);
 
 #if WITH_RHI_BREADCRUMBS
diff --git a/Engine/Source/Runtime/RenderCore/Private/RenderGraphBuilder.cpp b/Engine/Source/Runtime/RenderCore/Private/RenderGraphBuilder.cpp
--- a/Engine/Source/Runtime/RenderCore/Private/RenderGraphBuilder.cpp
+++ b/Engine/Source/Runtime/RenderCore/Private/RenderGraphBuilder.cpp
@@ -4257,7 +4257,7 @@ void FRDGBuilder::InitViewRHI(FRHICommandListBase& InRHICmdList, FRDGView* View)
 void FRDGBuilder::VisualizePassOutputs(const FRDGPass* Pass)
 {
 #if SUPPORTS_VISUALIZE_TEXTURE
-	if (!AuxiliaryPasses.IsVisualizeAllowed())
+	if (!GVisualizeTexture.IsRequestedView() || !AuxiliaryPasses.IsVisualizeAllowed())
 	{
 		return;
 	}
diff --git a/Engine/Source/Runtime/RenderCore/Private/VisualizeTexture.cpp b/Engine/Source/Runtime/RenderCore/Private/VisualizeTexture.cpp
--- a/Engine/Source/Runtime/RenderCore/Private/VisualizeTexture.cpp
+++ b/Engine/Source/Runtime/RenderCore/Private/VisualizeTexture.cpp
@@ -19,6 +19,7 @@ void FVisualizeTexture::ParseCommands(const TCHAR* Cmd, FOutputDevice &Ar)
 #if SUPPORTS_VISUALIZE_TEXTURE
 	// Find out what command to do based on first parameter.
 	ECommand Command = ECommand::Unknown;
+	int32 ViewId = 0;
 	FString ResourceName;
 	TOptional<uint32> ResourceVersion;
 	TOptional<FWildcardString> ResourceListWildCard;
@@ -42,6 +43,24 @@ void FVisualizeTexture::ParseCommands(const TCHAR* Cmd, FOutputDevice &Ar)
 		{
 			Command = ECommand::DisplayPoolResourceList;
 		}
+		else if (FirstParameter.StartsWith(TEXT("view=")))
+		{
+			if (FirstParameter.Len() == 5)
+			{
+				// Empty payload, reset view ID to zero
+				Command = ECommand::SetViewId;
+				ViewId = 0;
+			}
+			else if (FirstParameter[5] == TEXT('?'))
+			{
+				Command = ECommand::DisplayViewList;
+			}
+			else
+			{
+				Command = ECommand::SetViewId;
+				ViewId = FCString::Strtoi(&FirstParameter[5], nullptr, 0);
+			}
+		}
 		else
 		{
 			const TCHAR* AfterAt = *FirstParameter;
@@ -240,10 +259,24 @@ void FVisualizeTexture::ParseCommands(const TCHAR* Cmd, FOutputDevice &Ar)
 
 		DisplayResourceListToLog(ResourceListWildCard);
 	}
+	else if (Command == ECommand::DisplayViewList)
+	{
+		DisplayViewListToLog();
+	}
+	else if (Command == ECommand::SetViewId)
+	{
+		Requested.ViewUniqueId = ViewId;
+	}
 	else
 	{
 		unimplemented();
 	}
+
+	// Enable tracking when the system is first interacted with
+	if (State == EState::Inactive)
+	{
+		State = EState::TrackResources;
+	}
 #endif
 }
 
@@ -323,6 +356,9 @@ void FVisualizeTexture::DisplayHelp(FOutputDevice &Ar)
 	Ar.Logf(TEXT("  BYNAME   = sort pool list by name"));
 	Ar.Logf(TEXT("  BYSIZE   = show pool list by size"));
 	Ar.Logf(TEXT(""));
+	Ar.Logf(TEXT("VisualizeTexture/Vis view=N"));
+	Ar.Logf(TEXT("  Unique ID of view to visualize textures from, \"view=?\" to dump list of available views"));
+	Ar.Logf(TEXT(""));
 }
 
 void FVisualizeTexture::DisplayPoolResourceListToLog(FVisualizeTexture::ESortBy SortBy)
@@ -439,6 +475,13 @@ void FVisualizeTexture::DisplayPoolResourceListToLog(FVisualizeTexture::ESortBy
 
 void FVisualizeTexture::DisplayResourceListToLog(const TOptional<FWildcardString>& Wildcard)
 {
+	if (!IsActive())
+	{
+		State = EState::DisplayResources;
+		DisplayResourcesParam = Wildcard;
+		return;
+	}
+
 	UE_LOG(LogConsoleResponse, Log, TEXT("RDGResourceName (what was rendered this frame, use <RDGResourceName>@<Version> to get intermediate versions):"));
 
 	TArray<FString> Entries;
@@ -516,6 +559,35 @@ void FVisualizeTexture::DisplayResourceListToLog(const TOptional<FWildcardString
 	UE_LOG(LogConsoleResponse, Log, TEXT(""));
 }
 
+void FVisualizeTexture::DisplayViewListToLog()
+{
+	if (!IsActive())
+	{
+		State = EState::DisplayViews;
+		return;
+	}
+
+	// Display view list sorted by unique ID
+	TArray<FSetElementId> Entries;
+	Entries.Reserve(ViewDescriptionMap.Num());
+	for (auto ViewIt = ViewDescriptionMap.CreateConstIterator(); ViewIt; ++ViewIt)
+	{
+		Entries.Add(ViewIt.GetId());
+	}
+
+	Entries.Sort([ViewDescriptionMap = ViewDescriptionMap](const FSetElementId& A, const FSetElementId& B)
+	{
+		return ViewDescriptionMap.Get(A).Key < ViewDescriptionMap.Get(B).Key;
+	});
+
+	UE_LOG(LogConsoleResponse, Log, TEXT("Visualize Texture available views:"));
+
+	for (FSetElementId ElementId : Entries)
+	{
+		UE_LOG(LogConsoleResponse, Log, TEXT("   %d  %s"), ViewDescriptionMap.Get(ElementId).Key, *ViewDescriptionMap.Get(ElementId).Value);
+	}
+}
+
 static TAutoConsoleVariable<int32> CVarAllowBlinking(
 	TEXT("r.VisualizeTexture.AllowBlinking"), 1,
 	TEXT("Whether to allow blinking when visualizing NaN or inf that can become irritating over time.\n"),
@@ -644,10 +716,42 @@ FRDGTextureRef FVisualizeTexture::AddVisualizeTexturePass(
 	check(!EnumHasAnyFlags(InputTexture->Desc.Flags, TexCreate_CPUReadback));
 
 	const FRDGTextureDesc& InputDesc = InputTexture->Desc;
-	const FIntPoint InputExtent = InputDesc.Extent;
-
+	FIntPoint InputExtent = InputDesc.Extent;
 	FIntPoint OutputExtent = InputExtent;
 
+	// Scene textures are padded and shared across scene renderers, with a given scene renderer using a viewport in the shared buffer.
+	// We only want to visualize the portion actually used by the given scene renderer, as the rest will be blank or garbage.  The info
+	// text will display the actual texture size in addition to the viewport being visualized.
+	FIntPoint VisualizeTextureExtent = InputTexture->GetVisualizeExtent();
+	if ((VisualizeTextureExtent.X > 0) && (VisualizeTextureExtent.Y > 0))
+	{
+		// Clamp extent at actual dimensions of texture
+		OutputExtent.X = FMath::Min(VisualizeTextureExtent.X, OutputExtent.X);
+		OutputExtent.Y = FMath::Min(VisualizeTextureExtent.Y, OutputExtent.Y);
+	}
+
+	if (InputDesc.IsTextureCube())
+	{
+		// For pixel perfect display of cube map, we'll use a 4x3 flat unwrapping of the cube map, rather than a projection.  The visualization
+		// shader detects the 4x3 aspect ratio, and generates a seamless panorama in the middle, with the adjacent floor and sky tiles above
+		// and below.  There will be seams between floor and sky tiles, but the pixels shown will otherwise be exact.
+		if (GVisualizeTexture.Config.InputUVMapping == EInputUVMapping::PixelPerfectCenter)
+		{
+			InputExtent.X *= 4;
+			InputExtent.Y *= 3;
+			OutputExtent.X *= 4;
+			OutputExtent.Y *= 3;
+		}
+		else
+		{
+			// Longitudinal rendered cube maps look better with 2 to 1 aspect ratio (same as how the texture resource viewer displays cube maps)
+			InputExtent.X *= 2;
+			OutputExtent.X *= 2;
+		}
+	}
+
+	GVisualizeTexture.Captured.OutputExtent = OutputExtent;
+
 	// Clamp to reasonable value to prevent crash
 	OutputExtent.X = FMath::Max(OutputExtent.X, 1);
 	OutputExtent.Y = FMath::Max(OutputExtent.Y, 1);
@@ -792,6 +896,7 @@ void FVisualizeTexture::CreateContentCapturePass(FRDGBuilder& GraphBuilder, cons
 		Captured.PooledRenderTarget = nullptr;
 		Captured.Texture = OutputTexture;
 		Captured.InputValueMapping = InputValueMapping;
+		Captured.ViewRects = FamilyViewRects;
 
 		GraphBuilder.QueueTextureExtraction(OutputTexture, &Captured.PooledRenderTarget);
 	}
@@ -838,6 +943,62 @@ void FVisualizeTexture::CreateContentCapturePass(FRDGBuilder& GraphBuilder, cons
 	}
 }
 
+void FVisualizeTexture::BeginFrameRenderThread()
+{
+	bAnyViewRendered = false;
+	bIsRequestedView = false;
+	bFoundRequestedView = false;
+}
+
+void FVisualizeTexture::BeginViewRenderThread(ERHIFeatureLevel::Type InFeatureLevel, int32 UniqueId, const TCHAR* Description, bool bIsSceneCapture)
+{
+	// Only support visualization for views with a unique ID
+	if (State == EState::Inactive || !UniqueId)
+	{
+		return;
+	}
+
+	FeatureLevel = InFeatureLevel;
+
+	if (!bAnyViewRendered)
+	{
+		// Clear list of views out when we encounter the first view on the current frame
+		ViewDescriptionMap.Empty();
+		bAnyViewRendered = true;
+	}
+
+	ViewDescriptionMap.FindOrAdd(UniqueId) = Description;
+
+	if (Requested.ViewUniqueId == UniqueId)
+	{
+		// Found the specific view we requested
+		bIsRequestedView = true;
+		bFoundRequestedView = true;
+	}
+	else if (!bFoundRequestedView)
+	{
+		// If specific requested view hasn't been found, visualize any view that's not a scene capture, so we still get some sort of result
+		bIsRequestedView = !bIsSceneCapture;
+	}
+
+	// Clear outputs when we are processing a requested view
+	if (bIsRequestedView)
+	{
+		VersionCountMap.Empty();
+		Captured = {};
+		Captured.ViewUniqueId = UniqueId;
+	}
+}
+
+void FVisualizeTexture::SetSceneTextures(const TArray<FRDGTextureRef>& InSceneTextures, FIntPoint InFamilySize, const TArray<FIntRect>& InFamilyViewRects)
+{
+	for (FRDGTextureRef Texture : InSceneTextures)
+	{
+		Texture->EncloseVisualizeExtent(InFamilySize);
+	}
+	FamilyViewRects = InFamilyViewRects;
+}
+
 TOptional<uint32> FVisualizeTexture::ShouldCapture(const TCHAR* InName, uint32 InMipIndex)
 {
 	TOptional<uint32> CaptureId;
@@ -853,6 +1014,33 @@ TOptional<uint32> FVisualizeTexture::ShouldCapture(const TCHAR* InName, uint32 I
 	return CaptureId;
 }
 
+void FVisualizeTexture::EndViewRenderThread()
+{
+	if (bIsRequestedView)
+	{
+		bIsRequestedView = false;
+		FamilyViewRects.Empty();
+	}
+}
+
+void FVisualizeTexture::EndFrameRenderThread()
+{
+	if (bAnyViewRendered)
+	{
+		if (State == EState::DisplayResources)
+		{
+			DisplayResourceListToLog(DisplayResourcesParam);
+			DisplayResourcesParam.Reset();
+			State = EState::TrackResources;
+		}
+		else if (State == EState::DisplayViews)
+		{
+			DisplayViewListToLog();
+			State = EState::TrackResources;
+		}
+	}
+}
+
 uint32 FVisualizeTexture::GetVersionCount(const TCHAR* InName) const
 {
 	if (const uint32* VersionCount = VersionCountMap.Find(InName))
diff --git a/Engine/Source/Runtime/RenderCore/Public/RenderGraphDefinitions.h b/Engine/Source/Runtime/RenderCore/Public/RenderGraphDefinitions.h
--- a/Engine/Source/Runtime/RenderCore/Public/RenderGraphDefinitions.h
+++ b/Engine/Source/Runtime/RenderCore/Public/RenderGraphDefinitions.h
@@ -59,6 +59,8 @@
 	#define RDG_EVENTS RDG_EVENTS_NONE
 #endif
 
+#define SUPPORTS_VISUALIZE_TEXTURE (WITH_ENGINE && (!UE_BUILD_SHIPPING || WITH_EDITOR))
+
 /** ENUMS */
 
 enum class ERDGBuilderFlags
diff --git a/Engine/Source/Runtime/RenderCore/Public/RenderGraphResources.h b/Engine/Source/Runtime/RenderCore/Public/RenderGraphResources.h
--- a/Engine/Source/Runtime/RenderCore/Public/RenderGraphResources.h
+++ b/Engine/Source/Runtime/RenderCore/Public/RenderGraphResources.h
@@ -617,6 +617,25 @@ private:
 	RENDERCORE_API FRDGTextureDebugData& GetTextureDebugData() const;
 #endif
 
+#if SUPPORTS_VISUALIZE_TEXTURE
+	FIntPoint VisualizeTextureExtent = { 0, 0 };
+public:
+	FORCEINLINE void EncloseVisualizeExtent(const FIntPoint& Point)
+	{
+		VisualizeTextureExtent.X = FMath::Max(VisualizeTextureExtent.X, Point.X);
+		VisualizeTextureExtent.Y = FMath::Max(VisualizeTextureExtent.Y, Point.Y);
+	}
+	FORCEINLINE FIntPoint GetVisualizeExtent() const
+	{
+		return VisualizeTextureExtent;
+	}
+private:
+#else
+public:
+	FORCEINLINE void EncloseVisualizeExtent(const FIntPoint& Point) {}
+private:
+#endif
+
 	friend FRDGBuilder;
 	friend FRDGUserValidation;
 	friend FRDGBarrierValidation;
diff --git a/Engine/Source/Runtime/RenderCore/Public/VisualizeTexture.h b/Engine/Source/Runtime/RenderCore/Public/VisualizeTexture.h
--- a/Engine/Source/Runtime/RenderCore/Public/VisualizeTexture.h
+++ b/Engine/Source/Runtime/RenderCore/Public/VisualizeTexture.h
@@ -10,6 +10,7 @@
 #include "Misc/EnumClassFlags.h"
 #include "Misc/Optional.h"
 #include "Misc/WildcardString.h"
+#include "RenderGraphDefinitions.h"
 #include "RenderResource.h"
 #include "RendererInterface.h"
 #include "Templates/RefCounting.h"
@@ -17,7 +18,6 @@
 #if UE_ENABLE_INCLUDE_ORDER_DEPRECATED_IN_5_3
 #include "RHIDefinitions.h"
 #include "RenderGraph.h"
-#include "RenderGraphDefinitions.h"
 #include "RenderGraphResources.h"
 #endif
 
@@ -26,8 +26,6 @@ class FRDGBuilder;
 class FRHICommandListImmediate;
 class FWildcardString;
 
-#define SUPPORTS_VISUALIZE_TEXTURE (WITH_ENGINE && (!UE_BUILD_SHIPPING || WITH_EDITOR))
-
 class FVisualizeTexture : public FRenderResource
 {
 public:
@@ -39,15 +37,42 @@ public:
 
 	RENDERCORE_API void GetTextureInfos_GameThread(TArray<FString>& Infos) const;
 
-	/** Creates a new checkpoint (e.g. "SceneDepth@N") for the pooled render target. A null parameter is a no-op. */
 #if SUPPORTS_VISUALIZE_TEXTURE
+	RENDERCORE_API void BeginFrameRenderThread();
+	RENDERCORE_API void BeginViewRenderThread(ERHIFeatureLevel::Type InFeatureLevel, int32 UniqueId, const TCHAR* Description, bool bIsSceneCapture);
+	RENDERCORE_API void SetSceneTextures(const TArray<FRDGTextureRef>& InSceneTextures, FIntPoint InFamilySize, const TArray<FIntRect>& InFamilyViewRects);
+	RENDERCORE_API void EndViewRenderThread();
+	RENDERCORE_API void EndFrameRenderThread();
+
+	/** Creates a new checkpoint (e.g. "SceneDepth@N") for the pooled render target. A null parameter is a no-op. */
 	RENDERCORE_API void SetCheckPoint(FRDGBuilder& GraphBuilder, IPooledRenderTarget* PooledRenderTarget);
 	RENDERCORE_API void SetCheckPoint(FRHICommandListImmediate& RHICmdList, IPooledRenderTarget* PooledRenderTarget);
 #else
+	FORCEINLINE void BeginFrameRenderThread() {}
+	FORCEINLINE void EndFrameRenderThread() {}
+
 	inline void SetCheckPoint(FRDGBuilder& GraphBuilder, IPooledRenderTarget* PooledRenderTarget) {}
 	inline void SetCheckPoint(FRHICommandListImmediate& RHICmdList, IPooledRenderTarget* PooledRenderTarget) {}
 #endif
 
+	FORCEINLINE bool IsActive() const
+	{
+#if SUPPORTS_VISUALIZE_TEXTURE
+		return State != EState::Inactive;
+#else
+		return false;
+#endif
+	}
+
+	FORCEINLINE bool IsRequestedView() const
+	{
+#if SUPPORTS_VISUALIZE_TEXTURE
+		return bIsRequestedView;
+#else
+		return false;
+#endif
+	}
+
 	static RENDERCORE_API FRDGTextureRef AddVisualizeTexturePass(
 		FRDGBuilder& GraphBuilder,
 		class FGlobalShaderMap* ShaderMap,
@@ -67,6 +92,14 @@ private:
 	};
 	FRIEND_ENUM_CLASS_FLAGS(EFlags);
 
+	enum class EState
+	{
+		Inactive,				// Default initial state, negligible overhead
+		DisplayViews,			// Display views next render frame -- state activated on DisplayViewListToLog call if Inactive
+		DisplayResources,		// Display resources next render frame -- state activated on DisplayResourceListToLog call if Inactive
+		TrackResources,			// Track resources every frame, adding overhead -- state activated after visualize texture related command is issued
+	};
+
 	enum class ECommand
 	{
 		Unknown,
@@ -75,6 +108,8 @@ private:
 		DisplayHelp,
 		DisplayPoolResourceList,
 		DisplayResourceList,
+		DisplayViewList,
+		SetViewId
 	};
 
 	enum class EInputUVMapping
@@ -115,6 +150,7 @@ private:
 	static RENDERCORE_API void DisplayHelp(FOutputDevice &Ar);
 	RENDERCORE_API void DisplayPoolResourceListToLog(ESortBy SortBy);
 	RENDERCORE_API void DisplayResourceListToLog(const TOptional<FWildcardString>& Wildcard);
+	RENDERCORE_API void DisplayViewListToLog();
 
 	/** Determine whether a texture should be captured for debugging purposes and return the capture id if needed. */
 	RENDERCORE_API TOptional<uint32> ShouldCapture(const TCHAR* DebugName, uint32 MipIndex);
@@ -155,8 +191,19 @@ private:
 
 	FConfig Config;
 
+	EState State = EState::Inactive;
+	TOptional<FWildcardString> DisplayResourcesParam;		// Cached parameter for EState::DisplayResources
+
+	bool bAnyViewRendered = false;			// Track when any view is rendered in the current frame, so we can ignore frames where no views render
+	bool bIsRequestedView = false;			// Set when this is a requested view, and we should capture visualizations from it
+	bool bFoundRequestedView = false;		// Set so we can stop considering other views, after we found the specific view that was requested
+
+	// Initialized in SetSceneTextures, tracks viewports from whichever scene renderer contains the view being visualized
+	TArray<FIntRect> FamilyViewRects;
+
 	struct FRequested
 	{
+		uint32 ViewUniqueId = 0;				// View requested to be visualized -- zero visualizes the last non-scene-capture view
 		FString Name;
 		TOptional<uint32> Version;
 	} Requested;
@@ -172,10 +219,16 @@ private:
 		FRDGTextureRef Texture = nullptr;
 		FPooledRenderTargetDesc Desc;
 		EInputValueMapping InputValueMapping = EInputValueMapping::Color;
+		int32 ViewUniqueId = 0;					// View actually visualized
+		FIntPoint OutputExtent;					// Viewport extent for visualized scene renderer
+		TArray<FIntRect> ViewRects;				// Viewports from scene renderer being visualized
 	} Captured;
 
 	ERHIFeatureLevel::Type FeatureLevel = ERHIFeatureLevel::SM5;
 
+	// Map of unique view ID to description, updated when views get rendered.
+	TMap<int32, FString> ViewDescriptionMap;
+
 	// Maps a texture name to its checkpoint version.
 	TMap<FString, uint32> VersionCountMap;
 #endif
@@ -188,3 +241,17 @@ ENUM_CLASS_FLAGS(FVisualizeTexture::EFlags);
 
 /** The global render targets for easy shading. */
 extern RENDERCORE_API TGlobalResource<FVisualizeTexture> GVisualizeTexture;
+
+#if SUPPORTS_VISUALIZE_TEXTURE
+
+// We use a macro to compile out calls to BeginViewRenderThread, because generating the arguments to the call may involve utility function calls
+// that the compiler can't optimize out, even if the function itself was an empty inline.  This commonly includes a call to the "GetViewKey"
+// function to fetch UniqueId, which involves two function calls (one virtual), and any string formatting used to generate the Description.
+// For symmetry, a macro is also provided for EndViewRenderThread (even though for that case, an empty inline would compile out fine).
+#define VISUALIZE_TEXTURE_BEGIN_VIEW(FeatureLevel, UniqueId, Description, bIsSceneCapture) GVisualizeTexture.BeginViewRenderThread(FeatureLevel, UniqueId, Description, bIsSceneCapture)
+#define VISUALIZE_TEXTURE_END_VIEW() GVisualizeTexture.EndViewRenderThread()
+
+#else
+#define VISUALIZE_TEXTURE_BEGIN_VIEW(FeatureLevel, UniqueId, Description, bIsSceneCapture) (void)0
+#define VISUALIZE_TEXTURE_END_VIEW() (void)0
+#endif
diff --git a/Engine/Source/Runtime/Renderer/Internal/SceneTextures.h b/Engine/Source/Runtime/Renderer/Internal/SceneTextures.h
--- a/Engine/Source/Runtime/Renderer/Internal/SceneTextures.h
+++ b/Engine/Source/Runtime/Renderer/Internal/SceneTextures.h
@@ -54,7 +54,7 @@ struct FMinimalSceneTextures
 struct FSceneTextures : public FMinimalSceneTextures
 {
 	// Initializes the scene textures structure in the FViewFamilyInfo
-	static RENDERER_API void InitializeViewFamily(FRDGBuilder& GraphBuilder, FViewFamilyInfo& ViewFamily);
+	static RENDERER_API void InitializeViewFamily(FRDGBuilder& GraphBuilder, FViewFamilyInfo& ViewFamily, FIntPoint FamilySize);
 	static RENDERER_API EPixelFormat GetGBufferFFormatAndCreateFlags(ETextureCreateFlags& OutCreateFlags);
 
 	// Configures an array of render targets for the GBuffer pass.
@@ -65,6 +65,9 @@ struct FSceneTextures : public FMinimalSceneTextures
 		ERenderTargetLoadAction LoadAction,
 		FRenderTargetBindingSlots& RenderTargets,
 		EGBufferLayout Layout = GBL_Default) const;
+
+	// Returns list of valid textures in this structure
+	RENDERER_API TArray<FRDGTextureRef> EnumerateSceneTextures() const;
 	
 	// (Deferred) Texture containing conservative downsampled depth for occlusion.
 	FRDGTextureRef SmallDepth{};
diff --git a/Engine/Source/Runtime/Renderer/Private/CompositionLighting/PostProcessAmbientOcclusion.cpp b/Engine/Source/Runtime/Renderer/Private/CompositionLighting/PostProcessAmbientOcclusion.cpp
--- a/Engine/Source/Runtime/Renderer/Private/CompositionLighting/PostProcessAmbientOcclusion.cpp
+++ b/Engine/Source/Runtime/Renderer/Private/CompositionLighting/PostProcessAmbientOcclusion.cpp
@@ -552,6 +552,7 @@ FScreenPassTexture AddAmbientOcclusionSetupPass(
 		Output.Texture = GraphBuilder.CreateTexture(OutputDesc, TEXT("AmbientOcclusionSetup"));
 		Output.ViewRect = OutputViewport.Rect;
 		Output.LoadAction = ERenderTargetLoadAction::ENoAction;
+		Output.UpdateVisualizeTextureExtent();
 	}
 
 	const FFinalPostProcessSettings& Settings = View.FinalPostProcessSettings;
@@ -1126,7 +1127,7 @@ FGTAOHorizonSearchOutputs AddGTAOHorizonSearchIntegratePass(
 		Output.Texture = GraphBuilder.CreateTexture(OutputDesc, TEXT("GTAOCombined"));
 		Output.ViewRect = OutputViewport.Rect;
 		Output.LoadAction = ERenderTargetLoadAction::ENoAction;
-
+		Output.UpdateVisualizeTextureExtent();
 	}
 
 	const bool bUseNormals = CVarGTAOUseNormals.GetValueOnRenderThread() >= 1;
@@ -1220,6 +1221,7 @@ FScreenPassTexture AddGTAOInnerIntegratePass(
 		Output.Texture = GraphBuilder.CreateTexture(OutputDesc, TEXT("GTAOInnerIntegrate"));
 		Output.ViewRect = OutputViewport.Rect;
 		Output.LoadAction = ERenderTargetLoadAction::ENoAction;
+		Output.UpdateVisualizeTextureExtent();
 	}
 
 	FGTAOInnerIntegratePS::FParameters* PassParameters = GraphBuilder.AllocParameters<FGTAOInnerIntegratePS::FParameters>();
@@ -1410,6 +1412,7 @@ FGTAOTemporalOutputs AddGTAOTemporalPass(
 		OutputAO.Texture = GraphBuilder.CreateTexture(OutputDesc, TEXT("GTAOTemporalOutput"));
 		OutputAO.ViewRect = OutputViewport.Rect;
 		OutputAO.LoadAction = ERenderTargetLoadAction::ENoAction;
+		OutputAO.UpdateVisualizeTextureExtent();
 	}
 
 	const FVector2f HistoryTextureSize = FVector2f(HistoryColor.Texture->Desc.Extent);
@@ -1534,6 +1537,7 @@ FScreenPassTexture AddGTAOSpatialFilter(
 		Output.Texture = GraphBuilder.CreateTexture(OutputDesc, TEXT("GTAOFilter"));
 		Output.ViewRect = OutputViewport.Rect;
 		Output.LoadAction = ERenderTargetLoadAction::ENoAction;
+		Output.UpdateVisualizeTextureExtent();
 	}
 
 	FGTAOSpatialFilterCS::FParameters* PassParameters = GraphBuilder.AllocParameters<FGTAOSpatialFilterCS::FParameters>();
diff --git a/Engine/Source/Runtime/Renderer/Private/DeferredShadingRenderer.cpp b/Engine/Source/Runtime/Renderer/Private/DeferredShadingRenderer.cpp
--- a/Engine/Source/Runtime/Renderer/Private/DeferredShadingRenderer.cpp
+++ b/Engine/Source/Runtime/Renderer/Private/DeferredShadingRenderer.cpp
@@ -1659,7 +1659,7 @@ void FDeferredShadingSceneRenderer::Render(FRDGBuilder& GraphBuilder)
 
 	::Substrate::PreInitViews(*Scene);
 
-	FSceneTextures::InitializeViewFamily(GraphBuilder, ViewFamily);
+	FSceneTextures::InitializeViewFamily(GraphBuilder, ViewFamily, FamilySize);
 	FSceneTextures& SceneTextures = GetActiveSceneTextures();
 
 	{
diff --git a/Engine/Source/Runtime/Renderer/Private/LightShaftRendering.cpp b/Engine/Source/Runtime/Renderer/Private/LightShaftRendering.cpp
--- a/Engine/Source/Runtime/Renderer/Private/LightShaftRendering.cpp
+++ b/Engine/Source/Runtime/Renderer/Private/LightShaftRendering.cpp
@@ -477,6 +477,7 @@ FRDGTextureRef FDeferredShadingSceneRenderer::RenderLightShaftOcclusion(
 						else
 						{
 							Output.ViewRect = OutputViewport.Rect;
+							Output.UpdateVisualizeTextureExtent();
 						}
 
 						if (View.State)
diff --git a/Engine/Source/Runtime/Renderer/Private/MobileShadingRenderer.cpp b/Engine/Source/Runtime/Renderer/Private/MobileShadingRenderer.cpp
--- a/Engine/Source/Runtime/Renderer/Private/MobileShadingRenderer.cpp
+++ b/Engine/Source/Runtime/Renderer/Private/MobileShadingRenderer.cpp
@@ -1019,7 +1019,7 @@ void FMobileSceneRenderer::Render(FRDGBuilder& GraphBuilder)
 	GetSceneExtensionsRenderers().PreRender(GraphBuilder);
 	GEngine->GetPreRenderDelegateEx().Broadcast(GraphBuilder);
 
-	FSceneTextures::InitializeViewFamily(GraphBuilder, ViewFamily);
+	FSceneTextures::InitializeViewFamily(GraphBuilder, ViewFamily, FamilySize);
 	FSceneTextures& SceneTextures = GetActiveSceneTextures();
 
 	FSortedLightSetSceneInfo& SortedLightSet = *GraphBuilder.AllocObject<FSortedLightSetSceneInfo>();
diff --git a/Engine/Source/Runtime/Renderer/Private/PostProcess/PostProcessDownsample.cpp b/Engine/Source/Runtime/Renderer/Private/PostProcess/PostProcessDownsample.cpp
--- a/Engine/Source/Runtime/Renderer/Private/PostProcess/PostProcessDownsample.cpp
+++ b/Engine/Source/Runtime/Renderer/Private/PostProcess/PostProcessDownsample.cpp
@@ -128,6 +128,7 @@ FScreenPassTexture AddDownsamplePass(
 		}
 		Output.ViewRect = FIntRect::DivideAndRoundUp(Inputs.SceneColor.ViewRect, 2);
 		Output.LoadAction = ERenderTargetLoadAction::ENoAction;
+		Output.UpdateVisualizeTextureExtent();
 	}
 
 	if (bIsComputePass)
diff --git a/Engine/Source/Runtime/Renderer/Private/PostProcess/PostProcessUpscale.cpp b/Engine/Source/Runtime/Renderer/Private/PostProcess/PostProcessUpscale.cpp
--- a/Engine/Source/Runtime/Renderer/Private/PostProcess/PostProcessUpscale.cpp
+++ b/Engine/Source/Runtime/Renderer/Private/PostProcess/PostProcessUpscale.cpp
@@ -205,6 +205,7 @@ FScreenPassTexture ISpatialUpscaler::AddDefaultUpscalePass(
 
 		Output.Texture = GraphBuilder.CreateTexture(OutputDesc, TEXT("Upscale"));
 		Output.LoadAction = ERenderTargetLoadAction::EClear;
+		Output.UpdateVisualizeTextureExtent();
 	}
 
 	const FScreenPassTextureViewport InputViewport(Inputs.SceneColor);
diff --git a/Engine/Source/Runtime/Renderer/Private/PostProcess/PostProcessing.cpp b/Engine/Source/Runtime/Renderer/Private/PostProcess/PostProcessing.cpp
--- a/Engine/Source/Runtime/Renderer/Private/PostProcess/PostProcessing.cpp
+++ b/Engine/Source/Runtime/Renderer/Private/PostProcess/PostProcessing.cpp
@@ -2708,6 +2708,7 @@ void AddMobilePostProcessingPasses(FRDGBuilder& GraphBuilder, FScene* Scene, con
 
 			TonemapperInputs.OverrideOutput.ViewRect = OutputViewRect;
 			TonemapperInputs.OverrideOutput.LoadAction = OutputLoadAction;
+			TonemapperInputs.OverrideOutput.UpdateVisualizeTextureExtent();
 		}
 			
 		TonemapperInputs.SceneColor = FScreenPassTextureSlice::CreateFromScreenPassTexture(GraphBuilder, SceneColor);
diff --git a/Engine/Source/Runtime/Renderer/Private/SceneCaptureRendering.cpp b/Engine/Source/Runtime/Renderer/Private/SceneCaptureRendering.cpp
--- a/Engine/Source/Runtime/Renderer/Private/SceneCaptureRendering.cpp
+++ b/Engine/Source/Runtime/Renderer/Private/SceneCaptureRendering.cpp
@@ -1044,7 +1044,14 @@ void FScene::UpdateSceneCaptureContents(USceneCaptureComponent2D* CaptureCompone
 		}
 		else if (CaptureComponent->GetOwner())
 		{
-			CaptureComponent->GetOwner()->GetFName().ToString(EventName);
+			// The label might be non-unique, so include the actor name as well
+			EventName = CaptureComponent->GetOwner()->GetActorNameOrLabel();
+
+			FName ActorName = CaptureComponent->GetOwner()->GetFName();
+			if (ActorName != EventName)
+			{
+				EventName.Appendf(TEXT(" (%s)"), *ActorName.ToString());
+			}
 		}
 		FName TargetName = TextureRenderTarget->GetFName();
 
@@ -1126,8 +1133,12 @@ void FScene::UpdateSceneCaptureContents(USceneCaptureComponent2D* CaptureCompone
 				// Don't clear the render target when compositing, or in a tiling mode that fills in the render target in multiple passes.
 				bool bClearRenderTarget = !bIsCompositing && !bEnableOrthographicTiling;
 
+				VISUALIZE_TEXTURE_BEGIN_VIEW(SceneRenderer->FeatureLevel, SceneRenderer->Views[0].GetViewKey(), *EventName, true);
+
 				UpdateSceneCaptureContent_RenderThread(RHICmdList, SceneRenderer, TextureRenderTargetResource, TextureRenderTargetResource, EventName, CopyInfo, bGenerateMips, GenerateMipsParams, bClearRenderTarget, bOrthographicCamera);
 
+				VISUALIZE_TEXTURE_END_VIEW();
+
 #if WITH_EDITOR
 				if (ViewState)
 				{
@@ -1260,9 +1271,19 @@ void FScene::UpdateSceneCaptureContents(USceneCaptureComponentCube* CaptureCompo
 			}
 			else if (CaptureComponent->GetOwner())
 			{
-				CaptureComponent->GetOwner()->GetFName().ToString(EventName);
+				// The label might be non-unique, so include the actor name as well
+				EventName = CaptureComponent->GetOwner()->GetActorNameOrLabel();
+
+				FName ActorName = CaptureComponent->GetOwner()->GetFName();
+				if (ActorName != EventName)
+				{
+					EventName.Appendf(TEXT(" (%s)"), *ActorName.ToString());
+				}
 			}
 
+			// Include the cube face index in the event name
+			EventName.Appendf(TEXT(" [%d]"), faceidx);
+
 			for (const FSceneViewExtensionRef& Extension : SceneRenderer->ViewFamily.ViewExtensions)
 			{
 				Extension->BeginRenderViewFamily(SceneRenderer->ViewFamily);
@@ -1273,9 +1294,13 @@ void FScene::UpdateSceneCaptureContents(USceneCaptureComponentCube* CaptureCompo
 			ENQUEUE_RENDER_COMMAND(CaptureCommand)(
 				[SceneRenderer, TextureRenderTarget, EventName, TargetFace](FRHICommandListImmediate& RHICmdList)
 				{
+					VISUALIZE_TEXTURE_BEGIN_VIEW(SceneRenderer->FeatureLevel, SceneRenderer->Views[0].GetViewKey(), *EventName, true);
+
 					FRHICopyTextureInfo CopyInfo;
 					CopyInfo.DestSliceIndex = TargetFace;
 					UpdateSceneCaptureContent_RenderThread(RHICmdList, SceneRenderer, TextureRenderTarget, TextureRenderTarget, EventName, CopyInfo, false, FGenerateMipsParams(), true, false);
+
+					VISUALIZE_TEXTURE_END_VIEW();
 				}
 			);
 		}
diff --git a/Engine/Source/Runtime/Renderer/Private/SceneRendering.cpp b/Engine/Source/Runtime/Renderer/Private/SceneRendering.cpp
--- a/Engine/Source/Runtime/Renderer/Private/SceneRendering.cpp
+++ b/Engine/Source/Runtime/Renderer/Private/SceneRendering.cpp
@@ -3481,8 +3481,6 @@ IVisibilityTaskData* FSceneRenderer::OnRenderBegin(FRDGBuilder& GraphBuilder)
 	
 		LightFunctionAtlas::OnRenderBegin(LightFunctionAtlas, *Scene, Views, ViewFamily);
 	
-		FVisualizeTexturePresent::OnStartRender(Views[0]);
-	
 		GraphBuilder.RHICmdList.BeginScene();
 
 		VisibilityTaskData = LaunchVisibilityTasks(GraphBuilder.RHICmdList, *this, StaticMeshUpdateTask);
@@ -4665,8 +4663,17 @@ static void RenderViewFamilies_RenderThread(FRHICommandListImmediate& RHICmdList
 		}
 		else
 		{
+			// We don't want to run visualize texture for hit proxies, so do this here.  Note that the reason we run this here
+			// rather than inside the scene renderer is to customize the user facing description string for the visualization in
+			// other code paths (scene captures), in addition to skipping it for views where it's not useful (hit proxies and
+			// utility scene renderers like water).
+			VISUALIZE_TEXTURE_BEGIN_VIEW(SceneRenderer->FeatureLevel, SceneRenderer->Views[0].GetViewKey(),
+				ViewFamily.ProfileDescription.IsEmpty() ? (ViewFamily.bResolveScene ? TEXT("ScenePrimary") : TEXT("SceneAuxiliary")) : *ViewFamily.ProfileDescription, false);
+
 			// Render the scene.
 			SceneRenderer->Render(GraphBuilder);
+
+			VISUALIZE_TEXTURE_END_VIEW();
 		}
 
 		SceneRenderer->FlushCrossGPUFences(GraphBuilder);
diff --git a/Engine/Source/Runtime/Renderer/Private/SceneTextures.cpp b/Engine/Source/Runtime/Renderer/Private/SceneTextures.cpp
--- a/Engine/Source/Runtime/Renderer/Private/SceneTextures.cpp
+++ b/Engine/Source/Runtime/Renderer/Private/SceneTextures.cpp
@@ -21,6 +21,7 @@
 #include "PostProcess/PostProcessPixelProjectedReflectionMobile.h"
 #include "IHeadMountedDisplayModule.h"
 #include "Substrate/Substrate.h"
+#include "VisualizeTexture.h"
 
 static TAutoConsoleVariable<int32> CVarSceneTargetsResizeMethod(
 	TEXT("r.SceneRenderTargetResizeMethod"),
@@ -511,7 +512,7 @@ FSceneTextureShaderParameters FMinimalSceneTextures::GetSceneTextureShaderParame
 	return OutSceneTextureShaderParameters;
 }
 
-void FSceneTextures::InitializeViewFamily(FRDGBuilder& GraphBuilder, FViewFamilyInfo& ViewFamily)
+void FSceneTextures::InitializeViewFamily(FRDGBuilder& GraphBuilder, FViewFamilyInfo& ViewFamily, FIntPoint FamilySize)
 {
 	const FSceneTexturesConfig& Config = ViewFamily.SceneTexturesConfig;
 	FSceneTextures& SceneTextures = ViewFamily.SceneTextures;
@@ -638,6 +639,20 @@ void FSceneTextures::InitializeViewFamily(FRDGBuilder& GraphBuilder, FViewFamily
 		SceneTextures.QuadOverdraw = GraphBuilder.CreateTexture(QuadOverdrawDesc, TEXT("QuadOverdrawTexture"));
 	}
 #endif
+
+#if SUPPORTS_VISUALIZE_TEXTURE
+	if (GVisualizeTexture.IsRequestedView())
+	{
+		TArray<FIntRect> FamilyViewRects;
+		FamilyViewRects.SetNumUninitialized(ViewFamily.Views.Num());
+		for (int32 ViewIndex = 0; ViewIndex < ViewFamily.Views.Num(); ViewIndex++)
+		{
+			FamilyViewRects[ViewIndex] = ViewFamily.Views[ViewIndex]->UnconstrainedViewRect;
+		}
+
+		GVisualizeTexture.SetSceneTextures(SceneTextures.EnumerateSceneTextures(), FamilySize, FamilyViewRects);
+	}
+#endif
 }
 
 uint32 FSceneTextures::GetGBufferRenderTargets(
@@ -708,6 +723,53 @@ uint32 FSceneTextures::GetGBufferRenderTargets(
 	return RenderTargetCount;
 }
 
+static void AddTextureIfNonNull(FRDGTextureRef Texture, TArray<FRDGTextureRef>& OutTextures)
+{
+	if (Texture)
+	{
+		OutTextures.Add(Texture);
+	}
+}
+
+static void AddTextureIfNonNull(const FRDGTextureMSAA& Texture, TArray<FRDGTextureRef>& OutTextures)
+{
+	if (Texture.Target)
+	{
+		OutTextures.Add(Texture.Target);
+	}
+}
+
+TArray<FRDGTextureRef> FSceneTextures::EnumerateSceneTextures() const
+{
+	TArray<FRDGTextureRef> Results;
+	Results.Reserve(20);
+
+	AddTextureIfNonNull(Color, Results);
+	AddTextureIfNonNull(Depth, Results);
+	AddTextureIfNonNull(PartialDepth, Results);
+	AddTextureIfNonNull(CustomDepth.Depth, Results);
+	AddTextureIfNonNull(SmallDepth, Results);
+	AddTextureIfNonNull(GBufferA, Results);
+	AddTextureIfNonNull(GBufferB, Results);
+	AddTextureIfNonNull(GBufferC, Results);
+	AddTextureIfNonNull(GBufferD, Results);
+	AddTextureIfNonNull(GBufferE, Results);
+	AddTextureIfNonNull(GBufferF, Results);
+	AddTextureIfNonNull(DepthAux, Results);
+	AddTextureIfNonNull(Velocity, Results);
+	AddTextureIfNonNull(MobileLocalLightTextureA, Results);
+	AddTextureIfNonNull(MobileLocalLightTextureB, Results);
+	AddTextureIfNonNull(ScreenSpaceAO, Results);
+	AddTextureIfNonNull(QuadOverdraw, Results);
+	AddTextureIfNonNull(PixelProjectedReflection, Results);
+#if WITH_EDITOR
+	AddTextureIfNonNull(EditorPrimitiveColor, Results);
+	AddTextureIfNonNull(EditorPrimitiveDepth, Results);
+#endif
+
+	return Results;
+}
+
 void FSceneTextureExtracts::QueueExtractions(FRDGBuilder& GraphBuilder, const FSceneTextures& SceneTextures)
 {
 	// Free up the memory for reuse during the RDG execution phase.
diff --git a/Engine/Source/Runtime/Renderer/Private/ScreenSpaceDenoise.cpp b/Engine/Source/Runtime/Renderer/Private/ScreenSpaceDenoise.cpp
--- a/Engine/Source/Runtime/Renderer/Private/ScreenSpaceDenoise.cpp
+++ b/Engine/Source/Runtime/Renderer/Private/ScreenSpaceDenoise.cpp
@@ -835,7 +835,8 @@ FSSDSignalTextures CreateMultiplexedTextures(
 	FRDGBuilder& GraphBuilder,
 	int32 TextureCount,
 	const TStaticArray<FRDGTextureDesc, kMaxBufferProcessingCount>& DescArray,
-	const TCHAR* const* TextureNames)
+	const TCHAR* const* TextureNames,
+	const FIntRect& Viewport)
 {
 	check(TextureCount <= kMaxBufferProcessingCount);
 	FSSDSignalTextures SignalTextures;
@@ -843,6 +844,8 @@ FSSDSignalTextures CreateMultiplexedTextures(
 	{
 		const TCHAR* TextureName = TextureNames[i];
 		SignalTextures.Textures[i] = GraphBuilder.CreateTexture(DescArray[i], TextureName);
+
+		SignalTextures.Textures[i]->EncloseVisualizeExtent(Viewport.Max);
 	}
 	return SignalTextures;
 }
@@ -1600,6 +1603,8 @@ static void DenoiseSignalAtConstantPixelDensity(
 
 			CommonParameters.CompressedMetadata[0] = GraphBuilder.CreateTexture(Desc, TEXT("DenoiserMetadata0"));
 			CommonParameters.CompressedMetadata[1] = nullptr;
+
+			CommonParameters.CompressedMetadata[0]->EncloseVisualizeExtent(Viewport.Max);
 		}
 		else
 		{
@@ -1632,7 +1637,8 @@ static void DenoiseSignalAtConstantPixelDensity(
 		FSSDSignalTextures NewSignalOutput = CreateMultiplexedTextures(
 			GraphBuilder,
 			InjestTextureCount, InjestDescs,
-			GetResourceNames(kInjestResourceNames));
+			GetResourceNames(kInjestResourceNames),
+			Viewport);
 
 		FSSDInjestCS::FParameters* PassParameters = GraphBuilder.AllocParameters<FSSDInjestCS::FParameters>();
 		PassParameters->CommonParameters = CommonParameters;
@@ -1665,7 +1671,8 @@ static void DenoiseSignalAtConstantPixelDensity(
 		FSSDSignalTextures NewSignalOutput = CreateMultiplexedTextures(
 			GraphBuilder,
 			ReconstructionTextureCount, ReconstructionDescs,
-			GetResourceNames(kReconstructionResourceNames));
+			GetResourceNames(kReconstructionResourceNames),
+			Viewport);
 
 		FSSDSpatialAccumulationCS::FParameters* PassParameters = GraphBuilder.AllocParameters<FSSDSpatialAccumulationCS::FParameters>();
 		for (int32 BatchedSignalId = 0; BatchedSignalId < Settings.SignalBatchSize; BatchedSignalId++)
@@ -1721,7 +1728,8 @@ static void DenoiseSignalAtConstantPixelDensity(
 		FSSDSignalTextures NewSignalOutput = CreateMultiplexedTextures(
 			GraphBuilder,
 			ReconstructionTextureCount, ReconstructionDescs,
-			GetResourceNames(kPreConvolutionResourceNames));
+			GetResourceNames(kPreConvolutionResourceNames),
+			Viewport);
 
 		FSSDSpatialAccumulationCS::FParameters* PassParameters = GraphBuilder.AllocParameters<FSSDSpatialAccumulationCS::FParameters>();
 		PassParameters->CommonParameters = CommonParameters;
@@ -1802,7 +1810,8 @@ static void DenoiseSignalAtConstantPixelDensity(
 				RejectionPreConvolutionSignal = CreateMultiplexedTextures(
 					GraphBuilder,
 					RejectionTextureCount, RejectionSignalProcessingDescs,
-					GetResourceNames(kRejectionPreConvolutionResourceNames));
+					GetResourceNames(kRejectionPreConvolutionResourceNames),
+					Viewport);
 			}
 
 			FSSDSpatialAccumulationCS::FParameters* PassParameters = GraphBuilder.AllocParameters<FSSDSpatialAccumulationCS::FParameters>();
@@ -1831,7 +1840,8 @@ static void DenoiseSignalAtConstantPixelDensity(
 		FSSDSignalTextures SignalOutput = CreateMultiplexedTextures(
 			GraphBuilder,
 			HistoryTextureCount, HistoryDescs,
-			GetResourceNames(kTemporalAccumulationResourceNames));
+			GetResourceNames(kTemporalAccumulationResourceNames),
+			Viewport);
 
 		FSSDTemporalAccumulationCS::FPermutationDomain PermutationVector;
 		PermutationVector.Set<FSignalProcessingDim>(Settings.SignalProcessing);
@@ -2023,7 +2033,8 @@ static void DenoiseSignalAtConstantPixelDensity(
 		FSSDSignalTextures SignalOutput = CreateMultiplexedTextures(
 			GraphBuilder,
 			HistoryTextureCount, HistoryDescs,
-			GetResourceNames(kHistoryConvolutionResourceNames));
+			GetResourceNames(kHistoryConvolutionResourceNames),
+			Viewport);
 
 		FSSDSpatialAccumulationCS::FParameters* PassParameters = GraphBuilder.AllocParameters<FSSDSpatialAccumulationCS::FParameters>();
 		PassParameters->MaxSampleCount = FMath::Clamp(MaxPostFilterSampleCount, 1, kStackowiakMaxSampleCountPerSet);
@@ -2160,7 +2171,8 @@ static void DenoiseSignalAtConstantPixelDensity(
 		*OutputSignal = CreateMultiplexedTextures(
 			GraphBuilder,
 			Settings.SignalBatchSize, OutputDescs,
-			GetResourceNames(kDenoiserOutputResourceNames));
+			GetResourceNames(kDenoiserOutputResourceNames),
+			Viewport);
 
 		FSSDSpatialAccumulationCS::FParameters* PassParameters = GraphBuilder.AllocParameters<FSSDSpatialAccumulationCS::FParameters>();
 		PassParameters->CommonParameters = CommonParameters;
@@ -2471,6 +2483,9 @@ public:
 
 				ComposedHarmonics.Textures[0] = GraphBuilder.CreateTexture(Desc, TEXT("PolychromaticPenumbraComposition0"));
 				ComposedHarmonics.Textures[1] = GraphBuilder.CreateTexture(Desc, TEXT("PolychromaticPenumbraComposition1"));
+
+				ComposedHarmonics.Textures[0]->EncloseVisualizeExtent(View.ViewRect.Max);
+				ComposedHarmonics.Textures[1]->EncloseVisualizeExtent(View.ViewRect.Max);
 			}
 
 			ComposePassParameters->CommonParameters.ViewUniformBuffer = View.ViewUniformBuffer;
diff --git a/Engine/Source/Runtime/Renderer/Private/VisualizeTexturePresent.cpp b/Engine/Source/Runtime/Renderer/Private/VisualizeTexturePresent.cpp
--- a/Engine/Source/Runtime/Renderer/Private/VisualizeTexturePresent.cpp
+++ b/Engine/Source/Runtime/Renderer/Private/VisualizeTexturePresent.cpp
@@ -9,15 +9,6 @@
 #include "RenderTargetPool.h"
 #include "SceneRendering.h"
 
-void FVisualizeTexturePresent::OnStartRender(const FViewInfo& View)
-{
-#if SUPPORTS_VISUALIZE_TEXTURE
-	GVisualizeTexture.FeatureLevel = View.GetFeatureLevel();
-	GVisualizeTexture.Captured = {};
-	GVisualizeTexture.VersionCountMap.Empty();
-#endif
-}
-
 void FVisualizeTexturePresent::PresentContent(FRDGBuilder& GraphBuilder, const FViewInfo& View, FScreenPassRenderTarget Output)
 {
 	check(Output.IsValid());
@@ -99,10 +90,11 @@ void FVisualizeTexturePresent::PresentContent(FRDGBuilder& GraphBuilder, const F
 
 	Output.LoadAction = ERenderTargetLoadAction::ELoad;
 
-	const FIntPoint BufferSizeXY = View.GetSceneTexturesConfig().Extent;
+	const FIntPoint ViewportSizeXY = Captured.OutputExtent;
+	FString ViewDescription = GVisualizeTexture.ViewDescriptionMap[Captured.ViewUniqueId];
 
 	AddDrawCanvasPass(GraphBuilder, {}, View, Output,
-		[VisualizeTexture2D, BufferSizeXY, Desc, &View, InputUVMapping, InputValueMapping](FCanvas& Canvas)
+		[VisualizeTexture2D, ViewportSizeXY, Desc, &View, InputUVMapping, InputValueMapping, ViewUniqueId = Captured.ViewUniqueId, ViewDescription, ViewRects = CopyTemp(Captured.ViewRects)](FCanvas& Canvas)
 	{
 		float X = 100 + View.UnconstrainedViewRect.Min.X;
 		float Y = 160 + View.UnconstrainedViewRect.Min.Y;
@@ -146,21 +138,24 @@ void FVisualizeTexturePresent::PresentContent(FRDGBuilder& GraphBuilder, const F
 			Canvas.DrawShadowedString(X, Y += YStep, *Line, GetStatsFont(), FLinearColor(1, 1, 1));
 		}
 		{
-			FString Line = FString::Printf(TEXT("   TextureInfoString(): %s"), *(Desc.GenerateInfoString()));
+			FString Line = FString::Printf(TEXT("View ID %d: \"%s\""), ViewUniqueId, *ViewDescription);
+			Canvas.DrawShadowedString(X, Y += YStep, *Line, GetStatsFont(), FLinearColor(1, 1, 1));
+		}
+		{
+			FString Line = FString::Printf(TEXT("  TextureInfo: %s"), *(Desc.GenerateInfoString()));
 			Canvas.DrawShadowedString(X + 10, Y += YStep, *Line, GetStatsFont(), FLinearColor(1, 1, 1));
 		}
 		{
-			FString Line = FString::Printf(TEXT("  BufferSize:(%d,%d)"), BufferSizeXY.X, BufferSizeXY.Y);
+			FString Line = FString::Printf(TEXT("  ViewportSize:(%d,%d)"), ViewportSizeXY.X, ViewportSizeXY.Y);
 			Canvas.DrawShadowedString(X + 10, Y += YStep, *Line, GetStatsFont(), FLinearColor(1, 1, 1));
 		}
 
 		const FSceneViewFamily& ViewFamily = *View.Family;
 
-		for (int32 ViewId = 0; ViewId < ViewFamily.Views.Num(); ++ViewId)
+		for (int32 ViewId = 0; ViewId < ViewRects.Num(); ++ViewId)
 		{
-			const FViewInfo* ViewIt = static_cast<const FViewInfo*>(ViewFamily.Views[ViewId]);
-			FString Line = FString::Printf(TEXT("   View #%d: (%d,%d)-(%d,%d)"), ViewId + 1,
-				ViewIt->UnscaledViewRect.Min.X, ViewIt->UnscaledViewRect.Min.Y, ViewIt->UnscaledViewRect.Max.X, ViewIt->UnscaledViewRect.Max.Y);
+			const FIntRect& ViewRect = ViewRects[ViewId];
+			FString Line = FString::Printf(TEXT("   View #%d: (%d,%d)-(%d,%d)"), ViewId + 1, ViewRect.Min.X, ViewRect.Min.Y, ViewRect.Max.X, ViewRect.Max.Y);
 			Canvas.DrawShadowedString(X + 10, Y += YStep, *Line, GetStatsFont(), FLinearColor(1, 1, 1));
 		}
 
diff --git a/Engine/Source/Runtime/Renderer/Private/VisualizeTexturePresent.h b/Engine/Source/Runtime/Renderer/Private/VisualizeTexturePresent.h
--- a/Engine/Source/Runtime/Renderer/Private/VisualizeTexturePresent.h
+++ b/Engine/Source/Runtime/Renderer/Private/VisualizeTexturePresent.h
@@ -7,9 +7,6 @@
 class FVisualizeTexturePresent
 {
 public:
-	/** Starts texture visualization capture. */
-	static void OnStartRender(const FViewInfo& View);
-
 	/** Present the visualize texture tool on screen. */
 	static void PresentContent(FRDGBuilder& GraphBuilder, const FViewInfo& View, FScreenPassRenderTarget Output);
 };
\ No newline at end of file
diff --git a/Engine/Source/Runtime/Renderer/Public/ScreenPass.h b/Engine/Source/Runtime/Renderer/Public/ScreenPass.h
--- a/Engine/Source/Runtime/Renderer/Public/ScreenPass.h
+++ b/Engine/Source/Runtime/Renderer/Public/ScreenPass.h
@@ -52,6 +52,8 @@ struct FScreenPassTexture
 	bool operator==(FScreenPassTexture Other) const;
 	bool operator!=(FScreenPassTexture Other) const;
 
+	void UpdateVisualizeTextureExtent();
+
 	FRDGTextureRef Texture = nullptr;
 	FIntRect ViewRect;
 };
diff --git a/Engine/Source/Runtime/Renderer/Public/ScreenPass.inl b/Engine/Source/Runtime/Renderer/Public/ScreenPass.inl
--- a/Engine/Source/Runtime/Renderer/Public/ScreenPass.inl
+++ b/Engine/Source/Runtime/Renderer/Public/ScreenPass.inl
@@ -72,10 +72,20 @@ inline FScreenPassTexture::FScreenPassTexture(FRDGTextureRef InTexture)
 	}
 }
 
+inline void FScreenPassTexture::UpdateVisualizeTextureExtent()
+{
+	if (Texture)
+	{
+		Texture->EncloseVisualizeExtent(ViewRect.Max);
+	}
+}
+
 inline FScreenPassTexture::FScreenPassTexture(FRDGTextureRef InTexture, FIntRect InViewRect)
 	: Texture(InTexture)
 	, ViewRect(InViewRect)
-{ }
+{
+	UpdateVisualizeTextureExtent();
+}
 
 inline FScreenPassTexture::FScreenPassTexture(const FScreenPassTextureSlice& ScreenTexture)
 	: Texture(ScreenTexture.TextureSRV->Desc.Texture)
@@ -86,6 +96,8 @@ inline FScreenPassTexture::FScreenPassTexture(const FScreenPassTextureSlice& Scr
 		check(ScreenTexture.TextureSRV->Desc.FirstArraySlice == 0);
 		check(ScreenTexture.TextureSRV->Desc.NumArraySlices == 1);
 	}
+
+	UpdateVisualizeTextureExtent();
 }
 
 inline bool FScreenPassTexture::IsValid() const
