VisualizeTexture.cpp static analysis fix -- last else condition always true.  Original code suppressed this error in its last else, as it's there deliberately -- just needed to move the error suppression markup.

#rnx

diff --git a/Engine/Source/Runtime/RenderCore/Private/VisualizeTexture.cpp b/Engine/Source/Runtime/RenderCore/Private/VisualizeTexture.cpp
--- a/Engine/Source/Runtime/RenderCore/Private/VisualizeTexture.cpp
+++ b/Engine/Source/Runtime/RenderCore/Private/VisualizeTexture.cpp
@@ -238,7 +238,7 @@ void FVisualizeTexture::ParseCommands(const TCHAR* Cmd, FOutputDevice &Ar)
 
 		DisplayPoolResourceListToLog(SortBy);
 	}
-	else if (Command == ECommand::DisplayResourceList) //-V547
+	else if (Command == ECommand::DisplayResourceList)
 	{
 		bool bListAllocated = false;
 		ESortBy SortBy = ESortBy::Index;
@@ -263,7 +263,7 @@ void FVisualizeTexture::ParseCommands(const TCHAR* Cmd, FOutputDevice &Ar)
 	{
 		DisplayViewListToLog();
 	}
-	else if (Command == ECommand::SetViewId)
+	else if (Command == ECommand::SetViewId) //-V547
 	{
 		Requested.ViewUniqueId = ViewId;
 	}
