PR #11260: Fix for PSO file cache startup mode cvar not working as expected

diff --git a/Engine/Source/Runtime/RenderCore/Private/ShaderPipelineCache.cpp b/Engine/Source/Runtime/RenderCore/Private/ShaderPipelineCache.cpp
--- a/Engine/Source/Runtime/RenderCore/Private/ShaderPipelineCache.cpp
+++ b/Engine/Source/Runtime/RenderCore/Private/ShaderPipelineCache.cpp
@@ -1590,9 +1590,6 @@ FShaderPipelineCache::FShaderPipelineCache(EShaderPlatform Platform)
 			break;
 	}
 	
-	BatchSize = CVarPSOFileCacheBatchSize.GetValueOnAnyThread();
-	BatchTime = CVarPSOFileCacheBatchTime.GetValueOnAnyThread();
-	
 	FString StableShaderKeyFile;
 	if (FParse::Value(FCommandLine::Get(), TEXT("-shkfile="), StableShaderKeyFile) && !StableShaderKeyFile.IsEmpty())
 	{
