[shaders] remove shader compilation related code deprecated in 5.3

diff --git a/Engine/Source/Developer/ShaderCompilerCommon/Private/ShaderCompilerCommon.cpp b/Engine/Source/Developer/ShaderCompilerCommon/Private/ShaderCompilerCommon.cpp
--- a/Engine/Source/Developer/ShaderCompilerCommon/Private/ShaderCompilerCommon.cpp
+++ b/Engine/Source/Developer/ShaderCompilerCommon/Private/ShaderCompilerCommon.cpp
@@ -186,20 +186,6 @@ bool BuildResourceTableMapping(
 	return true;
 }
 
-PRAGMA_DISABLE_DEPRECATION_WARNINGS
-// Deprecated version of function
-bool BuildResourceTableMapping(
-	const TMap<FString, FResourceTableEntry>& ResourceTableMap,
-	const TMap<FString, FUniformBufferEntry>& UniformBufferMap,
-	TBitArray<>& UsedUniformBufferSlots,
-	FShaderParameterMap& ParameterMap,
-	FShaderCompilerResourceTable& OutSRT)
-{
-	UE_LOG(LogShaders, Error, TEXT("Using unimplemented deprecated version of BuildResourceTableMapping -- use version that accepts FShaderResourceTableMap instead."));
-	return false;
-}
-PRAGMA_ENABLE_DEPRECATION_WARNINGS
-
 void CullGlobalUniformBuffers(const TMap<FString, FUniformBufferEntry>& UniformBufferMap, FShaderParameterMap& ParameterMap)
 {
 	TArray<FString> ParameterNames;
diff --git a/Engine/Source/Developer/ShaderCompilerCommon/Public/ShaderCompilerCommon.h b/Engine/Source/Developer/ShaderCompilerCommon/Public/ShaderCompilerCommon.h
--- a/Engine/Source/Developer/ShaderCompilerCommon/Public/ShaderCompilerCommon.h
+++ b/Engine/Source/Developer/ShaderCompilerCommon/Public/ShaderCompilerCommon.h
@@ -33,17 +33,6 @@ extern SHADERCOMPILERCOMMON_API bool BuildResourceTableMapping(
 		FShaderCompilerResourceTable& OutSRT
 	);
 
-PRAGMA_DISABLE_DEPRECATION_WARNINGS
-UE_DEPRECATED(5.3, "No longer supported; use version of function that accepts a FShaderResourceTableMap instead")
-extern SHADERCOMPILERCOMMON_API bool BuildResourceTableMapping(
-		const TMap<FString,FResourceTableEntry>& ResourceTableMap,
-		const TMap<FString,FUniformBufferEntry>& UniformBufferMap,
-		TBitArray<>& UsedUniformBufferSlots,
-		FShaderParameterMap& ParameterMap,
-		FShaderCompilerResourceTable& OutSRT
-	);
-PRAGMA_ENABLE_DEPRECATION_WARNINGS
-
 /** Culls global uniform buffer entries from the parameter map. */
 extern SHADERCOMPILERCOMMON_API void CullGlobalUniformBuffers(const TMap<FString, FUniformBufferEntry>& UniformBufferMap, FShaderParameterMap& ParameterMap);
 
@@ -133,12 +122,7 @@ namespace UE::ShaderCompilerCommon
 		TFunction<FString()> AppendPreSource{};
 		TFunction<FString()> AppendPostSource{};
 		TArray<FAdditionalOutput> AdditionalOutputs;
-		union
-		{
-			bool bSourceOnly = false; // if true, will only output source .usf as directed and skip all other debug data artifacts
-			UE_DEPRECATED(5.3, "bSkipDirectCompileTxt is deprecated, use bSourceOnly flag instead")
-			bool bSkipDirectCompileTxt;
-		};
+		bool bSourceOnly = false; // if true, will only output source .usf as directed and skip all other debug data artifacts
 
 		SHADERCOMPILERCOMMON_API FString GetDebugShaderPath(const FShaderCompilerInput& Input) const;
 	};
diff --git a/Engine/Source/Developer/TargetPlatform/Public/Interfaces/IShaderFormat.h b/Engine/Source/Developer/TargetPlatform/Public/Interfaces/IShaderFormat.h
--- a/Engine/Source/Developer/TargetPlatform/Public/Interfaces/IShaderFormat.h
+++ b/Engine/Source/Developer/TargetPlatform/Public/Interfaces/IShaderFormat.h
@@ -119,13 +119,6 @@ public:
 	 */
 	virtual void AppendToKeyString(FString& KeyString) const { }
 
-	/**
-	 * Can the shader compiler use the HLSLcc library when compiling shaders
-	 * @returns True if the shader compiler can use the HLSLcc library when compiling shaders, otherwise false.
-	 */
-	UE_DEPRECATED(5.3, "UsesHLSLcc function is no longer used")
-	virtual bool UsesHLSLcc(const FShaderCompilerInput& Input) const { return false; }
-
 	/**
 	 * Execute all shader preprocessing steps, storing the output in the PreprocessOutput struct
 	 */
diff --git a/Engine/Source/Runtime/RenderCore/Public/ShaderCompilerCore.h b/Engine/Source/Runtime/RenderCore/Public/ShaderCompilerCore.h
--- a/Engine/Source/Runtime/RenderCore/Public/ShaderCompilerCore.h
+++ b/Engine/Source/Runtime/RenderCore/Public/ShaderCompilerCore.h
@@ -129,8 +129,6 @@ enum ECompilerFlags
 	CFLAG_Wave32,
 	// Enable support of inline raytracing in compute shader.
 	CFLAG_InlineRayTracing,
-	// Force using the SC rewrite functionality before calling DXC on D3D12
-	CFLAG_D3D12ForceShaderConductorRewrite UE_DEPRECATED(5.3, "CFLAG_D3D12ForceShaderConductorRewrite has been deprecated since UE5.3 and the flag is ignored"),
 	// Enable support of C-style data types for platforms that can. Check for PLATFORM_SUPPORTS_REAL_TYPES and FDataDrivenShaderPlatformInfo::GetSupportsRealTypes()
 	CFLAG_AllowRealTypes,
 	// Precompile HLSL to optimized HLSL, then forward to FXC. Speeds up some shaders that take longer with FXC and works around crashes in FXC.
@@ -147,7 +145,6 @@ enum ECompilerFlags
 	CFLAG_ForceBindful,
 	// EXPERIMENTAL: Run the shader re-writer that removes any unused functions/resources/types from source code before compilation.
 	CFLAG_RemoveDeadCode,
-	CFLAG_UseLegacyPreprocessor UE_DEPRECATED(5.3, "Legacy preprocessor has been removed as of UE 5.3; please report any issues with the new preprocessor to the UE rendering team."),
 	// Enable CullBeforeFetch optimization on supported platforms
 	CFLAG_CullBeforeFetch,
 	// Enable WarpCulling optimization on supported platforms
@@ -245,9 +242,6 @@ struct FShaderCompilerInput
 	UE_DEPRECATED(5.4, "bSkipPreprocessedCache member is deprecated; set EShaderDebugInfoFlags::CompileFromDebugUSF on DebugInfoFlags instead.")
 	bool bSkipPreprocessedCache;
 
-	UE_DEPRECATED(5.3, "Use DebugInfoFlags field (EDebugInfoFlags::DirectCompileCommandLine)")
-	bool bGenerateDirectCompileFile;
-
 	// Indicates which additional debug outputs should be written for this compile job.
 	EShaderDebugInfoFlags DebugInfoFlags;
 
@@ -316,8 +310,8 @@ struct FShaderCompilerInput
 
 	PRAGMA_DISABLE_DEPRECATION_WARNINGS
 	// Explicitly-defaulted copy/move ctors & assignment operators are needed temporarily due to 
-	// deprecation of bGenerateDirectCompileFile field. These can be removed once the deprecation
-	// window for said field ends.
+	// deprecation of bSkipPreprocessedCache/bIndependentPreprocessed fields. These can be removed once 
+	// the deprecation window for said fields ends.
 	FShaderCompilerInput(FShaderCompilerInput&&) = default;
 	FShaderCompilerInput(const FShaderCompilerInput&) = default;
 	FShaderCompilerInput& operator=(FShaderCompilerInput&&) = default;
@@ -511,18 +505,6 @@ struct FShaderCompilerInput
 
 		return RootParametersStructure != nullptr;
 	}
-
-	/** Returns whether this shader input *can* be compiled with the legacy FXC compiler. */
-	UE_DEPRECATED(5.3, "CanCompileWithLegacyFxc doesn't have enough information to correctly flag the input as requiring FXC. Please use internal shader format code instead.")
-	bool CanCompileWithLegacyFxc() const
-	{
-		return !(Target.GetPlatform() == SP_PCD3D_SM6
-			|| IsRayTracingShader()
-			|| Environment.CompilerFlags.Contains(CFLAG_WaveOperations)
-			|| Environment.CompilerFlags.Contains(CFLAG_ForceDXC)
-			|| Environment.CompilerFlags.Contains(CFLAG_InlineRayTracing)
-			);
-	}
 };
 
 /** A shader compiler error or warning. */
@@ -626,18 +608,14 @@ struct FShaderCompilerOutput
 	double CompileTime;
 	double PreprocessTime;
 	bool bSucceeded;
-	UE_DEPRECATED(5.3, "bFailedRemovingUnused field is no longer used")
-	bool bFailedRemovingUnused;
 	bool bSupportsQueryingUsedAttributes;
-	UE_DEPRECATED(5.3, "bUsedHLSLccCompiler field is no longer used")
-	bool bUsedHLSLccCompiler;
 	bool bSerializeModifiedSource;
 	TArray<FString> UsedAttributes;
 
 	PRAGMA_DISABLE_DEPRECATION_WARNINGS
 	// Explicitly-defaulted copy/move ctors & assignment operators are needed temporarily due to 
-	// deprecation of bFailedRemovingUnused/bUsedHLSLccCompiler fields. These can be removed once the deprecation
-	// window for said fields ends.
+	// deprecation of OptionalFinalShaderSource field. These can be removed once the deprecation
+	// window for said field ends.
 	FShaderCompilerOutput(FShaderCompilerOutput&&) = default;
 	FShaderCompilerOutput(const FShaderCompilerOutput&) = default;
 	FShaderCompilerOutput& operator=(FShaderCompilerOutput&&) = default;
@@ -750,8 +728,6 @@ extern RENDERCORE_API const IShaderFormat* FindShaderFormat(FName Format, const
 
 // Executes preprocessing for the given job, if the job is marked to be preprocessed independently prior to compilation.
 extern RENDERCORE_API bool ConditionalPreprocessShader(FShaderCommonCompileJob* Job);
-UE_DEPRECATED(5.3, "Use CompileShader overload which takes an FShaderCompileJob& rather than passing input/output directly.")
-extern RENDERCORE_API void CompileShader(const TArray<const IShaderFormat*>& ShaderFormats, FShaderCompilerInput& Input, FShaderCompilerOutput& Output, const FString& WorkingDirectory, int32* CompileCount = nullptr);
 extern RENDERCORE_API void CompileShader(const TArray<const IShaderFormat*>& ShaderFormats, FShaderCompileJob& Job, const FString& WorkingDirectory, int32* CompileCount = nullptr);
 extern RENDERCORE_API void CompileShaderPipeline(const TArray<const IShaderFormat*>& ShaderFormats, FShaderPipelineCompileJob* PipelineJob, const FString& WorkingDirectory, int32* CompileCount = nullptr);
 
diff --git a/Engine/Source/Runtime/RenderCore/Public/ShaderCompilerJobTypes.h b/Engine/Source/Runtime/RenderCore/Public/ShaderCompilerJobTypes.h
--- a/Engine/Source/Runtime/RenderCore/Public/ShaderCompilerJobTypes.h
+++ b/Engine/Source/Runtime/RenderCore/Public/ShaderCompilerJobTypes.h
@@ -311,8 +311,6 @@ public:
 
 	FShaderPipelineCompileJobKey Key;
 	TArray<TRefCountPtr<FShaderCompileJob>> StageJobs;
-	UE_DEPRECATED(5.3, "bFailedRemovingUnused field is no longer used")
-	bool bFailedRemovingUnused;
 
 	virtual RENDERCORE_API FInputHash GetInputHash() override;
 	virtual RENDERCORE_API void SerializeOutput(FArchive& Ar) override;
diff --git a/Engine/Source/Runtime/RenderCore/Public/ShaderCore.h b/Engine/Source/Runtime/RenderCore/Public/ShaderCore.h
--- a/Engine/Source/Runtime/RenderCore/Public/ShaderCore.h
+++ b/Engine/Source/Runtime/RenderCore/Public/ShaderCore.h
@@ -497,9 +497,6 @@ struct FShaderCompilerEnvironment
 	FShaderResourceTableMap ResourceTableMap;
 	TMap<FString, FUniformBufferEntry> UniformBufferMap;
 
-	UE_DEPRECATED(5.3, "RemoteServerData field is deprecated (no longer used in compilation backends).")
-	TMap<FString, FString> RemoteServerData;
-
 	const ITargetPlatform* TargetPlatform = nullptr;
 
 	// Used for mobile platforms to allow per shader/material precision modes
@@ -523,7 +520,7 @@ struct FShaderCompilerEnvironment
 	virtual ~FShaderCompilerEnvironment() = default;
 
 	// Explicitly default assignment operator and copy constructor operator with warnings disabled
-	// to avoid warnings in implicitly-generated functions due to deprecation of RemoteServerData. 
+	// to avoid warnings in implicitly-generated functions due to deprecation of IncludeVirtualPathToExternalContentsMap. 
 	// These can be removed entirely (revert to implicitly-generated) when the field itself is.
 	FShaderCompilerEnvironment(const FShaderCompilerEnvironment&) = default;
 	FShaderCompilerEnvironment& operator=(const FShaderCompilerEnvironment&) = default;
@@ -633,12 +630,6 @@ struct FShaderCompilerEnvironment
 		return false;
 	}
 
-	UE_DEPRECATED(5.3, "GetDefinitions is deprecated; preprocessor defines must now only be accessed by core shader system code. Use Get/SetCompileArgument for generic params instead.")
-	const TMap<FString,FString>& GetDefinitions() const
-	{
-		return UnusedStringDefinitions;
-	}
-
 	void SetRenderTargetOutputFormat(uint32 RenderTargetIndex, EPixelFormat PixelFormat)
 	{
 		RenderTargetOutputFormatsMap.Add(RenderTargetIndex, UE_PIXELFORMAT_TO_UINT8(PixelFormat));
