Enable hair AutoLOD by default.

diff --git a/Engine/Source/Runtime/RenderCore/Private/RenderUtils.cpp b/Engine/Source/Runtime/RenderCore/Private/RenderUtils.cpp
--- a/Engine/Source/Runtime/RenderCore/Private/RenderUtils.cpp
+++ b/Engine/Source/Runtime/RenderCore/Private/RenderUtils.cpp
@@ -1691,7 +1691,7 @@ static TAutoConsoleVariable<int32> CVarHairStrandsUsesTriangleStrips(
 
 static TAutoConsoleVariable<int32> CVarHairStrandsLODMode(
 	TEXT("r.HairStrands.LODMode"),
-	0,
+	1,
 	TEXT("Enable hair strands Auto LOD mode by default. Otherwise use Manual LOD mode. Auto LOD mode adapts hair curves based on screen coverage. Manual LOD mode relies on LODs manually setup per groom asset. This global behavior can be overridden per groom asset."),
 	ECVF_RenderThreadSafe);
 
