Fix linking failure due to pgo pass error

diff --git a/Engine/Source/Runtime/RenderCore/Public/RenderGraphDefinitions.h b/Engine/Source/Runtime/RenderCore/Public/RenderGraphDefinitions.h
--- a/Engine/Source/Runtime/RenderCore/Public/RenderGraphDefinitions.h
+++ b/Engine/Source/Runtime/RenderCore/Public/RenderGraphDefinitions.h
@@ -274,6 +274,8 @@ inline int32 GetResourceTransitionPlaneForMetadataAccess(ERDGTextureMetaDataAcce
 /** HANDLE UTILITIES */
 
 /** Handle helper class for internal tracking of RDG types. */
+// Disable false positive buffer overrun warning during pgo linking step
+PRAGMA_DISABLE_BUFFER_OVERRUN_WARNING
 template <typename LocalObjectType, typename LocalIndexType>
 class TRDGHandle
 {
@@ -368,6 +370,7 @@ private:
 		return Handle.GetIndex();
 	}
 };
+PRAGMA_ENABLE_BUFFER_OVERRUN_WARNING
 
 enum class ERDGHandleRegistryDestructPolicy
 {
