[shaders]
- store shadertype associations for each shader in a shadermap in the editor only data
- output a debug artifact containing per-shadertype permutation/memory stats; note that this is not directly representative of final shader memory usage since it doesn't account for shader library deduplication or shader library chunk re-duplication; it is only intended to be used as a tool for tracking/identifying shader growth due to added/modified shadertypes
- remove the old "dumpshadercodestats" path from the shader library, as well as the stats tracking "unique" shaders and shader memory; the former is no longer used and the latter does not correctly account for library chunking
- bump shader version due to shadermap editoronly data change

diff --git a/Engine/Plugins/Compositing/OpenColorIO/Source/OpenColorIO/Private/OpenColorIOShaderMap.cpp b/Engine/Plugins/Compositing/OpenColorIO/Source/OpenColorIO/Private/OpenColorIOShaderMap.cpp
--- a/Engine/Plugins/Compositing/OpenColorIO/Source/OpenColorIO/Private/OpenColorIOShaderMap.cpp
+++ b/Engine/Plugins/Compositing/OpenColorIO/Source/OpenColorIO/Private/OpenColorIOShaderMap.cpp
@@ -433,7 +433,7 @@ FShader* FOpenColorIOShaderMap::ProcessCompilationResultsForSingleJob(const TRef
 	FShaderCompileJob* CurrentJob = SingleJob->GetSingleShaderJob();
 	check(CurrentJob->Id == CompilingId);
 
-	GetResourceCode()->AddShaderCompilerOutput(CurrentJob->Output, CurrentJob->Key.ToString());
+	GetResourceCode()->AddShaderCompilerOutput(CurrentJob->Output, CurrentJob->Key);
 
 	FShader* Shader = nullptr;
 
diff --git a/Engine/Plugins/FX/Niagara/Source/NiagaraShader/Private/NiagaraShader.cpp b/Engine/Plugins/FX/Niagara/Source/NiagaraShader/Private/NiagaraShader.cpp
--- a/Engine/Plugins/FX/Niagara/Source/NiagaraShader/Private/NiagaraShader.cpp
+++ b/Engine/Plugins/FX/Niagara/Source/NiagaraShader/Private/NiagaraShader.cpp
@@ -780,7 +780,7 @@ FShader* FNiagaraShaderMap::ProcessCompilationResultsForSingleJob(const TRefCoun
 	auto CurrentJob = SingleJob->GetSingleShaderJob();
 	check(CurrentJob->Id == CompilingId);
 
-	GetResourceCode()->AddShaderCompilerOutput(CurrentJob->Output, CurrentJob->Key.ToString());
+	GetResourceCode()->AddShaderCompilerOutput(CurrentJob->Output, CurrentJob->Key);
 
 	FShader* Shader = nullptr;
 
diff --git a/Engine/Plugins/Runtime/ComputeFramework/Source/ComputeFramework/Private/ComputeKernelShaderMap.cpp b/Engine/Plugins/Runtime/ComputeFramework/Source/ComputeFramework/Private/ComputeKernelShaderMap.cpp
--- a/Engine/Plugins/Runtime/ComputeFramework/Source/ComputeFramework/Private/ComputeKernelShaderMap.cpp
+++ b/Engine/Plugins/Runtime/ComputeFramework/Source/ComputeFramework/Private/ComputeKernelShaderMap.cpp
@@ -432,7 +432,7 @@ FShader* FComputeKernelShaderMap::ProcessCompilationResultsForSingleJob(FShaderC
 {
 	check(CurrentJob.Id == CompilingId);
 
-	GetResourceCode()->AddShaderCompilerOutput(CurrentJob.Output, CurrentJob.Key.ToString());
+	GetResourceCode()->AddShaderCompilerOutput(CurrentJob.Output, CurrentJob.Key);
 
 	FShader* Shader = nullptr;
 
diff --git a/Engine/Shaders/Public/ShaderVersion.ush b/Engine/Shaders/Public/ShaderVersion.ush
--- a/Engine/Shaders/Public/ShaderVersion.ush
+++ b/Engine/Shaders/Public/ShaderVersion.ush
@@ -3,4 +3,4 @@
 // Each time the console command is executed it generates a new GUID. As this file is included
 // in Platform.ush (which should be included in any shader) it allows to invalidate the shader DDC.
 // 
-#pragma message("UESHADERMETADATA_VERSION 3F6C865B-CABB-4223-A659-0102915647FF")
+#pragma message("UESHADERMETADATA_VERSION E521F6F8-C003-420E-8FB1-5D1A3CFC6F12")
diff --git a/Engine/Source/Editor/UnrealEd/Private/CookOnTheFlyServer.cpp b/Engine/Source/Editor/UnrealEd/Private/CookOnTheFlyServer.cpp
--- a/Engine/Source/Editor/UnrealEd/Private/CookOnTheFlyServer.cpp
+++ b/Engine/Source/Editor/UnrealEd/Private/CookOnTheFlyServer.cpp
@@ -10413,6 +10413,8 @@ void UCookOnTheFlyServer::CookByTheBookFinishedInternal()
 					// Save shader code map
 					SaveShaderLibrary(TargetPlatform, LibraryName);
 					CreatePipelineCache(TargetPlatform, LibraryName);
+
+					FShaderLibraryCooker::DumpShaderTypeStats(GShaderCompilingManager->GetAbsoluteShaderDebugInfoDirectory(), ConvertToFullSandboxPath(GetMetadataDirectory(), true, PlatformNameString));
 				}
 				if (FParse::Param(FCommandLine::Get(), TEXT("fastcook")))
 				{
diff --git a/Engine/Source/Runtime/Engine/Private/Materials/MaterialShader.cpp b/Engine/Source/Runtime/Engine/Private/Materials/MaterialShader.cpp
--- a/Engine/Source/Runtime/Engine/Private/Materials/MaterialShader.cpp
+++ b/Engine/Source/Runtime/Engine/Private/Materials/MaterialShader.cpp
@@ -2264,7 +2264,7 @@ FShader* FMaterialShaderMap::ProcessCompilationResultsForSingleJob(FShaderCompil
 		*LexToString(CurrentJob.Input.Target.GetPlatform()),
 		*LexToString(GetShaderPlatform()));
 
-	GetResourceCode()->AddShaderCompilerOutput(CurrentJob.Output, CurrentJob.Key.ToString());
+	GetResourceCode()->AddShaderCompilerOutput(CurrentJob.Output, CurrentJob.Key);
 
 #if ALLOW_SHADERMAP_DEBUG_DATA
 	CompileTime += SingleJob->Output.CompileTime;
diff --git a/Engine/Source/Runtime/Engine/Private/ShaderCompiler/ShaderCompiler.cpp b/Engine/Source/Runtime/Engine/Private/ShaderCompiler/ShaderCompiler.cpp
--- a/Engine/Source/Runtime/Engine/Private/ShaderCompiler/ShaderCompiler.cpp
+++ b/Engine/Source/Runtime/Engine/Private/ShaderCompiler/ShaderCompiler.cpp
@@ -9220,7 +9220,7 @@ FShader* FGlobalShaderTypeCompiler::FinishCompileShader(const FGlobalShaderType*
 		EShaderPlatform Platform = CurrentJob.Input.Target.GetPlatform();
 		FGlobalShaderMapSection* Section = GGlobalShaderMap[Platform]->FindOrAddSection(ShaderType);
 
-		Section->GetResourceCode()->AddShaderCompilerOutput(CurrentJob.Output, CurrentJob.Key.ToString());
+		Section->GetResourceCode()->AddShaderCompilerOutput(CurrentJob.Output, CurrentJob.Key);
 
 		if (ShaderPipelineType && !ShaderPipelineType->ShouldOptimizeUnusedOutputs(CurrentJob.Input.Target.GetPlatform()))
 		{
diff --git a/Engine/Source/Runtime/RenderCore/Private/ShaderCodeArchive.cpp b/Engine/Source/Runtime/RenderCore/Private/ShaderCodeArchive.cpp
--- a/Engine/Source/Runtime/RenderCore/Private/ShaderCodeArchive.cpp
+++ b/Engine/Source/Runtime/RenderCore/Private/ShaderCodeArchive.cpp
@@ -156,6 +156,9 @@ bool FSerializedShaderArchive::FindOrAddShader(const FSHAHash& Hash, int32& OutI
 		ShaderEntries.AddDefaulted();
 		check(ShaderEntries.Num() == ShaderHashes.Num());
 		ShaderHashTable.Add(Key, OutIndex);
+#if WITH_EDITOR
+		ShaderTypes.AddDefaulted();
+#endif
 		return true;
 	}
 
@@ -200,6 +203,18 @@ FCbWriter& operator<<(FCbWriter& Writer, const FSerializedShaderArchive& Archive
 			Writer.EndArray();
 		}
 		Writer.EndArray();
+
+		Writer.BeginArray("ShaderTypes");
+		for (const TArray<uint64>& SingleShaderTypes : Archive.ShaderTypes)
+		{
+			Writer.BeginArray();
+			for (uint64 ShaderTypeHash : SingleShaderTypes)
+			{
+				Writer << ShaderTypeHash;
+			}
+			Writer.EndArray();
+		}
+		Writer.EndArray();
 	}
 	Writer.EndObject();
 	return Writer;
@@ -218,29 +233,52 @@ bool LoadFromCompactBinary(FCbFieldView Field, FSerializedShaderArchive& OutArch
 		}
 	}
 
-	FCbFieldView ShaderCodeToAssetsField = Field["ShaderCodeToAssets"];
-	// Map size is array size divided by two because pairs are written as successive elements
-	int32 NumShaderCodeToAssets = ShaderCodeToAssetsField.AsArrayView().Num()/2;
-	bool bOk = !ShaderCodeToAssetsField.HasError();
-	OutArchive.ShaderCodeToAssets.Empty(NumShaderCodeToAssets);
-	FCbFieldViewIterator It = ShaderCodeToAssetsField.CreateViewIterator();
-	while (It)
+	bool bOk = true;
+
 	{
-		FSHAHash ShaderMapHash;
-		if (!LoadFromCompactBinary(*It++, ShaderMapHash))
+		FCbFieldView ShaderCodeToAssetsField = Field["ShaderCodeToAssets"];
+		// Map size is array size divided by two because pairs are written as successive elements
+		int32 NumShaderCodeToAssets = ShaderCodeToAssetsField.AsArrayView().Num() / 2;
+		bOk = !ShaderCodeToAssetsField.HasError();
+		OutArchive.ShaderCodeToAssets.Empty(NumShaderCodeToAssets);
+		FCbFieldViewIterator It = ShaderCodeToAssetsField.CreateViewIterator();
+		while (It)
 		{
-			bOk = false;
-			continue;
+			FSHAHash ShaderMapHash;
+			if (!LoadFromCompactBinary(*It++, ShaderMapHash))
+			{
+				bOk = false;
+				continue;
+			}
+			FShaderMapAssetPaths& Paths = OutArchive.ShaderCodeToAssets.FindOrAdd(ShaderMapHash);
+			FCbFieldView AssetNameArrayField = *It++;
+			Paths.Reserve((*It).AsArrayView().Num());
+			bOk = (!AssetNameArrayField.HasError()) && bOk;
+			for (FCbFieldView AssetNameField : AssetNameArrayField)
+			{
+				FName AssetName;
+				bOk = LoadFromCompactBinary(AssetNameField, AssetName) && bOk;
+				Paths.Add(AssetName);
+			}
 		}
-		FShaderMapAssetPaths& Paths = OutArchive.ShaderCodeToAssets.FindOrAdd(ShaderMapHash);
-		FCbFieldView AssetNameArrayField = *It++;
-		Paths.Reserve((*It).AsArrayView().Num());
-		bOk = (!AssetNameArrayField.HasError()) & bOk;
-		for (FCbFieldView AssetNameField : AssetNameArrayField)
+	}
+
+	{
+		FCbFieldView ShaderTypesField = Field["ShaderTypes"];
+		bOk = !ShaderTypesField.HasError() && bOk;
+		OutArchive.ShaderTypes.Empty(ShaderTypesField.AsArrayView().Num());
+		FCbFieldViewIterator It = ShaderTypesField.CreateViewIterator();
+		while (It)
 		{
-			FName AssetName;
-			bOk = LoadFromCompactBinary(AssetNameField, AssetName) & bOk;
-			Paths.Add(AssetName);
+			bOk = !It->HasError() && bOk;
+			TArray<uint64>& ShaderTypeHashes = OutArchive.ShaderTypes.AddDefaulted_GetRef();
+			FCbArrayView ArrayView = It.AsArrayView();
+			ShaderTypeHashes.Reserve(ArrayView.Num());
+			for (FCbFieldViewIterator ArrayIt = It.AsArrayView().CreateViewIterator(); ArrayIt; ++ArrayIt)
+			{
+				ShaderTypeHashes.Add(ArrayIt.AsUInt64());
+			}
+			It++;
 		}
 	}
 
@@ -739,6 +777,9 @@ void FSerializedShaderArchive::CreateAsChunkFrom(const FSerializedShaderArchive&
 
 						// copy the entry as is
 						ShaderEntries[ShaderIndex] = Parent.ShaderEntries[ParentShaderIndex];
+#if WITH_EDITOR
+						ShaderTypes[ShaderIndex] = Parent.ShaderTypes[ParentShaderIndex];
+#endif
 					}
 					ShaderIndices[ShaderMapDescriptor.ShaderIndicesOffset + ShaderIdx] = ShaderIndex;
 				}
diff --git a/Engine/Source/Runtime/RenderCore/Private/ShaderCodeLibrary.cpp b/Engine/Source/Runtime/RenderCore/Private/ShaderCodeLibrary.cpp
--- a/Engine/Source/Runtime/RenderCore/Private/ShaderCodeLibrary.cpp
+++ b/Engine/Source/Runtime/RenderCore/Private/ShaderCodeLibrary.cpp
@@ -31,6 +31,7 @@ ShaderCodeLibrary.cpp: Bound shader state cache implementation.
 #include "Misc/StringBuilder.h"
 #include "PipelineFileCache.h"
 #include "ProfilingDebugging/LoadTimeTracker.h"
+#include "ProfilingDebugging/DiagnosticTable.h"
 #include "RenderingThread.h"
 #include "Shader.h"
 #include "ShaderCodeArchive.h"
@@ -1284,11 +1285,14 @@ bool FShaderMapResource_SharedCode::TryRelease()
 #if WITH_EDITOR
 struct FShaderCodeStats
 {
-	int64 ShadersSize;
-	int64 ShadersUniqueSize;
-	int32 NumShaders;
-	int32 NumUniqueShaders;
+	struct FPerTypeStats
+	{
+		int64 ShadersSize = 0;
+		int32 NumShaders = 0;
+	};
 	int32 NumShaderMaps;
+
+	TMap<uint64, FPerTypeStats> ShaderTypeStats;
 };
 
 struct FEditorShaderCodeArchive
@@ -1352,9 +1356,9 @@ struct FEditorShaderCodeArchive
 			UE_LOG(LogShaderLibrary, Warning, TEXT("Shadermap %s does not have assets associated with it, library layout may be inconsistent between builds"), *Code->ResourceHash.ToString());
 		}
 
+		const int32 NumShaders = Code->ShaderEntries.Num();
 		if (SerializedShaders.FindOrAddShaderMap(Code->ResourceHash, ShaderMapIndex, &AssociatedAssets))
 		{
-			const int32 NumShaders = Code->ShaderEntries.Num();
 			FShaderMapEntry& ShaderMapEntry = SerializedShaders.ShaderMapEntries[ShaderMapIndex];
 			ShaderMapEntry.NumShaders = NumShaders;
 			ShaderMapEntry.ShaderIndicesOffset = SerializedShaders.ShaderIndices.AddZeroed(NumShaders);
@@ -1365,26 +1369,38 @@ struct FEditorShaderCodeArchive
 				if (SerializedShaders.FindOrAddShader(Code->ShaderHashes[i], ShaderIndex))
 				{
 					const FShaderMapResourceCode::FShaderEntry& SourceShaderEntry = Code->ShaderEntries[i];
+					int32 CodeSize = SourceShaderEntry.Code.Num();
 					FShaderCodeEntry& SerializedShaderEntry = SerializedShaders.ShaderEntries[ShaderIndex];
 					SerializedShaderEntry.Frequency = SourceShaderEntry.Frequency;
-					SerializedShaderEntry.Size = SourceShaderEntry.Code.Num();
+					SerializedShaderEntry.Size = CodeSize;
 					SerializedShaderEntry.UncompressedSize = SourceShaderEntry.UncompressedSize;
 					check(!SourceShaderEntry.Code.IsEmpty());
 					ShaderCode.Add(SourceShaderEntry.Code);
 					check(ShaderCode.Num() == SerializedShaders.ShaderEntries.Num());
-
-					CodeStats.NumUniqueShaders++;
-					CodeStats.ShadersUniqueSize += SourceShaderEntry.Code.Num();
 				}
-				CodeStats.ShadersSize += Code->ShaderEntries[i].Code.Num();
 				SerializedShaders.ShaderIndices[ShaderMapEntry.ShaderIndicesOffset + i] = ShaderIndex;
 			}
 
-			// for total shaders, only count shaders when we're adding a new shadermap. AddShaderCode() for the same shadermap can be called several times during
-			// the cook because of serialization path being reused for other purposes than actual saving, so counting them every time artificially inflates number of shaders.
-			CodeStats.NumShaders += Code->ShaderEntries.Num();
 			CodeStats.NumShaderMaps++;
 		}
+
+		for (int32 i = 0; i < NumShaders; ++i)
+		{
+			for (uint64 ShaderTypeHash : Code->ShaderEditorOnlyDataEntries[i].ShaderTypeHashes)
+			{
+				int32 ShaderIndex = SerializedShaders.FindShader(Code->ShaderHashes[i]);
+				bool bExists = SerializedShaders.ShaderTypes[ShaderIndex].Find(ShaderTypeHash) != INDEX_NONE;
+				if (!bExists)
+				{
+					SerializedShaders.ShaderTypes[ShaderIndex].Add(ShaderTypeHash);
+					FShaderCodeStats::FPerTypeStats& TypeCodeStats = CodeStats.ShaderTypeStats.FindOrAdd(ShaderTypeHash);
+					int32 CodeSize = SerializedShaders.ShaderEntries[ShaderIndex].Size;
+					TypeCodeStats.NumShaders++;
+					TypeCodeStats.ShadersSize += CodeSize;
+				}
+			}
+		}
+
 		// always mark the shadermap dirty, because it might have gotten new asset associations
 		MarkShaderMapDirty(ShaderMapIndex);
 		return ShaderMapIndex;
@@ -1516,6 +1532,10 @@ struct FEditorShaderCodeArchive
 						++NumShadersSentWithCode;
 
 						TargetFlatShaderCode.Append(SourceShaderCode);
+						for (uint64 ShaderTypeHash : SerializedShaders.ShaderTypes[SourceShaderIndex])
+						{
+							TargetArchive.ShaderTypes[TargetShaderIndex].Add(ShaderTypeHash);
+						}
 
 						// Empty the ShaderCode to save memory in the local process. The consumer of the TargetArchive and
 						// TargetFlatShaderCode will be the only one that needs to read it.
@@ -1548,6 +1568,24 @@ struct FEditorShaderCodeArchive
 			const FSHAHash& SourceShaderHash = SourceArchive.ShaderHashes[SourceShaderIndex];
 			int32 TargetShaderIndex = INDEX_NONE;
 			const bool bShaderIsNew = TargetArchive.FindOrAddShader(SourceShaderHash, TargetShaderIndex);
+			const TArray<uint64>& ShaderTypeHashes = SourceArchive.ShaderTypes[SourceShaderIndex];
+			const FShaderCodeEntry& SourceShaderEntry = SourceArchive.ShaderEntries[SourceShaderIndex];
+
+			TArray<uint64>& ExistingShaderTypeHashes = TargetArchive.ShaderTypes[TargetShaderIndex];
+			bool bAnyTypesAdded = false;
+			for (uint64 ShaderTypeHash : ShaderTypeHashes)
+			{
+				bool bExists = ExistingShaderTypeHashes.Find(ShaderTypeHash) != INDEX_NONE;
+				if (!bExists)
+				{
+					ExistingShaderTypeHashes.Add(ShaderTypeHash);
+					bAnyTypesAdded = true;
+					FShaderCodeStats::FPerTypeStats& TypeCodeStats = CodeStats.ShaderTypeStats.FindOrAdd(ShaderTypeHash);
+					TypeCodeStats.NumShaders++;
+					TypeCodeStats.ShadersSize += SourceShaderEntry.Size;
+				}
+			}
+
 			if (!bShaderIsNew)
 			{
 				continue;
@@ -1556,14 +1594,10 @@ struct FEditorShaderCodeArchive
 				TargetShaderCodes.Num() == TargetArchive.ShaderEntries.Num() - 1);
 			TArray<uint8>& TargetShaderCode = TargetShaderCodes.Emplace_GetRef();
 
-			const FShaderCodeEntry& SourceShaderEntry = SourceArchive.ShaderEntries[SourceShaderIndex];
 			FShaderCodeEntry& TargetShaderEntry = TargetArchive.ShaderEntries[TargetShaderIndex];
 			TargetShaderEntry = SourceShaderEntry;
 			TargetShaderEntry.Offset = 0;
 
-			CodeStats.NumUniqueShaders++;
-			CodeStats.ShadersUniqueSize += SourceShaderEntry.Size;
-
 			if (SourceShaderEntry.Offset == INDEX_NONE)
 			{
 				UE_LOG(LogShaderLibrary, Error, TEXT("ShaderMapLibrary transfer received from a remote machine has incomplete record for Shader %s. ")
@@ -1622,8 +1656,6 @@ struct FEditorShaderCodeArchive
 				// Every shader in the SourceArchive should have already been added by the loop above over SourceArchive.ShaderHashes
 				check(!bShaderIsNew);
 			}
-
-			CodeStats.NumShaders += TargetEntry.NumShaders; // Sum of shader counts used by each ShaderMap, without removing duplicates
 			CodeStats.NumShaderMaps++;
 		}
 		return bOk;
@@ -1677,6 +1709,7 @@ struct FEditorShaderCodeArchive
 				{
 					const FShaderCodeEntry& OtherShaderEntry = OtherArchive.SerializedShaders.ShaderEntries[OtherShaderIndex];
 					SerializedShaders.ShaderEntries[ShaderIndex] = OtherShaderEntry;
+					SerializedShaders.ShaderTypes[ShaderIndex] = OtherArchive.SerializedShaders.ShaderTypes[OtherShaderIndex];
 
 					const TArray<uint8>& OtherShaderCodeEntry = OtherArchive.ShaderCode[OtherShaderIndex];
 					check(!OtherShaderCodeEntry.IsEmpty());
@@ -1715,6 +1748,7 @@ struct FEditorShaderCodeArchive
 				{
 					const FShaderCodeEntry& OtherShaderEntry = OtherShaders.ShaderEntries[OtherShaderIndex];
 					SerializedShaders.ShaderEntries[ShaderIndex] = OtherShaderEntry;
+					SerializedShaders.ShaderTypes[ShaderIndex] = OtherShaders.ShaderTypes[OtherShaderIndex];
 
 					TArray<uint8>& Code = ShaderCode.AddDefaulted_GetRef();
 					check(ShaderCode.Num() == SerializedShaders.GetNumShaders());
@@ -3370,26 +3404,68 @@ public:
 		return bOK;
 	}
 
-	void DumpShaderCodeStats()
+	void DumpShaderTypeStats(const FString& DebugInfoDir, const FString& MetaDataDir)
 	{
 		int32 PlatformId = 0;
-		for (const FShaderCodeStats& CodeStats : EditorShaderCodeStats)
+		for (const FShaderCodeStats& PlatformCodeStats : EditorShaderCodeStats)
 		{
-			if (CodeStats.NumShaders > 0)
+			if (PlatformCodeStats.NumShaderMaps > 0)
 			{
-				float UniqueSize = CodeStats.ShadersUniqueSize;
-				float UniqueSizeMB = FUnitConversion::Convert(UniqueSize, EUnit::Bytes, EUnit::Megabytes);
-				float TotalSize = CodeStats.ShadersSize;
-				float TotalSizeMB = FUnitConversion::Convert(TotalSize, EUnit::Bytes, EUnit::Megabytes);
-
-				UE_LOG(LogShaderLibrary, Display, TEXT(""));
-				UE_LOG(LogShaderLibrary, Display, TEXT("Shader Code Stats: %s"), *LegacyShaderPlatformToShaderFormat((EShaderPlatform)PlatformId).ToString());
-				UE_LOG(LogShaderLibrary, Display, TEXT("================="));
-				UE_LOG(LogShaderLibrary, Display, TEXT("Unique Shaders: %d, Total Shaders: %d, Unique Shadermaps: %d"), CodeStats.NumUniqueShaders, CodeStats.NumShaders, CodeStats.NumShaderMaps);
-				UE_LOG(LogShaderLibrary, Display, TEXT("Unique Shaders Size: %.2fmb, Total Shader Size: %.2fmb"), UniqueSizeMB, TotalSizeMB);
-				UE_LOG(LogShaderLibrary, Display, TEXT("================="));
-			}
+				FString PlatformName = FGenericDataDrivenShaderPlatformInfo::GetName((EShaderPlatform)PlatformId).ToString();
+				FString Filename = FString::Printf(TEXT("%s/%s/ShaderTypeStats.csv"), *DebugInfoDir, *PlatformName);
+				TMap<const FShaderType*, FShaderCodeStats::FPerTypeStats> CountersByTypeName;
+				for (const TPair<uint64, FShaderCodeStats::FPerTypeStats>& Pair : PlatformCodeStats.ShaderTypeStats)
+				{
+					const FShaderType* ShaderType = FindShaderTypeByName(FHashedName(Pair.Key));
+					CountersByTypeName.Add(ShaderType, Pair.Value);
+				}
+				
+				// sort first by the type of shadertype, then by the name (so all global, material, etc. shaders are grouped together in the CSV)
+				struct FShaderTypeSort
+				{
+					bool operator()(const FShaderType& A, const FShaderType& B) const
+					{
+						if (A.GetTypeForDynamicCast() == B.GetTypeForDynamicCast())
+						{
+							return FCString::Strcmp(A.GetName(), B.GetName()) < 0;
+						}
+						return A.GetTypeForDynamicCast() < B.GetTypeForDynamicCast();
+					}
+				};
+				CountersByTypeName.KeySort(FShaderTypeSort());
 
+				{
+					TUniquePtr<FArchive> ShaderTypeStatsCsv(IFileManager::Get().CreateFileWriter(*Filename));
+					if (ShaderTypeStatsCsv)
+					{
+						FDiagnosticTableWriterCSV StatWriter(ShaderTypeStatsCsv.Get());
+						StatWriter.AddColumn(TEXT("ShaderTypeName"));
+						StatWriter.AddColumn(TEXT("ShaderType"));
+						StatWriter.AddColumn(TEXT("Frequency"));
+						StatWriter.AddColumn(TEXT("NumShaders"));
+						StatWriter.AddColumn(TEXT("ShadersSize"));
+						StatWriter.CycleRow();
+
+						for (const TPair<const FShaderType*, FShaderCodeStats::FPerTypeStats>& Pair : CountersByTypeName)
+						{
+							StatWriter.AddColumn(Pair.Key->GetName());
+							StatWriter.AddColumn(*LexToString(Pair.Key->GetTypeForDynamicCast()));
+							StatWriter.AddColumn(GetShaderFrequencyString(Pair.Key->GetFrequency()));
+							StatWriter.AddColumn(TEXT("%u"), Pair.Value.NumShaders);
+							StatWriter.AddColumn(TEXT("%u"), Pair.Value.ShadersSize);
+							StatWriter.CycleRow();
+						}
+					}
+					else
+					{
+						UE_LOG(LogShaderLibrary, Warning, TEXT("Failed to create shader type stats file '%s'"), *Filename);
+					}
+				}
+
+				// Add a copy of the type stats to the build metadata as well
+				FString MetaFilename = FString::Printf(TEXT("%s/ShaderTypeStats-%s.csv"), *MetaDataDir, *PlatformName);
+				IFileManager::Get().Copy(*MetaFilename, *Filename);
+			}
 			PlatformId++;
 		}
 	}
@@ -3779,8 +3855,6 @@ void FShaderLibraryCooker::Shutdown()
 {
 	if (FShaderLibrariesCollection::Impl)
 	{
-		//DumpShaderCodeStats();
-
 		delete FShaderLibrariesCollection::Impl;
 		FShaderLibrariesCollection::Impl = nullptr;
 	}
@@ -3875,11 +3949,11 @@ void FShaderLibraryCooker::AddShaderStableKeyValue(EShaderPlatform ShaderPlatfor
 	}
 }
 
-void FShaderLibraryCooker::DumpShaderCodeStats()
+void FShaderLibraryCooker::DumpShaderTypeStats(const FString& DebugInfoDir, const FString& MetaDataDir)
 {
 	if (FShaderLibrariesCollection::Impl)
 	{
-		FShaderLibrariesCollection::Impl->DumpShaderCodeStats();
+		FShaderLibrariesCollection::Impl->DumpShaderTypeStats(DebugInfoDir, MetaDataDir);
 	}
 }
 
diff --git a/Engine/Source/Runtime/RenderCore/Private/ShaderResource.cpp b/Engine/Source/Runtime/RenderCore/Private/ShaderResource.cpp
--- a/Engine/Source/Runtime/RenderCore/Private/ShaderResource.cpp
+++ b/Engine/Source/Runtime/RenderCore/Private/ShaderResource.cpp
@@ -18,6 +18,7 @@
 #include "UObject/RenderingObjectVersion.h"
 #include "Misc/MemStack.h"
 #include "ShaderCompilerCore.h"
+#include "ShaderCompilerJobTypes.h"
 #include "Compression/OodleDataCompression.h"
 #include "RHIResources.h"	// Access to FRHIRayTracingShader::RayTracingPayloadType requires this
 #include "DataDrivenShaderPlatformInfo.h"
@@ -237,20 +238,21 @@ int32 FShaderMapResourceCode::FindShaderIndex(const FSHAHash& InHash) const
 	return Algo::BinarySearch(ShaderHashes, InHash);
 }
 
-void FShaderMapResourceCode::AddShaderCompilerOutput(const FShaderCompilerOutput& Output, const FString& DebugName)
+void FShaderMapResourceCode::AddShaderCompilerOutput(const FShaderCompilerOutput& Output, const FShaderCompileJobKey& Key)
 {
 	TRACE_CPUPROFILER_EVENT_SCOPE(FShaderMapResourceCode::AddShaderCode);
 
 	const FSHAHash& InHash = Output.OutputHash;
 	const FShaderCode& InCode = Output.ShaderCode;
 	const int32 Index = Algo::LowerBound(ShaderHashes, InHash);
+	FString DebugName = Key.ToString();
 	if (Index >= ShaderHashes.Num() || ShaderHashes[Index] != InHash)
 	{
 		ShaderHashes.Insert(InHash, Index);
 
 #if WITH_EDITORONLY_DATA
 		// Output.Errors contains warnings in the case any exist (no errors since if there were the job would have failed)
-		AddEditorOnlyData(Index, DebugName, Output.PlatformDebugData, Output.Errors);
+		AddEditorOnlyData(Index, DebugName, Key.ShaderType->GetHashedName().GetHash(), Output.PlatformDebugData, Output.Errors);
 #endif
 
 		FShaderEntry& Entry = ShaderEntries.InsertDefaulted_GetRef(Index);
@@ -318,23 +320,23 @@ void FShaderMapResourceCode::AddShaderCompilerOutput(const FShaderCompilerOutput
 	else
 	{
 		// Output.Errors contains warnings in the case any exist (no errors since if there were the job would have failed)
-		// We append the warnings for any additional jobs which resulted in the same bytecode for the sake of determinism in the
-		// results saved to DDC. 
-		AppendWarningsToEditorOnlyData(Index, DebugName, Output.Errors);
+		// We append the warnings and shadertype hash for any additional jobs which resulted in the same bytecode for the 
+		// sake of determinism in the results saved to DDC. 
+		AppendToEditorOnlyData(Index, DebugName, Key.ShaderType->GetHashedName().GetHash(), Output.Errors);
 	}
 #endif
 }
 
 #if WITH_EDITORONLY_DATA
-void FShaderMapResourceCode::AddEditorOnlyData(int32 Index, const FString& DebugName, TConstArrayView<uint8> InPlatformDebugData, TConstArrayView<FShaderCompilerError> InCompilerWarnings)
+void FShaderMapResourceCode::AddEditorOnlyData(int32 Index, const FString& DebugName, uint64 ShaderTypeHash, TConstArrayView<uint8> InPlatformDebugData, TConstArrayView<FShaderCompilerError> InCompilerWarnings)
 {
 	FShaderEditorOnlyDataEntry& Entry = ShaderEditorOnlyDataEntries.InsertDefaulted_GetRef(Index);
 	Entry.PlatformDebugData = InPlatformDebugData;
 
-	AppendWarningsToEditorOnlyData(Index, DebugName, InCompilerWarnings);
+	AppendToEditorOnlyData(Index, DebugName, ShaderTypeHash, InCompilerWarnings);
 }
 
-void FShaderMapResourceCode::AppendWarningsToEditorOnlyData(int32 Index, const FString& DebugName, TConstArrayView<FShaderCompilerError> InCompilerWarnings)
+void FShaderMapResourceCode::AppendToEditorOnlyData(int32 Index, const FString& DebugName, uint64 ShaderTypeHash, TConstArrayView<FShaderCompilerError> InCompilerWarnings)
 {
 	FShaderEditorOnlyDataEntry& Entry = ShaderEditorOnlyDataEntries[Index];
 	for (const FShaderCompilerError& Warning : InCompilerWarnings)
@@ -347,6 +349,7 @@ void FShaderMapResourceCode::AppendWarningsToEditorOnlyData(int32 Index, const F
 			Entry.CompilerWarnings.Insert(ModifiedWarning, WarningIndex);
 		}
 	}
+	Entry.ShaderTypeHashes.Add(ShaderTypeHash);
 }
 
 void FShaderMapResourceCode::LogShaderCompilerWarnings()
diff --git a/Engine/Source/Runtime/RenderCore/Public/Shader.h b/Engine/Source/Runtime/RenderCore/Public/Shader.h
--- a/Engine/Source/Runtime/RenderCore/Public/Shader.h
+++ b/Engine/Source/Runtime/RenderCore/Public/Shader.h
@@ -79,6 +79,7 @@ class FShaderPipelineType;
 class FShaderType;
 class FVertexFactoryType;
 struct FShaderCompiledShaderInitializerType;
+struct FShaderCompileJobKey;
 struct FShaderCompilerOutput;
 using FShaderMapAssetPaths = TSet<FName>; // Copied from ShaderCodeLibrary.h
 
@@ -492,9 +493,13 @@ public:
 		 *  Does not contain errors since if there were any errors, this object wouldn't exist. */
 		TArray<FString> CompilerWarnings;
 
+		/** An array of the hash of all FShaderType names which generate this particular shader (i.e. the return value
+		 * of FShaderType::GetHashedName().GetHash()). */
+		TArray<uint64> ShaderTypeHashes;
+
 		friend FArchive& operator<<(FArchive& Ar, FShaderEditorOnlyDataEntry& Entry)
 		{
-			return Ar << Entry.PlatformDebugData << Entry.CompilerWarnings;
+			return Ar << Entry.PlatformDebugData << Entry.ShaderTypeHashes << Entry.CompilerWarnings;
 		}
 	};
 #endif // WITH_EDITORONLY_DATA
@@ -512,13 +517,16 @@ public:
 
 	RENDERCORE_API uint32 GetSizeBytes() const;
 
-	RENDERCORE_API void AddShaderCompilerOutput(const FShaderCompilerOutput& Output, const FString& DebugName = FString());
+	UE_DEPRECATED(5.4, "AddShaderCompilerOutput now accepts an FShaderCompileJobKey instead of DebugName")
+	inline void AddShaderCompilerOutput(const FShaderCompilerOutput& Output, const FString& DebugName = FString()) {}
+
+	RENDERCORE_API void AddShaderCompilerOutput(const FShaderCompilerOutput& Output, const FShaderCompileJobKey& Key);
 
 	int32 FindShaderIndex(const FSHAHash& InHash) const;
 
 #if WITH_EDITORONLY_DATA
-	void AddEditorOnlyData(int32 Index, const FString& DebugName, TConstArrayView<uint8> InPlatformDebugData, TConstArrayView<FShaderCompilerError> InCompilerWarnings);
-	void AppendWarningsToEditorOnlyData(int32 Index, const FString& DebugName, TConstArrayView<FShaderCompilerError> InCompilerWarnings);
+	void AddEditorOnlyData(int32 Index, const FString& DebugName, uint64 ShaderTypeHash, TConstArrayView<uint8> InPlatformDebugData, TConstArrayView<FShaderCompilerError> InCompilerWarnings);
+	void AppendToEditorOnlyData(int32 Index, const FString& DebugName, uint64 ShaderTypeHash, TConstArrayView<FShaderCompilerError> InCompilerWarnings);
 	RENDERCORE_API void LogShaderCompilerWarnings();
 #endif
 
@@ -1383,6 +1391,7 @@ public:
 		return ShaderTypeForDynamicCast;
 	}
 
+
 	// Accessors.
 	inline const FTypeLayoutDesc& GetLayout() const
 	{
@@ -1495,6 +1504,23 @@ protected:
 #endif // WITH_EDITOR
 };
 
+inline FString LexToString(FShaderType::EShaderTypeForDynamicCast ShaderType)
+{
+	static const TCHAR* TypeStrings[] =
+	{
+		TEXT("Global"),
+		TEXT("Material"),
+		TEXT("MeshMaterial"),
+		TEXT("Niagara"),
+		TEXT("OCIO"),
+		TEXT("ComputeKernel")
+	};
+	checkf(ShaderType < FShaderType::EShaderTypeForDynamicCast::NumShaderTypes, TEXT("Invalid shader type"));
+
+	return TypeStrings[static_cast<uint32>(ShaderType)];
+}
+
+
 /**
  * Registers a shader type in various systems. Should be created as a static field/global.
  * 
diff --git a/Engine/Source/Runtime/RenderCore/Public/ShaderCodeArchive.h b/Engine/Source/Runtime/RenderCore/Public/ShaderCodeArchive.h
--- a/Engine/Source/Runtime/RenderCore/Public/ShaderCodeArchive.h
+++ b/Engine/Source/Runtime/RenderCore/Public/ShaderCodeArchive.h
@@ -105,6 +105,9 @@ public:
 	/** Mapping from shadermap hashes to an array of asset names - this is used for on-disk storage as it is shorter. */
 	TMap<FSHAHash, FShaderMapAssetPaths> ShaderCodeToAssets;
 
+	/** Array of shader types for each individual shader; indexed as ShaderEntries array */
+	TArray<TArray<uint64>> ShaderTypes;
+
 	enum class EAssetInfoVersion : uint8
 	{
 		CurrentVersion = 2
@@ -162,6 +165,7 @@ public:
 			ShaderIndices.GetAllocatedSize()
 #if WITH_EDITOR
 			+ ShaderCodeToAssets.GetAllocatedSize()
+			+ ShaderTypes.GetAllocatedSize()
 #endif // WITH_EDITOR
 			;
 	}
@@ -184,6 +188,7 @@ public:
 		ShaderMapHashTable.Clear();
 #if WITH_EDITOR
 		ShaderCodeToAssets.Empty();
+		ShaderTypes.Empty();
 #endif
 	}
 
@@ -202,6 +207,7 @@ public:
 		return ShaderMapEntries.IsEmpty() && ShaderEntries.IsEmpty() && PreloadEntries.IsEmpty()
 #if WITH_EDITOR
 			&& ShaderCodeToAssets.IsEmpty()
+			&& ShaderTypes.IsEmpty()
 #endif
 			;
 	}
diff --git a/Engine/Source/Runtime/RenderCore/Public/ShaderCodeLibrary.h b/Engine/Source/Runtime/RenderCore/Public/ShaderCodeLibrary.h
--- a/Engine/Source/Runtime/RenderCore/Public/ShaderCodeLibrary.h
+++ b/Engine/Source/Runtime/RenderCore/Public/ShaderCodeLibrary.h
@@ -370,8 +370,11 @@ struct FShaderLibraryCooker
 		const FString& SandboxDestinationPath, const FString& SandboxMetadataPath, TArray<FString>& OutChunkFilenames, bool& bOutHasData);
 
 	// Dump collected stats for each shader platform
+	UE_DEPRECATED(5.4, "DumpShaderCodeStats is deprecated due to being incorrect w.r.t shader duplication due to chunking.")
 	static RENDERCORE_API void DumpShaderCodeStats();
 
+	static RENDERCORE_API void DumpShaderTypeStats(const FString& DebugInfoPath, const FString& MetaDataPath);
+
 	// Create a smaller 'patch' library that only contains data from 'NewMetaDataDir' not contained in any of 'OldMetaDataDirs'
 	static RENDERCORE_API bool CreatePatchLibrary(TArray<FString> const& OldMetaDataDirs, FString const& NewMetaDataDir, FString const& OutDir, bool bNativeFormat, bool bNeedsDeterministicOrder);
 };
