GpuProfiler: Support gpu trace events named with the <basename>_N pattern, where N is some integer.

Before this fix, Gpu trace events with said naming pattern would all appear with the same name, which would correspond to the first one. e.g. VP_0, VP_1, VP_2 would all appear as VP_0 in Unreal Insights.

This is because the event is being sent as a uint32 corresponding to the ComparisonIndex of FName, which is the same for the above naming pattern - the N ends up encoded in the Number of the FName, corresponding to the instance number.

The fix in this CL forces a space at the end of the FName when the instance number is greater than 0.

Upgrading the GpuProfiler to include the FName instance number was also possible without breaking backwards compatibility (open old traces with new insights) but adds 4 bytes per event to handle the FName instance number, so went for the simpler solution.

diff --git a/Engine/Source/Runtime/RenderCore/Private/ProfilingDebugging/RealtimeGPUProfiler.cpp b/Engine/Source/Runtime/RenderCore/Private/ProfilingDebugging/RealtimeGPUProfiler.cpp
--- a/Engine/Source/Runtime/RenderCore/Private/ProfilingDebugging/RealtimeGPUProfiler.cpp
+++ b/Engine/Source/Runtime/RenderCore/Private/ProfilingDebugging/RealtimeGPUProfiler.cpp
@@ -338,6 +338,16 @@ void TraverseEventTree(
 			EventName = GpuProfilerEvents[Root].GetName();
 		}
 
+		// Since the GpuProfiler uses the Comparison Index of FName, Gpu trace events named with the pattern <base>_N where N
+		// is some non-negative integer, will all end up having the same name in Unreal Insights. Appending a space to the name
+		// avoids this.
+		if (EventName.GetNumber())
+		{
+			EventName = FName(EventName.ToString() + TEXT(" "));
+
+			checkSlow(EventName.GetNumber() == 0);
+		}
+
 		check(GpuProfilerEvents[Root].GetGPUMask().Contains(GPUIndex));
 		FGpuProfilerTrace::SpecifyEventByName(EventName);
 		FGpuProfilerTrace::BeginEventByName(EventName, GpuProfilerEvents[Root].GetFrameNumber(), GpuProfilerEvents[Root].GetStartResultMicroseconds(GPUIndex));
