Add the possibility to pushevent to the RealTimeGPUProfiler while calling Execute on the RenderGraph and ParallelExecute.bEnabled==false. This allows to preserve gpu scopes stack when calling SCOPED_GPU_STAT inside the execution of the lambda of an RDGPass.

diff --git a/Engine/Source/Runtime/RenderCore/Private/RenderGraphBuilder.cpp b/Engine/Source/Runtime/RenderCore/Private/RenderGraphBuilder.cpp
--- a/Engine/Source/Runtime/RenderCore/Private/RenderGraphBuilder.cpp
+++ b/Engine/Source/Runtime/RenderCore/Private/RenderGraphBuilder.cpp
@@ -580,7 +580,7 @@ bool FRDGBuilder::IsTransientInternal(FRDGViewableResource* Resource, bool bFast
 }
 
 FRDGBuilder::FRDGBuilder(FRHICommandListImmediate& InRHICmdList, FRDGEventName InName, ERDGBuilderFlags InFlags)
-	: FRDGScopeState(InRHICmdList, IsImmediateMode())
+	: FRDGScopeState(InRHICmdList, IsImmediateMode(), ::IsParallelExecuteEnabled() && EnumHasAnyFlags(InFlags, ERDGBuilderFlags::AllowParallelExecute))
 	, RootAllocatorScope(Allocators.Root)
 	, Blackboard(Allocators.Root)
 	, BuilderName(InName)
diff --git a/Engine/Source/Runtime/RenderCore/Private/RenderGraphPrivate.cpp b/Engine/Source/Runtime/RenderCore/Private/RenderGraphPrivate.cpp
--- a/Engine/Source/Runtime/RenderCore/Private/RenderGraphPrivate.cpp
+++ b/Engine/Source/Runtime/RenderCore/Private/RenderGraphPrivate.cpp
@@ -664,8 +664,9 @@ bool IsParallelSetupEnabled()
 		;
 }
 
-FRDGScopeState::FState::FState(bool bImmediate)
-	: bImmediate(bImmediate)
+FRDGScopeState::FState::FState(bool bInImmediate, bool bInParallelExecute)
+	: bImmediate(bInImmediate),
+	  bParallelExecute(bInParallelExecute)
 #if RDG_EVENTS
 	, ScopeMode([]
 	{
diff --git a/Engine/Source/Runtime/RenderCore/Public/RenderGraphEvent.h b/Engine/Source/Runtime/RenderCore/Public/RenderGraphEvent.h
--- a/Engine/Source/Runtime/RenderCore/Public/RenderGraphEvent.h
+++ b/Engine/Source/Runtime/RenderCore/Public/RenderGraphEvent.h
@@ -84,8 +84,13 @@ struct FRDGScope_GPU
 	FRealtimeGPUProfilerQuery StartQuery;
 	FRealtimeGPUProfilerQuery StopQuery;
 
+	FName StatName;
+	TStatId StatId;
+	FString StatDescription;
+
 	TOptional<FRHIDrawStatsCategory const*> PreviousCategory {};
 	FRHIDrawStatsCategory const* CurrentCategory  = nullptr;
+	bool bEmitDuringExecute;
 
 	inline FRDGScope_GPU(FRDGScopeState& State, FRHIGPUMask GPUMask, const FName& CsvStatName, const TStatId& Stat, const TCHAR* Description, FRHIDrawStatsCategory const& Category);
 	inline ~FRDGScope_GPU();
@@ -429,6 +434,7 @@ protected:
 		uint32 Mask = 0;
 
 		bool const bImmediate;
+		bool const bParallelExecute;
 
 #if RDG_EVENTS == RDG_EVENTS_NONE
 		static constexpr ERDGScopeMode const ScopeMode = ERDGScopeMode::Disabled;
@@ -436,7 +442,7 @@ protected:
 		ERDGScopeMode const ScopeMode;
 #endif
 
-		FState(bool bImmediate);
+		FState(bool bInImmediate, bool bInParallelExecute);
 
 	} ScopeState;
 	
@@ -474,8 +480,8 @@ public:
 #endif // WITH_RHI_BREADCRUMBS
 
 public:
-	FRDGScopeState(FRHICommandListImmediate& InRHICmdList, bool bImmediate)
-		: ScopeState(bImmediate)
+	FRDGScopeState(FRHICommandListImmediate& InRHICmdList, bool bImmediate, bool bParallelExecute)
+		: ScopeState(bImmediate, bParallelExecute)
 		, RHICmdList(InRHICmdList)
 #if WITH_RHI_BREADCRUMBS
 		, CurrentBreadcrumbRef(bImmediate ? InRHICmdList.GetCurrentBreadcrumbRef() : LocalCurrentBreadcrumb)
diff --git a/Engine/Source/Runtime/RenderCore/Public/RenderGraphEvent.inl b/Engine/Source/Runtime/RenderCore/Public/RenderGraphEvent.inl
--- a/Engine/Source/Runtime/RenderCore/Public/RenderGraphEvent.inl
+++ b/Engine/Source/Runtime/RenderCore/Public/RenderGraphEvent.inl
@@ -4,11 +4,21 @@
 
 #if HAS_GPU_STATS
 inline FRDGScope_GPU::FRDGScope_GPU(FRDGScopeState& State, FRHIGPUMask GPUMask, const FName& CsvStatName, const TStatId& Stat, const TCHAR* Description, FRHIDrawStatsCategory const& Category)
-	: CurrentCategory(Category.ShouldCountDraws() ? &Category : nullptr)
+	: CurrentCategory(Category.ShouldCountDraws() ? &Category : nullptr),
+	  bEmitDuringExecute(!State.ScopeState.bParallelExecute)
 {
 	if (AreGPUStatsEnabled())
 	{
-		StartQuery = FRealtimeGPUProfiler::Get()->PushEvent(GPUMask, CsvStatName, Stat, Description);
+		if(bEmitDuringExecute)
+		{ 
+			StatName = CsvStatName;
+			StatId = Stat;
+			StatDescription = FString(Description);
+		}
+		else
+		{
+			StartQuery = FRealtimeGPUProfiler::Get()->PushEvent(GPUMask, CsvStatName, Stat, Description);
+		}
 	}
 }
 
@@ -50,17 +60,32 @@ inline void FRDGScope_GPU::EndCPU(FRHIComputeCommandList& RHICmdList, bool bPreS
 
 inline void FRDGScope_GPU::BeginGPU(FRHIComputeCommandList& RHICmdList)
 {
-	if (StartQuery && EnumHasAnyFlags(RHICmdList.GetPipeline(), ERHIPipeline::Graphics))
+	if (EnumHasAnyFlags(RHICmdList.GetPipeline(), ERHIPipeline::Graphics))
 	{
-		StartQuery.Submit(static_cast<FRHICommandList&>(RHICmdList), true);
+		if (bEmitDuringExecute)
+		{
+			FRealtimeGPUProfiler::Get()->PushStat(RHICmdList.GetAsImmediate(), StatName, StatId, *StatDescription);
+		}
+		else if (StartQuery)
+		{
+			StartQuery.Submit(static_cast<FRHICommandList&>(RHICmdList), true);
+		}
 	}
 }
 
 inline void FRDGScope_GPU::EndGPU(FRHIComputeCommandList& RHICmdList)
 {
-	if (StopQuery && EnumHasAnyFlags(RHICmdList.GetPipeline(), ERHIPipeline::Graphics))
+	if (EnumHasAnyFlags(RHICmdList.GetPipeline(), ERHIPipeline::Graphics))
 	{
-		StopQuery.Submit(static_cast<FRHICommandList&>(RHICmdList), false);
+		if (bEmitDuringExecute)
+		{
+			FRealtimeGPUProfiler::Get()->PopStat(static_cast<FRHICommandList&>(RHICmdList).GetAsImmediate());
+		}
+		else if (StopQuery)
+		{
+			StopQuery.Submit(static_cast<FRHICommandList&>(RHICmdList), false);
+		}
+	
 	}
 }
 #endif
