Compile fix

diff --git a/Engine/Source/Runtime/RenderCore/Public/ProfilingDebugging/RealtimeGPUProfiler.h b/Engine/Source/Runtime/RenderCore/Public/ProfilingDebugging/RealtimeGPUProfiler.h
--- a/Engine/Source/Runtime/RenderCore/Public/ProfilingDebugging/RealtimeGPUProfiler.h
+++ b/Engine/Source/Runtime/RenderCore/Public/ProfilingDebugging/RealtimeGPUProfiler.h
@@ -294,8 +294,8 @@ public:
 	RENDERCORE_API FRealtimeGPUProfilerQuery PopEvent();
 
 	/** Push/pop stats which do additional draw call tracking on top of events. */
-	void PushStat(FRHICommandListImmediate& RHICmdList, const FName& Name, const TStatId& Stat, const TCHAR* Description);
-	void PopStat(FRHICommandListImmediate& RHICmdList);
+	RENDERCORE_API void PushStat(FRHICommandListImmediate& RHICmdList, const FName& Name, const TStatId& Stat, const TCHAR* Description);
+	RENDERCORE_API void PopStat(FRHICommandListImmediate& RHICmdList);
 
 #if GPUPROFILERTRACE_ENABLED
 	RENDERCORE_API void FetchPerfByDescription(TArray<FRealtimeGPUProfilerDescriptionResult> & OutResults) const;
diff --git a/Engine/Source/Runtime/RenderCore/Public/RenderGraphEvent.inl b/Engine/Source/Runtime/RenderCore/Public/RenderGraphEvent.inl
--- a/Engine/Source/Runtime/RenderCore/Public/RenderGraphEvent.inl
+++ b/Engine/Source/Runtime/RenderCore/Public/RenderGraphEvent.inl
@@ -79,7 +79,7 @@ inline void FRDGScope_GPU::EndGPU(FRHIComputeCommandList& RHICmdList)
 	{
 		if (bEmitDuringExecute)
 		{
-			FRealtimeGPUProfiler::Get()->PopStat(static_cast<FRHICommandList&>(RHICmdList).GetAsImmediate());
+			FRealtimeGPUProfiler::Get()->PopStat(RHICmdList.GetAsImmediate());
 		}
 		else if (StopQuery)
 		{
