Remove deprecated code from FRayTracingGeometry.

diff --git a/Engine/Source/Runtime/RenderCore/Private/RayTracingGeometry.cpp b/Engine/Source/Runtime/RenderCore/Private/RayTracingGeometry.cpp
--- a/Engine/Source/Runtime/RenderCore/Private/RayTracingGeometry.cpp
+++ b/Engine/Source/Runtime/RenderCore/Private/RayTracingGeometry.cpp
@@ -64,25 +64,6 @@ void FRayTracingGeometry::ReleaseRHIForStreaming(FRHIResourceReplaceBatcher& Bat
 	}
 }
 
-void FRayTracingGeometry::CreateRayTracingGeometryFromCPUData(TResourceArray<uint8>& OfflineData)
-{
-	check(OfflineData.Num() == 0 || Initializer.OfflineData == nullptr);
-	if (OfflineData.Num())
-	{
-		Initializer.OfflineData = &OfflineData;
-	}
-
-	if (CVarDebugForceRuntimeBLAS.GetValueOnAnyThread() && Initializer.OfflineData != nullptr)
-	{
-		Initializer.OfflineData->Discard();
-		Initializer.OfflineData = nullptr;
-	}
-	
-	FRHICommandList& RHICmdList = FRHICommandListImmediate::Get();
-	RayTracingGeometryRHI = RHICmdList.CreateRayTracingGeometry(Initializer);
-	SetRequiresBuild(Initializer.OfflineData == nullptr || RayTracingGeometryRHI->IsCompressed());
-}
-
 void FRayTracingGeometry::RequestBuildIfNeeded(ERTAccelerationStructureBuildPriority InBuildPriority)
 {
 	RayTracingGeometryRHI->SetInitializer(Initializer);
@@ -136,13 +117,6 @@ void FRayTracingGeometry::MakeResident(FRHICommandList& RHICmdList)
 	}
 }
 
-void FRayTracingGeometry::InitRHIForDynamicRayTracing()
-{
-	check(GetRayTracingMode() == ERayTracingMode::Dynamic);
-
-	MakeResident(FRHICommandListImmediate::Get());
-}
-
 void FRayTracingGeometry::Evict()
 {
 	check(!EnumHasAllFlags(GeometryState, EGeometryStateFlags::Evicted) && RayTracingGeometryRHI != nullptr);
@@ -228,11 +202,6 @@ void FRayTracingGeometry::CreateRayTracingGeometry(FRHICommandListBase& RHICmdLi
 	}
 }
 
-void FRayTracingGeometry::CreateRayTracingGeometry(ERTAccelerationStructureBuildPriority InBuildPriority)
-{
-	CreateRayTracingGeometry(FRHICommandListImmediate::Get(), InBuildPriority);
-}
-
 bool FRayTracingGeometry::IsValid() const
 {
 	// can't check IsInitialized() because current implementation of hair ray tracing support doesn't initialize resource
diff --git a/Engine/Source/Runtime/RenderCore/Public/RayTracingGeometry.h b/Engine/Source/Runtime/RenderCore/Public/RayTracingGeometry.h
--- a/Engine/Source/Runtime/RenderCore/Public/RayTracingGeometry.h
+++ b/Engine/Source/Runtime/RenderCore/Public/RayTracingGeometry.h
@@ -113,19 +113,10 @@ public:
 	RENDERCORE_API void InitRHIForStreaming(FRHIRayTracingGeometry* IntermediateGeometry, FRHIResourceReplaceBatcher& Batcher);
 	RENDERCORE_API void ReleaseRHIForStreaming(FRHIResourceReplaceBatcher& Batcher);
 
-	UE_DEPRECATED(5.4, "Use FStaticMeshStreamIn::FIntermediateRayTracingGeometry instead.")
-	RENDERCORE_API void CreateRayTracingGeometryFromCPUData(TResourceArray<uint8>& OfflineData);
-
 	RENDERCORE_API void RequestBuildIfNeeded(ERTAccelerationStructureBuildPriority InBuildPriority);
 
-	UE_DEPRECATED(5.4, "InitRHIForDynamicRayTracing now requires a command list and was renamed to MakeResident().")
-	RENDERCORE_API void InitRHIForDynamicRayTracing();
-
 	RENDERCORE_API void CreateRayTracingGeometry(FRHICommandListBase& RHICmdList, ERTAccelerationStructureBuildPriority InBuildPriority);
 
-	UE_DEPRECATED(5.4, "CreateRayTracingGeometry now requires a command list.")
-	RENDERCORE_API void CreateRayTracingGeometry(ERTAccelerationStructureBuildPriority InBuildPriority);
-
 	RENDERCORE_API void MakeResident(FRHICommandList& RHICmdList);
 	RENDERCORE_API void Evict();
 	
