Parallelize UpdateReferencedUniformBufferNames, which saves a second on startup

diff --git a/Engine/Source/Runtime/RenderCore/Private/Shader.cpp b/Engine/Source/Runtime/RenderCore/Private/Shader.cpp
--- a/Engine/Source/Runtime/RenderCore/Private/Shader.cpp
+++ b/Engine/Source/Runtime/RenderCore/Private/Shader.cpp
@@ -508,14 +508,26 @@ void FShaderType::Initialize(const TMap<FString, TArray<const TCHAR*> >& ShaderF
 #if UE_BUILD_DEBUG
 		TArray<FShaderType*> UniqueShaderTypes;
 #endif
+		TArray<UE::Tasks::TTask<void>> Tasks;
+		Tasks.Reserve(FShaderType::GetNameToTypeMap().Num());
+
 		for(TLinkedList<FShaderType*>::TIterator It(FShaderType::GetTypeList()); It; It.Next())
 		{
 			FShaderType* Type = *It;
 #if UE_BUILD_DEBUG
 			UniqueShaderTypes.Add(Type);
 #endif
-			Type->UpdateReferencedUniformBufferNames(ShaderFileToUniformBufferVariables);
+			Tasks.Emplace(
+				UE::Tasks::Launch(
+					TEXT("UpdateReferencedUniformBufferNames"),
+					[Type, &ShaderFileToUniformBufferVariables]() mutable
+					{
+						Type->UpdateReferencedUniformBufferNames(ShaderFileToUniformBufferVariables);
+					}
+				)
+			);
 		}
+		UE::Tasks::Wait(Tasks);
 	
 #if UE_BUILD_DEBUG
 		// Check for duplicated shader type names
