Fixed r.rdg.clobberresources crash and incorrect max discard transition for transient resources.

diff --git a/Engine/Source/Runtime/RenderCore/Private/RenderGraphBuilder.cpp b/Engine/Source/Runtime/RenderCore/Private/RenderGraphBuilder.cpp
--- a/Engine/Source/Runtime/RenderCore/Private/RenderGraphBuilder.cpp
+++ b/Engine/Source/Runtime/RenderCore/Private/RenderGraphBuilder.cpp
@@ -2854,7 +2854,6 @@ void FRDGBuilder::AllocateTransientResources(TConstArrayView<FCollectResourceOp>
 				TransientResourceAllocator->DeallocateMemory(TransientBuffer, PassHandle.GetIndex());
 
 				Buffer->MinDiscardPass = FRDGPassHandle(TransientBuffer->GetDiscardPasses().Min);
-				Buffer->MaxDiscardPass = FRDGPassHandle(FMath::Min<uint32>(TransientBuffer->GetDiscardPasses().Max, GetEpiloguePassHandle().GetIndex()));
 			}
 			else
 			{
@@ -2879,7 +2878,6 @@ void FRDGBuilder::AllocateTransientResources(TConstArrayView<FCollectResourceOp>
 				if (!TransientTexture->IsAcquired())
 				{
 					Texture->MinDiscardPass = FRDGPassHandle(TransientTexture->GetDiscardPasses().Min);
-					Texture->MaxDiscardPass = FRDGPassHandle(FMath::Min<uint32>(TransientTexture->GetDiscardPasses().Max, GetEpiloguePassHandle().GetIndex()));
 				}
 			}
 		}
@@ -3634,9 +3632,10 @@ void FRDGBuilder::AddLastTextureTransition(FRDGTexture* Texture)
 	// Texture is using the RHI transient allocator. Transition it back to Discard in the final pass it is used.
 	if (Texture->MinDiscardPass.IsValid())
 	{
-		AddAliasingTransition(Texture->MinDiscardPass, Texture->MaxDiscardPass, Texture, FRHITransientAliasingInfo::Discard(Texture->GetRHIUnchecked()));
+		FRDGPassHandle MaxDiscardPass = FRDGPassHandle(FMath::Min<uint32>(Texture->TransientTexture->GetDiscardPasses().Max, GetEpiloguePassHandle().GetIndex()));
+		AddAliasingTransition(Texture->MinDiscardPass, MaxDiscardPass, Texture, FRHITransientAliasingInfo::Discard(Texture->GetRHIUnchecked()));
 
-		SubresourceStateAfter.SetPass(ERHIPipeline::Graphics, Texture->MaxDiscardPass);
+		SubresourceStateAfter.SetPass(ERHIPipeline::Graphics, MaxDiscardPass);
 		SubresourceStateAfter.Access = ERHIAccess::Discard;
 	}
 	else
@@ -3716,9 +3715,10 @@ void FRDGBuilder::AddLastBufferTransition(FRDGBuffer* Buffer)
 	// Texture is using the RHI transient allocator. Transition it back to Discard in the final pass it is used.
 	if (Buffer->MinDiscardPass.IsValid())
 	{
-		AddAliasingTransition(Buffer->MinDiscardPass, Buffer->MaxDiscardPass, Buffer, FRHITransientAliasingInfo::Discard(Buffer->GetRHIUnchecked()));
+		FRDGPassHandle MaxDiscardPass = FRDGPassHandle(FMath::Min<uint32>(Buffer->TransientBuffer->GetDiscardPasses().Max, GetEpiloguePassHandle().GetIndex()));
+		AddAliasingTransition(Buffer->MinDiscardPass, MaxDiscardPass, Buffer, FRHITransientAliasingInfo::Discard(Buffer->GetRHIUnchecked()));
 
-		StateAfter->SetPass(ERHIPipeline::Graphics, Buffer->MaxDiscardPass);
+		StateAfter->SetPass(ERHIPipeline::Graphics, MaxDiscardPass);
 		StateAfter->Access = ERHIAccess::Discard;
 	}
 	else
@@ -4375,7 +4375,7 @@ void FRDGBuilder::ClobberPassOutputs(const FRDGPass* Pass)
 		{
 			FRDGTextureRef Texture = TextureAccess.GetTexture();
 	
-			if (UserValidation.TryMarkForClobber(Texture))
+			if (Texture && UserValidation.TryMarkForClobber(Texture))
 			{
 				if (EnumHasAnyFlags(TextureAccess.GetAccess(), ERHIAccess::UAVMask))
 				{
@@ -4398,7 +4398,7 @@ void FRDGBuilder::ClobberPassOutputs(const FRDGPass* Pass)
 		{
 			FRDGBufferRef Buffer = BufferAccess.GetBuffer();
 	
-			if (UserValidation.TryMarkForClobber(Buffer))
+			if (Buffer && UserValidation.TryMarkForClobber(Buffer))
 			{
 				AddClearUAVPass(*this, CreateUAV(Buffer), GetClobberBufferValue());
 			}
diff --git a/Engine/Source/Runtime/RenderCore/Public/RenderGraphResources.h b/Engine/Source/Runtime/RenderCore/Public/RenderGraphResources.h
--- a/Engine/Source/Runtime/RenderCore/Public/RenderGraphResources.h
+++ b/Engine/Source/Runtime/RenderCore/Public/RenderGraphResources.h
@@ -399,7 +399,6 @@ protected:
 	FRDGPassHandle LastPass;
 	FRDGPassHandle MinAcquirePass;
 	FRDGPassHandle MinDiscardPass;
-	FRDGPassHandle MaxDiscardPass;
 
 	/** Number of references in passes and deferred queries. */
 	uint16 ReferenceCount;
