[shaders] hash all ueshadermetadata directives found when preprocessing shaders for inclusion in the per-shader DDC key, not just VERSION ones. others are typically used to modify compilation behaviour in some manner.

diff --git a/Engine/Source/Runtime/RenderCore/Private/ShaderCore.cpp b/Engine/Source/Runtime/RenderCore/Private/ShaderCore.cpp
--- a/Engine/Source/Runtime/RenderCore/Private/ShaderCore.cpp
+++ b/Engine/Source/Runtime/RenderCore/Private/ShaderCore.cpp
@@ -3588,14 +3588,15 @@ FShaderCommonCompileJob::FInputHash FShaderCompileJob::GetInputHash()
 		}
 		MergedEnvironment.SerializeCompilationDependencies(Hasher);
 		
-		auto HashVersion = [&Hasher](const FString* VersionDirective)
+		auto HashDirectives = [&Hasher](const FString* Directive)
 		{
-			check(VersionDirective && !VersionDirective->IsEmpty());
+			check(Directive && !Directive->IsEmpty());
 			// const_cast due to serialization API requiring non-const. better than not having const correctness in the API.
-			Hasher << const_cast<FString&>(*VersionDirective);
+			Hasher << const_cast<FString&>(*Directive);
 		};
+		// Hash all UESHADERMETADATA_ directives encountered during preprocessing (assume these may be used to modify compilation behaviour)
+		PreprocessOutput.VisitDirectives(HashDirectives);
 
-		PreprocessOutput.VisitDirectivesWithPrefix(TEXT("VERSION"), HashVersion);
 		// const_cast due to serialization API requiring non-const. better than not having const correctness in the API.
 		Hasher << PreprocessOutput.EditSource();
 		if (SecondaryPreprocessOutput.IsValid())
diff --git a/Engine/Source/Runtime/RenderCore/Public/ShaderPreprocessTypes.h b/Engine/Source/Runtime/RenderCore/Public/ShaderPreprocessTypes.h
--- a/Engine/Source/Runtime/RenderCore/Public/ShaderPreprocessTypes.h
+++ b/Engine/Source/Runtime/RenderCore/Public/ShaderPreprocessTypes.h
@@ -163,6 +163,16 @@ public:
 		return false;
 	}
 
+	inline void VisitDirectives(TFunction<void(const FString*)> Action) const
+	{
+		const int32 NumberOfDirectives = PragmaDirectives.Num();
+		for (int32 i = 0; i < NumberOfDirectives; i++)
+		{
+			const FString& CurrentDirective = PragmaDirectives[i];
+			Action(&CurrentDirective);
+		}
+	}
+
 	inline void VisitDirectivesWithPrefix(const TCHAR* Prefix, TFunction<void(const FString*)> Action) const
 	{
 		const int32 NumberOfDirectives = PragmaDirectives.Num();
