Workaround deprecation warnings caused by default constructor/operators of FRayTracingGeometry.

diff --git a/Engine/Source/Runtime/RenderCore/Private/RayTracingGeometry.cpp b/Engine/Source/Runtime/RenderCore/Private/RayTracingGeometry.cpp
--- a/Engine/Source/Runtime/RenderCore/Private/RayTracingGeometry.cpp
+++ b/Engine/Source/Runtime/RenderCore/Private/RayTracingGeometry.cpp
@@ -22,11 +22,6 @@ static TAutoConsoleVariable<int32> CVarDebugForceRuntimeBLAS(
 
 PRAGMA_DISABLE_DEPRECATION_WARNINGS // Remove this in UE 5.6 once RayTracingGeometryRHI is made private
 
-FRayTracingGeometry::FRayTracingGeometry() = default;
-
-/** Destructor. */
-FRayTracingGeometry::~FRayTracingGeometry() = default;
-
 #if RHI_RAYTRACING
 
 void FRayTracingGeometry::InitRHIForStreaming(FRHIRayTracingGeometry* IntermediateGeometry, FRHIResourceReplaceBatcher& Batcher)
diff --git a/Engine/Source/Runtime/RenderCore/Public/RayTracingGeometry.h b/Engine/Source/Runtime/RenderCore/Public/RayTracingGeometry.h
--- a/Engine/Source/Runtime/RenderCore/Public/RayTracingGeometry.h
+++ b/Engine/Source/Runtime/RenderCore/Public/RayTracingGeometry.h
@@ -31,8 +31,15 @@ class FRayTracingGeometry : public FRenderResource
 public:
 	TResourceArray<uint8> RawData;
 
-	RENDERCORE_API FRayTracingGeometry();
-	RENDERCORE_API virtual ~FRayTracingGeometry();
+	//disable deprecation warnings for default constructors
+	PRAGMA_DISABLE_DEPRECATION_WARNINGS
+	FRayTracingGeometry() = default;
+	FRayTracingGeometry(const FRayTracingGeometry&) = default;
+	FRayTracingGeometry& operator=(const FRayTracingGeometry&) = default;
+	FRayTracingGeometry(FRayTracingGeometry&&) = default;
+	FRayTracingGeometry& operator=(FRayTracingGeometry&&) = default;
+	~FRayTracingGeometry() = default;
+	PRAGMA_ENABLE_DEPRECATION_WARNINGS
 
 #if RHI_RAYTRACING
 
