Make RenderUtils read from TargetPlatformSettings.
Fixes the crashes/wrong rendering preview had when the sdk was not installed for the previewed platform

diff --git a/Engine/Source/Runtime/RenderCore/Private/RenderUtils.cpp b/Engine/Source/Runtime/RenderCore/Private/RenderUtils.cpp
--- a/Engine/Source/Runtime/RenderCore/Private/RenderUtils.cpp
+++ b/Engine/Source/Runtime/RenderCore/Private/RenderUtils.cpp
@@ -768,10 +768,10 @@ RENDERCORE_API void RenderUtilsInit()
 	ITargetPlatformManagerModule* TargetPlatformManager = GetTargetPlatformManager();
 	if (TargetPlatformManager)
 	{
-		for (ITargetPlatform* TargetPlatform : TargetPlatformManager->GetTargetPlatforms())
+		for (ITargetPlatformSettings* TargetPlatformSettings : TargetPlatformManager->GetTargetPlatformSettings())
 		{
 			TArray<FName> PlatformPossibleShaderFormats;
-			TargetPlatform->GetAllPossibleShaderFormats(PlatformPossibleShaderFormats);
+			TargetPlatformSettings->GetAllPossibleShaderFormats(PlatformPossibleShaderFormats);
 
 			for (FName Format : PlatformPossibleShaderFormats)
 			{
@@ -791,27 +791,27 @@ RENDERCORE_API void RenderUtilsInit()
 						continue;
 					}
 
-					GForwardShadingPlatformMask[ShaderPlatformIndex] = TargetPlatform->UsesForwardShading();
+					GForwardShadingPlatformMask[ShaderPlatformIndex] = TargetPlatformSettings->UsesForwardShading();
 
-					GDBufferPlatformMask[ShaderPlatformIndex] = IsMobilePlatform(ShaderPlatformToEdit) ? (TargetPlatform->UsesMobileDBuffer() && !IsMobileDeferredShadingEnabled(ShaderPlatformToEdit)) : TargetPlatform->UsesDBuffer();
+					GDBufferPlatformMask[ShaderPlatformIndex] = IsMobilePlatform(ShaderPlatformToEdit) ? (TargetPlatformSettings->UsesMobileDBuffer() && !IsMobileDeferredShadingEnabled(ShaderPlatformToEdit)) : TargetPlatformSettings->UsesDBuffer();
 
-					GSelectiveBasePassOutputsPlatformMask[ShaderPlatformIndex] = TargetPlatform->UsesSelectiveBasePassOutputs();
+					GSelectiveBasePassOutputsPlatformMask[ShaderPlatformIndex] = TargetPlatformSettings->UsesSelectiveBasePassOutputs();
 
-					GDistanceFieldsPlatformMask[ShaderPlatformIndex] = TargetPlatform->UsesDistanceFields();
+					GDistanceFieldsPlatformMask[ShaderPlatformIndex] = TargetPlatformSettings->UsesDistanceFields();
 
-					GSimpleSkyDiffusePlatformMask[ShaderPlatformIndex] = TargetPlatform->ForcesSimpleSkyDiffuse();
+					GSimpleSkyDiffusePlatformMask[ShaderPlatformIndex] = TargetPlatformSettings->ForcesSimpleSkyDiffuse();
 
-					GVelocityEncodeDepthPlatformMask[ShaderPlatformIndex] = TargetPlatform->VelocityEncodeDepth();
+					GVelocityEncodeDepthPlatformMask[ShaderPlatformIndex] = TargetPlatformSettings->VelocityEncodeDepth();
 
-					GMobileAmbientOcclusionPlatformMask[ShaderPlatformIndex] = TargetPlatform->UsesMobileAmbientOcclusion();
+					GMobileAmbientOcclusionPlatformMask[ShaderPlatformIndex] = TargetPlatformSettings->UsesMobileAmbientOcclusion();
 				}
 			}
 
 
-			if (TargetPlatform->UsesRayTracing())
+			if (TargetPlatformSettings->UsesRayTracing())
 			{
 				TArray<FName> PlatformRayTracingShaderFormats;
-				TargetPlatform->GetRayTracingShaderFormats(PlatformRayTracingShaderFormats);
+				TargetPlatformSettings->GetRayTracingShaderFormats(PlatformRayTracingShaderFormats);
 
 				for (FName FormatName : PlatformRayTracingShaderFormats)
 				{
