Fix RDG clobber mode trying to fill block compressed textures with RHIClearUAVFloat.

These UAVs need to be handled as UINT, not float. Also fixed some undefined float to int type casting in the function which was computing the integer clobber value.

#rnx

diff --git a/Engine/Source/Runtime/RenderCore/Private/RenderGraphBuilder.cpp b/Engine/Source/Runtime/RenderCore/Private/RenderGraphBuilder.cpp
--- a/Engine/Source/Runtime/RenderCore/Private/RenderGraphBuilder.cpp
+++ b/Engine/Source/Runtime/RenderCore/Private/RenderGraphBuilder.cpp
@@ -4363,6 +4363,14 @@ void FRDGBuilder::ClobberPassOutputs(const FRDGPass* Pass)
 		{
 			AddClearUAVPass(*this, TextureUAV, GetClobberBufferValue());
 		}
+		else if (IsBlockCompressedFormat(TextureUAV->GetParent()->Desc.Format))
+		{
+			// We shouldn't see BCn UAVs if SupportsUAVFormatAliasing is false in the first place, but it can't hurt to check.
+			if (GRHIGlobals.SupportsUAVFormatAliasing)
+			{
+				AddClearUAVPass(*this, TextureUAV, GetClobberBufferValue());
+			}
+		}
 		else
 		{
 			AddClearUAVPass(*this, TextureUAV, ClobberColor);
diff --git a/Engine/Source/Runtime/RenderCore/Private/RenderGraphPrivate.cpp b/Engine/Source/Runtime/RenderCore/Private/RenderGraphPrivate.cpp
--- a/Engine/Source/Runtime/RenderCore/Private/RenderGraphPrivate.cpp
+++ b/Engine/Source/Runtime/RenderCore/Private/RenderGraphPrivate.cpp
@@ -210,7 +210,8 @@ FLinearColor GetClobberColor()
 uint32 GetClobberBufferValue()
 {
 	float ClobberValue = GetClobberValue();
-	uint32 ClobberValueUint = reinterpret_cast<const uint32*>(&ClobberValue)[0];
+	uint32 ClobberValueUint;
+	FMemory::Memcpy(&ClobberValueUint, &ClobberValue, sizeof(ClobberValueUint));
 	return ClobberValueUint;
 }
 
