Add missing RENDERCORE_API to FRayTracingGeometry constructor/destructor
- also make FRayTracingGeometry copy constructors/operators deleted to prevent copies.

diff --git a/Engine/Source/Runtime/RenderCore/Public/RayTracingGeometry.h b/Engine/Source/Runtime/RenderCore/Public/RayTracingGeometry.h
--- a/Engine/Source/Runtime/RenderCore/Public/RayTracingGeometry.h
+++ b/Engine/Source/Runtime/RenderCore/Public/RayTracingGeometry.h
@@ -33,12 +33,13 @@ public:
 
 	//disable deprecation warnings for default constructors
 	PRAGMA_DISABLE_DEPRECATION_WARNINGS
-	FRayTracingGeometry() = default;
-	FRayTracingGeometry(const FRayTracingGeometry&) = default;
-	FRayTracingGeometry& operator=(const FRayTracingGeometry&) = default;
-	FRayTracingGeometry(FRayTracingGeometry&&) = default;
-	FRayTracingGeometry& operator=(FRayTracingGeometry&&) = default;
-	~FRayTracingGeometry() = default;
+	RENDERCORE_API FRayTracingGeometry() = default;
+	RENDERCORE_API virtual ~FRayTracingGeometry() = default;
+
+	FRayTracingGeometry(const FRayTracingGeometry&) = delete;
+	FRayTracingGeometry& operator=(const FRayTracingGeometry&) = delete;
+	FRayTracingGeometry(FRayTracingGeometry&&) = delete;
+	FRayTracingGeometry& operator=(FRayTracingGeometry&&) = delete;
 	PRAGMA_ENABLE_DEPRECATION_WARNINGS
 
 #if RHI_RAYTRACING
