Make RDG use 32 bit handles.
 - Back compatible with legacy 16 bit insights traces.
 - Fixed SkipTracking validation to not check states to allow for writable access.

diff --git a/Engine/Plugins/RenderGraphInsights/Source/RenderGraphInsights/Private/RenderGraphProvider.cpp b/Engine/Plugins/RenderGraphInsights/Source/RenderGraphInsights/Private/RenderGraphProvider.cpp
--- a/Engine/Plugins/RenderGraphInsights/Source/RenderGraphInsights/Private/RenderGraphProvider.cpp
+++ b/Engine/Plugins/RenderGraphInsights/Source/RenderGraphInsights/Private/RenderGraphProvider.cpp
@@ -18,6 +18,46 @@ INSIGHTS_IMPLEMENT_RTTI(FPassIntervalPacket)
 
 FName FRenderGraphProvider::ProviderName("RenderGraphProvider");
 
+template <typename ObjectType>
+TRDGHandle<ObjectType, uint32> ConvertTo32Bits(TRDGHandle<ObjectType, uint16> Handle)
+{
+	if (Handle.IsValid())
+	{
+		return TRDGHandle<ObjectType, uint32>(Handle.GetIndex());
+	}
+	return {};
+}
+
+template <typename HandleType>
+HandleType SerializeHandle(const UE::Trace::IAnalyzer::FEventData& EventData, const ANSICHAR* FieldName)
+{
+	if (EventData.GetValue<bool>("IsHandleType32Bits"))
+	{
+		return EventData.GetValue<HandleType>(FieldName);
+	}
+	return ConvertTo32Bits(EventData.GetValue<TRDGHandle<typename HandleType::ObjectType, uint16>>(FieldName));
+}
+
+template <typename HandleType>
+TArray<HandleType> SerializeHandleArray(const UE::Trace::IAnalyzer::FEventData& EventData, const ANSICHAR* FieldName)
+{
+	if (EventData.GetValue<bool>("IsHandleType32Bits"))
+	{
+		static_assert(sizeof(typename HandleType::IndexType) == 4);
+		return TArray<HandleType>(EventData.GetArrayView<HandleType>(FieldName));
+	}
+
+	using HandleType16 = TRDGHandle<typename HandleType::ObjectType, uint16>;
+	TConstArrayView<HandleType16> Handles = EventData.GetArrayView<HandleType16>(FieldName);
+	TArray<HandleType> Result;
+	Result.Reserve(Handles.Num());
+	for (HandleType16 Handle : Handles)
+	{
+		Result.Emplace(ConvertTo32Bits(Handle));
+	}
+	return Result;
+}
+
 FString GetSizeName(uint32 Bytes)
 {
 	const uint32 KB = 1024;
@@ -48,8 +88,8 @@ FPacket::FPacket(const UE::Trace::IAnalyzer::FOnEventContext& Context)
 
 FPassIntervalPacket::FPassIntervalPacket(const UE::Trace::IAnalyzer::FOnEventContext& Context)
 	: FPacket(Context)
-	, FirstPass(Context.EventData.GetValue<FRDGPassHandle>("FirstPass"))
-	, LastPass(Context.EventData.GetValue<FRDGPassHandle>("LastPass"))
+	, FirstPass(SerializeHandle<FRDGPassHandle>(Context.EventData, "FirstPass"))
+	, LastPass(SerializeHandle<FRDGPassHandle>(Context.EventData, "LastPass"))
 {}
 
 FScopePacket::FScopePacket(const UE::Trace::IAnalyzer::FOnEventContext& Context)
@@ -61,7 +101,7 @@ FResourcePacket::FResourcePacket(const UE::Trace::IAnalyzer::FOnEventContext& Co
 	: FPassIntervalPacket(Context)
 	, Order(Context.EventData.GetValue<uint16>("Order"))
 	, SizeInBytes(Context.EventData.GetValue<uint64>("SizeInBytes"))
-	, Passes(Context.EventData.GetArrayView<FRDGPassHandle>("Passes"))
+	, Passes(SerializeHandleArray<FRDGPassHandle>(Context.EventData, "Passes"))
 	, bExternal(Context.EventData.GetValue<bool>("IsExternal"))
 	, bExtracted(Context.EventData.GetValue<bool>("IsExtracted"))
 	, bCulled(Context.EventData.GetValue<bool>("IsCulled"))
@@ -93,8 +133,8 @@ FResourcePacket::FResourcePacket(const UE::Trace::IAnalyzer::FOnEventContext& Co
 
 FTexturePacket::FTexturePacket(const UE::Trace::IAnalyzer::FOnEventContext& Context)
 	: FResourcePacket(Context)
-	, Handle(Context.EventData.GetValue<FRDGTextureHandle>("Handle"))
-	, NextOwnerHandle(Context.EventData.GetValue<FRDGTextureHandle>("NextOwnerHandle"))
+	, Handle(SerializeHandle<FRDGTextureHandle>(Context.EventData, "Handle"))
+	, NextOwnerHandle(SerializeHandle<FRDGTextureHandle>(Context.EventData, "NextOwnerHandle"))
 {
 	Desc.Flags = ETextureCreateFlags(Context.EventData.GetValue<uint64>("CreateFlags"));
 	Desc.Dimension = ETextureDimension(Context.EventData.GetValue<uint16>("Dimension"));
@@ -111,8 +151,8 @@ FTexturePacket::FTexturePacket(const UE::Trace::IAnalyzer::FOnEventContext& Cont
 
 FBufferPacket::FBufferPacket(const UE::Trace::IAnalyzer::FOnEventContext& Context)
 	: FResourcePacket(Context)
-	, Handle(Context.EventData.GetValue<FRDGBufferHandle>("Handle"))
-	, NextOwnerHandle(Context.EventData.GetValue<FRDGBufferHandle>("NextOwnerHandle"))
+	, Handle(SerializeHandle<FRDGBufferHandle>(Context.EventData, "Handle"))
+	, NextOwnerHandle(SerializeHandle<FRDGBufferHandle>(Context.EventData, "NextOwnerHandle"))
 {
 	Desc.Usage = EBufferUsageFlags(Context.EventData.GetValue<uint32>("UsageFlags"));
 	Desc.BytesPerElement = Context.EventData.GetValue<uint32>("BytesPerElement");
@@ -124,11 +164,11 @@ FBufferPacket::FBufferPacket(const UE::Trace::IAnalyzer::FOnEventContext& Contex
 
 FPassPacket::FPassPacket(const UE::Trace::IAnalyzer::FOnEventContext& Context)
 	: FPacket(Context)
-	, Textures(Context.EventData.GetArrayView<FRDGTextureHandle>("Textures"))
-	, Buffers(Context.EventData.GetArrayView<FRDGBufferHandle>("Buffers"))
-	, Handle(Context.EventData.GetValue<FRDGPassHandle>("Handle"))
-	, GraphicsForkPass(Context.EventData.GetValue<FRDGPassHandle>("GraphicsForkPass"))
-	, GraphicsJoinPass(Context.EventData.GetValue<FRDGPassHandle>("GraphicsJoinPass"))
+	, Textures(SerializeHandleArray<FRDGTextureHandle>(Context.EventData, "Textures"))
+	, Buffers(SerializeHandleArray<FRDGBufferHandle>(Context.EventData, "Buffers"))
+	, Handle(SerializeHandle<FRDGPassHandle>(Context.EventData, "Handle"))
+	, GraphicsForkPass(SerializeHandle<FRDGPassHandle>(Context.EventData, "GraphicsForkPass"))
+	, GraphicsJoinPass(SerializeHandle<FRDGPassHandle>(Context.EventData, "GraphicsJoinPass"))
 	, Flags(ERDGPassFlags(Context.EventData.GetValue<uint16>("Flags")))
 	, Pipeline(ERHIPipeline(Context.EventData.GetValue<uint8>("Pipeline")))
 	, bCulled(Context.EventData.GetValue<bool>("IsCulled"))
diff --git a/Engine/Source/Runtime/RenderCore/Private/RenderGraphBuilder.cpp b/Engine/Source/Runtime/RenderCore/Private/RenderGraphBuilder.cpp
--- a/Engine/Source/Runtime/RenderCore/Private/RenderGraphBuilder.cpp
+++ b/Engine/Source/Runtime/RenderCore/Private/RenderGraphBuilder.cpp
@@ -2248,7 +2248,7 @@ void FRDGBuilder::SetupPassResources(FRDGPass* Pass)
 		if (Texture->LastPass != PassHandle)
 		{
 			Texture->LastPass = PassHandle;
-			Texture->PassStateIndex = static_cast<uint16>(Pass->TextureStates.Num());
+			Texture->PassStateIndex = Pass->TextureStates.Num();
 
 			PassState = &Pass->TextureStates.Emplace_GetRef(Texture);
 		}
diff --git a/Engine/Source/Runtime/RenderCore/Private/RenderGraphResources.cpp b/Engine/Source/Runtime/RenderCore/Private/RenderGraphResources.cpp
--- a/Engine/Source/Runtime/RenderCore/Private/RenderGraphResources.cpp
+++ b/Engine/Source/Runtime/RenderCore/Private/RenderGraphResources.cpp
@@ -51,7 +51,7 @@ FRDGViewableResource::FRDGViewableResource(const TCHAR* InName, const ERDGViewab
 {
 	if (bSkipTracking)
 	{
-		SetExternalAccessMode(ERHIAccess::ReadOnlyExclusiveMask, ERHIPipeline::All);
+		SetExternalAccessMode(ERHIAccess::Mask, ERHIPipeline::All);
 		AccessModeState.bLocked = 1;
 		AccessModeState.ActiveMode = AccessModeState.Mode;
 	}
diff --git a/Engine/Source/Runtime/RenderCore/Private/RenderGraphTrace.cpp b/Engine/Source/Runtime/RenderCore/Private/RenderGraphTrace.cpp
--- a/Engine/Source/Runtime/RenderCore/Private/RenderGraphTrace.cpp
+++ b/Engine/Source/Runtime/RenderCore/Private/RenderGraphTrace.cpp
@@ -13,7 +13,7 @@ UE_TRACE_EVENT_BEGIN(RDGTrace, GraphMessage)
 	UE_TRACE_EVENT_FIELD(UE::Trace::WideString, Name)
 	UE_TRACE_EVENT_FIELD(uint64, StartCycles)
 	UE_TRACE_EVENT_FIELD(uint64, EndCycles)
-	UE_TRACE_EVENT_FIELD(uint16, PassCount)
+	UE_TRACE_EVENT_FIELD(uint32, PassCount)
 	UE_TRACE_EVENT_FIELD(uint64[], TransientMemoryCommitSizes)
 	UE_TRACE_EVENT_FIELD(uint64[], TransientMemoryCapacities)
 	UE_TRACE_EVENT_FIELD(uint8[],  TransientMemoryFlags)
@@ -26,11 +26,11 @@ UE_TRACE_EVENT_BEGIN(RDGTrace, PassMessage)
 	UE_TRACE_EVENT_FIELD(UE::Trace::WideString, Name)
 	UE_TRACE_EVENT_FIELD(uint64, StartCycles)
 	UE_TRACE_EVENT_FIELD(uint64, EndCycles)
-	UE_TRACE_EVENT_FIELD(uint16, Handle)
-	UE_TRACE_EVENT_FIELD(uint16, GraphicsForkPass)
-	UE_TRACE_EVENT_FIELD(uint16, GraphicsJoinPass)
-	UE_TRACE_EVENT_FIELD(uint16[], Textures)
-	UE_TRACE_EVENT_FIELD(uint16[], Buffers)
+	UE_TRACE_EVENT_FIELD(uint32, Handle)
+	UE_TRACE_EVENT_FIELD(uint32, GraphicsForkPass)
+	UE_TRACE_EVENT_FIELD(uint32, GraphicsJoinPass)
+	UE_TRACE_EVENT_FIELD(uint32[], Textures)
+	UE_TRACE_EVENT_FIELD(uint32[], Buffers)
 	UE_TRACE_EVENT_FIELD(uint16, Flags)
 	UE_TRACE_EVENT_FIELD(uint16, Pipeline)
 	UE_TRACE_EVENT_FIELD(bool, IsCulled)
@@ -42,6 +42,7 @@ UE_TRACE_EVENT_BEGIN(RDGTrace, PassMessage)
 	UE_TRACE_EVENT_FIELD(bool, IsParallelExecuteEnd)
 	UE_TRACE_EVENT_FIELD(bool, IsParallelExecute)
 	UE_TRACE_EVENT_FIELD(bool, IsParallelExecuteAllowed)
+	UE_TRACE_EVENT_FIELD(bool, IsHandleType32Bits)
 UE_TRACE_EVENT_END()
 
 UE_TRACE_EVENT_BEGIN(RDGTrace, BufferMessage)
@@ -49,10 +50,10 @@ UE_TRACE_EVENT_BEGIN(RDGTrace, BufferMessage)
 	UE_TRACE_EVENT_FIELD(uint32, UsageFlags)
 	UE_TRACE_EVENT_FIELD(uint32, BytesPerElement)
 	UE_TRACE_EVENT_FIELD(uint32, NumElements)
-	UE_TRACE_EVENT_FIELD(uint16, Handle)
-	UE_TRACE_EVENT_FIELD(uint16, NextOwnerHandle)
-	UE_TRACE_EVENT_FIELD(uint16, Order)
-	UE_TRACE_EVENT_FIELD(uint16[], Passes)
+	UE_TRACE_EVENT_FIELD(uint32, Handle)
+	UE_TRACE_EVENT_FIELD(uint32, NextOwnerHandle)
+	UE_TRACE_EVENT_FIELD(uint32, Order)
+	UE_TRACE_EVENT_FIELD(uint32[], Passes)
 	UE_TRACE_EVENT_FIELD(uint64[], TransientAllocationOffsetMins)
 	UE_TRACE_EVENT_FIELD(uint64[], TransientAllocationOffsetMaxs)
 	UE_TRACE_EVENT_FIELD(uint16[], TransientAllocationMemoryRanges)
@@ -63,16 +64,17 @@ UE_TRACE_EVENT_BEGIN(RDGTrace, BufferMessage)
 	UE_TRACE_EVENT_FIELD(bool, IsTransient)
 	UE_TRACE_EVENT_FIELD(bool, IsTransientUntracked)
 	UE_TRACE_EVENT_FIELD(bool, IsTransientCacheHit)
+	UE_TRACE_EVENT_FIELD(bool, IsHandleType32Bits)
 UE_TRACE_EVENT_END()
  
 UE_TRACE_EVENT_BEGIN(RDGTrace, TextureMessage)
 	UE_TRACE_EVENT_FIELD(UE::Trace::WideString, Name)
 	UE_TRACE_EVENT_FIELD(uint64, StartCycles)
 	UE_TRACE_EVENT_FIELD(uint64, EndCycles)
-	UE_TRACE_EVENT_FIELD(uint16, Handle)
-	UE_TRACE_EVENT_FIELD(uint16, NextOwnerHandle)
-	UE_TRACE_EVENT_FIELD(uint16, Order)
-	UE_TRACE_EVENT_FIELD(uint16[], Passes)
+	UE_TRACE_EVENT_FIELD(uint32, Handle)
+	UE_TRACE_EVENT_FIELD(uint32, NextOwnerHandle)
+	UE_TRACE_EVENT_FIELD(uint32, Order)
+	UE_TRACE_EVENT_FIELD(uint32[], Passes)
 	UE_TRACE_EVENT_FIELD(uint64[], TransientAllocationOffsetMins)
 	UE_TRACE_EVENT_FIELD(uint64[], TransientAllocationOffsetMaxs)
 	UE_TRACE_EVENT_FIELD(uint16[], TransientAllocationMemoryRanges)
@@ -93,18 +95,20 @@ UE_TRACE_EVENT_BEGIN(RDGTrace, TextureMessage)
 	UE_TRACE_EVENT_FIELD(bool, IsTransient)
 	UE_TRACE_EVENT_FIELD(bool, IsTransientUntracked)
 	UE_TRACE_EVENT_FIELD(bool, IsTransientCacheHit)
+	UE_TRACE_EVENT_FIELD(bool, IsHandleType32Bits)
 UE_TRACE_EVENT_END()
 
 UE_TRACE_EVENT_BEGIN(RDGTrace, ScopeMessage)
 	UE_TRACE_EVENT_FIELD(UE::Trace::WideString, Name)
-	UE_TRACE_EVENT_FIELD(uint16, FirstPass)
-	UE_TRACE_EVENT_FIELD(uint16, LastPass)
+	UE_TRACE_EVENT_FIELD(uint32, FirstPass)
+	UE_TRACE_EVENT_FIELD(uint32, LastPass)
 	UE_TRACE_EVENT_FIELD(uint16, Depth)
+	UE_TRACE_EVENT_FIELD(bool, IsHandleType32Bits)
 UE_TRACE_EVENT_END()
 
-static_assert(sizeof(FRDGPassHandle) == sizeof(uint16), "Expected 16 bit pass handles.");
-static_assert(sizeof(FRDGTextureHandle) == sizeof(uint16), "Expected 16 bit texture handles.");
-static_assert(sizeof(FRDGBufferHandle) == sizeof(uint16), "Expected 16 bit buffer handles.");
+static_assert(sizeof(FRDGPassHandle) == sizeof(uint32), "Expected 32 bit pass handles.");
+static_assert(sizeof(FRDGTextureHandle) == sizeof(uint32), "Expected 32 bit texture handles.");
+static_assert(sizeof(FRDGBufferHandle) == sizeof(uint32), "Expected 32 bit buffer handles.");
 
 FRDGTrace::FRDGTrace()
 	: bEnabled(UE_TRACE_CHANNELEXPR_IS_ENABLED(RDGChannel) && !IsImmediateMode())
@@ -162,7 +166,7 @@ void FRDGTrace::OutputGraphEnd(const FRDGBuilder& GraphBuilder)
 			<< GraphMessage.Name(Name, uint16(FCString::Strlen(Name)))
 			<< GraphMessage.StartCycles(GraphStartCycles)
 			<< GraphMessage.EndCycles(FPlatformTime::Cycles64())
-			<< GraphMessage.PassCount(uint16(Passes.Num()))
+			<< GraphMessage.PassCount(uint32(Passes.Num()))
 			<< GraphMessage.TransientMemoryCommitSizes(TransientMemoryCommitSizes.GetData(), (uint16)TransientMemoryCommitSizes.Num())
 			<< GraphMessage.TransientMemoryCapacities(TransientMemoryCapacities.GetData(), (uint16)TransientMemoryCapacities.Num())
 			<< GraphMessage.TransientMemoryFlags(TransientMemoryFlags.GetData(), (uint16)TransientMemoryFlags.Num());
@@ -178,8 +182,8 @@ void FRDGTrace::OutputGraphEnd(const FRDGBuilder& GraphBuilder)
 			<< PassMessage.Handle(Handle.GetIndex())
 			<< PassMessage.GraphicsForkPass(Pass->GetGraphicsForkPass().GetIndexUnchecked())
 			<< PassMessage.GraphicsJoinPass(Pass->GetGraphicsJoinPass().GetIndexUnchecked())
-			<< PassMessage.Textures((const uint16*)Pass->TraceTextures.GetData(), (uint16)Pass->TraceTextures.Num())
-			<< PassMessage.Buffers((const uint16*)Pass->TraceBuffers.GetData(), (uint16)Pass->TraceBuffers.Num())
+			<< PassMessage.Textures((const uint32*)Pass->TraceTextures.GetData(), (uint32)Pass->TraceTextures.Num())
+			<< PassMessage.Buffers((const uint32*)Pass->TraceBuffers.GetData(), (uint32)Pass->TraceBuffers.Num())
 			<< PassMessage.Flags(uint16(Pass->GetFlags()))
 			<< PassMessage.Pipeline(uint16(Pass->GetPipeline()))
 			<< PassMessage.IsCulled(Pass->bCulled != 0)
@@ -190,7 +194,8 @@ void FRDGTrace::OutputGraphEnd(const FRDGBuilder& GraphBuilder)
 			<< PassMessage.IsParallelExecuteBegin(Pass->bParallelExecuteBegin != 0)
 			<< PassMessage.IsParallelExecuteEnd(Pass->bParallelExecuteEnd != 0)
 			<< PassMessage.IsParallelExecute(Pass->bParallelExecute != 0)
-			<< PassMessage.IsParallelExecuteAllowed(Pass->bParallelExecuteAllowed != 0);
+			<< PassMessage.IsParallelExecuteAllowed(Pass->bParallelExecuteAllowed != 0)
+			<< PassMessage.IsHandleType32Bits(true);
 	}
 
 #if RDG_EVENTS
@@ -226,7 +231,8 @@ void FRDGTrace::OutputGraphEnd(const FRDGBuilder& GraphBuilder)
 						<< ScopeMessage.Name(Name, uint16(FCString::Strlen(Name)))
 						<< ScopeMessage.FirstPass(Current->CPUFirstPass->GetHandle().GetIndexUnchecked())
 						<< ScopeMessage.LastPass(Current->CPULastPass->GetHandle().GetIndexUnchecked())
-						<< ScopeMessage.Depth(Depth);
+						<< ScopeMessage.Depth(Depth)
+						<< ScopeMessage.IsHandleType32Bits(true);
 				}
 			}
 		};
@@ -294,7 +300,7 @@ void FRDGTrace::OutputGraphEnd(const FRDGBuilder& GraphBuilder)
 			<< TextureMessage.Handle(Handle.GetIndex())
 			<< TextureMessage.NextOwnerHandle(Texture->NextOwner.GetIndexUnchecked())
 			<< TextureMessage.Order(Texture->TraceOrder)
-			<< TextureMessage.Passes((const uint16*)Texture->TracePasses.GetData(), (uint16)Texture->TracePasses.Num())
+			<< TextureMessage.Passes((const uint32*)Texture->TracePasses.GetData(), (uint32)Texture->TracePasses.Num())
 			<< TextureMessage.TransientAllocationOffsetMins(TransientAllocation.OffsetMins.GetData(), TransientAllocation.OffsetMins.Num())
 			<< TextureMessage.TransientAllocationOffsetMaxs(TransientAllocation.OffsetMaxs.GetData(), TransientAllocation.OffsetMaxs.Num())
 			<< TextureMessage.TransientAllocationMemoryRanges(TransientAllocation.MemoryRanges.GetData(), TransientAllocation.MemoryRanges.Num())
@@ -314,7 +320,8 @@ void FRDGTrace::OutputGraphEnd(const FRDGBuilder& GraphBuilder)
 			<< TextureMessage.IsTrackingSkipped(EnumHasAnyFlags(Texture->Flags, ERDGTextureFlags::SkipTracking))
 			<< TextureMessage.IsTransient(bool(Texture->bTransient))
 			<< TextureMessage.IsTransientUntracked(false)
-			<< TextureMessage.IsTransientCacheHit(TransientAllocation.bCacheHit);
+			<< TextureMessage.IsTransientCacheHit(TransientAllocation.bCacheHit)
+			<< TextureMessage.IsHandleType32Bits(true);
 	}
 
 	for (FRDGBufferHandle Handle = Buffers.Begin(); Handle != Buffers.End(); ++Handle)
@@ -335,7 +342,7 @@ void FRDGTrace::OutputGraphEnd(const FRDGBuilder& GraphBuilder)
 			<< BufferMessage.Handle(Buffer->Handle.GetIndex())
 			<< BufferMessage.NextOwnerHandle(Buffer->NextOwner.GetIndexUnchecked())
 			<< BufferMessage.Order(Buffer->TraceOrder)
-			<< BufferMessage.Passes((const uint16*)Buffer->TracePasses.GetData(), (uint16)Buffer->TracePasses.Num())
+			<< BufferMessage.Passes((const uint32*)Buffer->TracePasses.GetData(), (uint32)Buffer->TracePasses.Num())
 			<< BufferMessage.TransientAllocationOffsetMins(TransientAllocation.OffsetMins.GetData(), TransientAllocation.OffsetMins.Num())
 			<< BufferMessage.TransientAllocationOffsetMaxs(TransientAllocation.OffsetMaxs.GetData(), TransientAllocation.OffsetMaxs.Num())
 			<< BufferMessage.TransientAllocationMemoryRanges(TransientAllocation.MemoryRanges.GetData(), TransientAllocation.MemoryRanges.Num())
@@ -348,7 +355,8 @@ void FRDGTrace::OutputGraphEnd(const FRDGBuilder& GraphBuilder)
 			<< BufferMessage.IsTrackingSkipped(EnumHasAnyFlags(Buffer->Flags, ERDGBufferFlags::SkipTracking))
 			<< BufferMessage.IsTransient(bool(Buffer->bTransient))
 			<< BufferMessage.IsTransientUntracked(false)
-			<< BufferMessage.IsTransientCacheHit(TransientAllocation.bCacheHit);
+			<< BufferMessage.IsTransientCacheHit(TransientAllocation.bCacheHit)
+			<< BufferMessage.IsHandleType32Bits(true);
 	}
 
 	int32 TextureIndex = Textures.Num();
@@ -392,7 +400,8 @@ void FRDGTrace::OutputGraphEnd(const FRDGBuilder& GraphBuilder)
 				<< TextureMessage.IsTrackingSkipped(false)
 				<< TextureMessage.IsTransient(true)
 				<< TextureMessage.IsTransientUntracked(true)
-				<< TextureMessage.IsTransientCacheHit(TransientAllocation.bCacheHit);
+				<< TextureMessage.IsTransientCacheHit(TransientAllocation.bCacheHit)
+				<< TextureMessage.IsHandleType32Bits(true);
 
 			TextureIndex++;
 		}
@@ -415,7 +424,8 @@ void FRDGTrace::OutputGraphEnd(const FRDGBuilder& GraphBuilder)
 				<< BufferMessage.IsTrackingSkipped(false)
 				<< BufferMessage.IsTransient(true)
 				<< BufferMessage.IsTransientUntracked(true)
-				<< BufferMessage.IsTransientCacheHit(TransientAllocation.bCacheHit);
+				<< BufferMessage.IsTransientCacheHit(TransientAllocation.bCacheHit)
+				<< BufferMessage.IsHandleType32Bits(true);
 
 			BufferIndex++;
 		}
diff --git a/Engine/Source/Runtime/RenderCore/Public/RenderGraphDefinitions.h b/Engine/Source/Runtime/RenderCore/Public/RenderGraphDefinitions.h
--- a/Engine/Source/Runtime/RenderCore/Public/RenderGraphDefinitions.h
+++ b/Engine/Source/Runtime/RenderCore/Public/RenderGraphDefinitions.h
@@ -287,7 +287,7 @@ public:
 
 	TRDGHandle() = default;
 
-	explicit inline TRDGHandle(int32 InIndex)
+	explicit inline TRDGHandle(uint32 InIndex)
 	{
 		check(InIndex >= 0 && InIndex <= kNullIndex);
 		Index = (IndexType)InIndex;
@@ -677,25 +677,25 @@ class FRDGUserValidation;
 
 class FRDGViewableResource;
 
-using FRDGPassHandle = TRDGHandle<FRDGPass, uint16>;
+using FRDGPassHandle = TRDGHandle<FRDGPass, uint32>;
 using FRDGPassRegistry = TRDGHandleRegistry<FRDGPassHandle>;
 using FRDGPassHandleArray = TArray<FRDGPassHandle, TInlineAllocator<4, FRDGArrayAllocator>>;
 using FRDGPassBitArray = TRDGHandleBitArray<FRDGPassHandle>;
 
-using FRDGUniformBufferHandle = TRDGHandle<FRDGUniformBuffer, uint16>;
+using FRDGUniformBufferHandle = TRDGHandle<FRDGUniformBuffer, uint32>;
 using FRDGUniformBufferRegistry = TRDGHandleRegistry<FRDGUniformBufferHandle>;
 using FRDGUniformBufferBitArray = TRDGHandleBitArray<FRDGUniformBufferHandle>;
 
-using FRDGViewHandle = TRDGHandle<FRDGView, uint16>;
+using FRDGViewHandle = TRDGHandle<FRDGView, uint32>;
 using FRDGViewRegistry = TRDGHandleRegistry<FRDGViewHandle, ERDGHandleRegistryDestructPolicy::Never>;
 using FRDGViewUniqueFilter = TRDGHandleUniqueFilter<FRDGViewHandle>;
 using FRDGViewBitArray = TRDGHandleBitArray<FRDGViewHandle>;
 
-using FRDGTextureHandle = TRDGHandle<FRDGTexture, uint16>;
+using FRDGTextureHandle = TRDGHandle<FRDGTexture, uint32>;
 using FRDGTextureRegistry = TRDGHandleRegistry<FRDGTextureHandle, ERDGHandleRegistryDestructPolicy::Never>;
 using FRDGTextureBitArray = TRDGHandleBitArray<FRDGTextureHandle>;
 
-using FRDGBufferHandle = TRDGHandle<FRDGBuffer, uint16>;
+using FRDGBufferHandle = TRDGHandle<FRDGBuffer, uint32>;
 using FRDGBufferRegistry = TRDGHandleRegistry<FRDGBufferHandle, ERDGHandleRegistryDestructPolicy::Registry>;
 using FRDGBufferBitArray = TRDGHandleBitArray<FRDGBufferHandle>;
 
diff --git a/Engine/Source/Runtime/RenderCore/Public/RenderGraphPass.h b/Engine/Source/Runtime/RenderCore/Public/RenderGraphPass.h
--- a/Engine/Source/Runtime/RenderCore/Public/RenderGraphPass.h
+++ b/Engine/Source/Runtime/RenderCore/Public/RenderGraphPass.h
@@ -44,9 +44,9 @@ struct FRDGBarrierBatchBeginId
 
 	friend uint32 GetTypeHash(FRDGBarrierBatchBeginId Id)
 	{
-		static_assert(sizeof(Id.Passes) == 4);
-		uint32 Hash = *(const uint32*)Id.Passes.GetData();
-		return (Hash << GetRHIPipelineCount()) | uint32(Id.PipelinesAfter);
+		static_assert(sizeof(Id.Passes) <= 8);
+		uint32 Hash = GetTypeHash(*(const uint64*)Id.Passes.GetData());
+		return HashCombineFast(Hash, (uint32)Id.PipelinesAfter);
 	}
 
 	FRDGPassHandlesByPipeline Passes;
@@ -452,7 +452,7 @@ protected:
 		FRDGTextureRef Texture = nullptr;
 		FRDGTextureSubresourceState State;
 		FRDGTextureSubresourceState MergeState;
-		uint16 ReferenceCount = 0;
+		uint32 ReferenceCount = 0;
 	};
 
 	struct FBufferState
@@ -466,7 +466,7 @@ protected:
 		FRDGBufferRef Buffer = nullptr;
 		FRDGSubresourceState State;
 		FRDGSubresourceState* MergeState = nullptr;
-		uint16 ReferenceCount = 0;
+		uint32 ReferenceCount = 0;
 	};
 
 	/** Maps textures / buffers to information on how they are used in the pass. */
@@ -505,7 +505,7 @@ protected:
 
 	EAsyncComputeBudget AsyncComputeBudget = EAsyncComputeBudget::EAll_4;
 
-	uint16 ParallelPassSetIndex = 0;
+	uint32 ParallelPassSetIndex = 0;
 
 #if WITH_MGPU
 	FRHIGPUMask GPUMask;
diff --git a/Engine/Source/Runtime/RenderCore/Public/RenderGraphResources.h b/Engine/Source/Runtime/RenderCore/Public/RenderGraphResources.h
--- a/Engine/Source/Runtime/RenderCore/Public/RenderGraphResources.h
+++ b/Engine/Source/Runtime/RenderCore/Public/RenderGraphResources.h
@@ -401,10 +401,10 @@ protected:
 	FRDGPassHandle MinDiscardPass;
 
 	/** Number of references in passes and deferred queries. */
-	uint16 ReferenceCount;
+	uint32 ReferenceCount;
 
 	/** Scratch index allocated for the resource in the pass being setup. */
-	uint16 PassStateIndex = 0;
+	uint32 PassStateIndex = 0;
 
 	/** Set of aliasing overlaps to apply to the acquire transition if transient. */
 	TArrayView<const FRHITransientAliasingOverlap> AliasingOverlaps;
@@ -415,19 +415,19 @@ protected:
 private:
 	static const uint16 DeallocatedReferenceCount = ~0;
 
-	void SetExternalAccessMode(ERHIAccess InReadOnlyAccess, ERHIPipeline InPipelines)
+	void SetExternalAccessMode(ERHIAccess InAccess, ERHIPipeline InPipelines)
 	{
 		check(!AccessModeState.bLocked);
 
 		AccessModeState.Mode = EAccessMode::External;
-		AccessModeState.Access = InReadOnlyAccess;
+		AccessModeState.Access = InAccess;
 		AccessModeState.Pipelines = InPipelines;
 
-		EpilogueAccess = InReadOnlyAccess;
+		EpilogueAccess = InAccess;
 	}
 
 #if RDG_ENABLE_TRACE
-	uint16 TraceOrder = 0;
+	uint32 TraceOrder = 0;
 	TArray<FRDGPassHandle, FRDGArrayAllocator> TracePasses;
 #endif
 
