[shaders] optimize global shader config overrides mechanism to store a cache of EShaderPlatform to constructed FPlatformData; these were already cached per shader format but we weren't caching the platform -> format association (and enumerating all target formats for each target platform is actually surprisingly slow); this improves performance of GetGlobalShaderMapKeyString by approximately 4x

diff --git a/Engine/Source/Runtime/RenderCore/Private/GlobalShader.cpp b/Engine/Source/Runtime/RenderCore/Private/GlobalShader.cpp
--- a/Engine/Source/Runtime/RenderCore/Private/GlobalShader.cpp
+++ b/Engine/Source/Runtime/RenderCore/Private/GlobalShader.cpp
@@ -153,27 +153,45 @@ private:
 	struct FPlatformInfo
 	{
 		FName Name;
-		bool bInitializedFromConfig = false;
+		bool bInitialized = false;
 		FShaderDefines ShaderDefines;
 	};
 	
-	static TMap<FName, FPlatformInfo> ConfigDefines;
+	static TArray<FPlatformInfo> ConfigDefines;
+	static TStaticArray<const FPlatformInfo*, EShaderPlatform::SP_NumPlatforms> PerPlatformConfigs;
 	static TStaticBitArray<EShaderPlatform::SP_NumPlatforms> ErrorCheckedPlatforms;
 
 	static const FPlatformInfo* GetPlatformInfoAndErrorCheck(EShaderPlatform ShaderPlatform, FName& OutShaderFormat)
 	{
+		const TArray<ITargetPlatform*> AllPlatforms = GetTargetPlatformManagerRef().GetTargetPlatforms();
+		if (AllPlatforms.Num() != ConfigDefines.Num())
+		{
+			// if the number of platforms has changed since this was last called, we need to resize the array 
+			// and reset PerPlatformConfigs since the pointers will now be invalid.
+			ConfigDefines.SetNum(AllPlatforms.Num());
+			for (uint16 SpIndex = 0; SpIndex < EShaderPlatform::SP_NumPlatforms; ++SpIndex)
+			{
+				PerPlatformConfigs[SpIndex] = nullptr;
+			}
+		}
+
+		if (PerPlatformConfigs[ShaderPlatform] != nullptr)
+		{
+			return PerPlatformConfigs[ShaderPlatform];
+		}
+
 		OutShaderFormat = LegacyShaderPlatformToShaderFormat(ShaderPlatform);
 
 		// Search for all target platforms that support this shader platform
-		TArray<FName, TInlineAllocator<16>> IniPlatforms;
-		for(const ITargetPlatform* TP : GetTargetPlatformManagerRef().GetTargetPlatforms())
+		TArray<const ITargetPlatform*, TInlineAllocator<16>> IniPlatforms;
+		for(const ITargetPlatform* TP : AllPlatforms)
 		{
 			check(TP);
 			TArray<FName> PlatformShaderFormats;
 			TP->GetAllPossibleShaderFormats(PlatformShaderFormats);
 			if (PlatformShaderFormats.Contains(OutShaderFormat))
 			{
-				IniPlatforms.AddUnique(FName(TP->IniPlatformName()));
+				IniPlatforms.AddUnique(TP);
 			}
 		}
 		
@@ -183,15 +201,20 @@ private:
 			return nullptr;
 		}
 		
-		// first add all platforms to populate the map :
+		// first add all platforms to populate the array
 		for (int32 PlatformIndex = 0; PlatformIndex < IniPlatforms.Num(); ++PlatformIndex)
 		{
-			FPlatformInfo& Platform = ConfigDefines.FindOrAdd(IniPlatforms[PlatformIndex]);
+			int32 PlatformOrdinal = IniPlatforms[PlatformIndex]->GetPlatformOrdinal();
+			if (ConfigDefines.Num() < PlatformOrdinal + 1)
+			{
+				ConfigDefines.SetNum(PlatformOrdinal + 1);
+			}
+			FPlatformInfo& Platform = ConfigDefines[PlatformOrdinal];
 
-			if (!Platform.bInitializedFromConfig)
+			if (!Platform.bInitialized)
 			{
 				InitializePlatform(Platform, IniPlatforms[PlatformIndex]);
-				check( Platform.bInitializedFromConfig );
+				check(Platform.bInitialized);
 			}
 		}
 
@@ -203,31 +226,28 @@ private:
 			{
 				// This shader platform is shared by multiple target platforms that can be configured independently. We need to make sure all config defines
 				// match up, and that no platform-specific defines exist that might introduce shader compiler output that diverges between target platforms
-								
-				// pointer to Platform0 is safe to hold now because all are added first :
-				const FPlatformInfo * Platform0 = ConfigDefines.Find(IniPlatforms[0]);
-				check( Platform0 != nullptr );
+				const FPlatformInfo& Platform0 = ConfigDefines[IniPlatforms[0]->GetPlatformOrdinal()];
+				check(Platform0.bInitialized);
 
 				for (int32 PlatformIndex = 1; PlatformIndex < IniPlatforms.Num(); ++PlatformIndex)
 				{
-					const FPlatformInfo * OtherPlatform = ConfigDefines.Find(IniPlatforms[PlatformIndex]);
-					check( OtherPlatform != nullptr );
+					const FPlatformInfo& OtherPlatform = ConfigDefines[IniPlatforms[PlatformIndex]->GetPlatformOrdinal()];
+					check(OtherPlatform.bInitialized);
 
-					ErrorCheckPlatformsForShaderFormat(*Platform0, *OtherPlatform, OutShaderFormat);
+					ErrorCheckPlatformsForShaderFormat(Platform0, OtherPlatform, OutShaderFormat);
 				}
 			}
 		}
 		
-		// reget the pointer to IniPlatforms[0] after all Map adds are done, to return out :
-		// note returned pointer is not safe if any more adds are done
-		const FPlatformInfo * Platform = ConfigDefines.Find(IniPlatforms[0]);
-
-		return Platform;
+		const FPlatformInfo& Platform = ConfigDefines[IniPlatforms[0]->GetPlatformOrdinal()];
+		PerPlatformConfigs[ShaderPlatform] = &Platform;
+		return &Platform;
 	}
 
-	static void InitializePlatform(FPlatformInfo& Platform, FName PlatformName)
+	static void InitializePlatform(FPlatformInfo& Platform, const ITargetPlatform* TP)
 	{
-		if (FConfigCacheIni* ConfigCache = FConfigCacheIni::ForPlatform(PlatformName))
+		Platform.Name = FName(TP->IniPlatformName());
+		if (FConfigCacheIni* ConfigCache = FConfigCacheIni::ForPlatform(Platform.Name))
 		{
 			TArray<FString> DefineStrings;
 			ConfigCache->GetArray(TEXT("GlobalShaderDefines"), TEXT("Definitions"), DefineStrings, GEngineIni);
@@ -279,11 +299,9 @@ private:
 
 					Define->Value = DefineValue;
 				}
-			}			
-		}		
-
-		Platform.Name = PlatformName;
-		Platform.bInitializedFromConfig = true;
+			}
+		}
+		Platform.bInitialized = true;
 	}
 
 	static void ErrorCheckPlatformsForShaderFormat(const FPlatformInfo& PlatformA, const FPlatformInfo& PlatformB, FName ShaderFormat)
@@ -350,7 +368,8 @@ private:
 	}
 };
 
-TMap<FName, FGlobalShaderConfigDefines::FPlatformInfo> FGlobalShaderConfigDefines::ConfigDefines;
+TArray<FGlobalShaderConfigDefines::FPlatformInfo> FGlobalShaderConfigDefines::ConfigDefines;
+TStaticArray<const FGlobalShaderConfigDefines::FPlatformInfo*, EShaderPlatform::SP_NumPlatforms> FGlobalShaderConfigDefines::PerPlatformConfigs;
 TStaticBitArray<EShaderPlatform::SP_NumPlatforms> FGlobalShaderConfigDefines::ErrorCheckedPlatforms;
 
 /** Used to identify the global shader map in compile queues. */
