[shaders] optimize AppendKeyStringShaderDependencies by caching FShaderParametersMetadata pointers in each shader/vf type instead of uniform buffer names; this eliminates a bunch of string hashing/comparison/fname construction that was previously needed every time this function is called and improves performance of GetMaterialShaderMapKeyString by approximately 20%

diff --git a/Engine/Plugins/FX/Niagara/Source/NiagaraShader/Private/NiagaraShader.cpp b/Engine/Plugins/FX/Niagara/Source/NiagaraShader/Private/NiagaraShader.cpp
--- a/Engine/Plugins/FX/Niagara/Source/NiagaraShader/Private/NiagaraShader.cpp
+++ b/Engine/Plugins/FX/Niagara/Source/NiagaraShader/Private/NiagaraShader.cpp
@@ -462,7 +462,7 @@ void FNiagaraShaderType::BeginCompileShaderFromSource(
 
 void FNiagaraShaderType::AddUniformBufferIncludesToEnvironment(FShaderCompilerEnvironment& OutEnvironment, EShaderPlatform Platform) const
 {
-	UE::ShaderParameters::AddUniformBufferIncludesToEnvironment(OutEnvironment, ReferencedUniformBufferNames);
+	UE::ShaderParameters::AddUniformBufferIncludesToEnvironment(OutEnvironment, ReferencedUniformBuffers);
 }
 
 /**
diff --git a/Engine/Source/Runtime/Engine/Private/ShaderCompiler/ShaderCompiler.cpp b/Engine/Source/Runtime/Engine/Private/ShaderCompiler/ShaderCompiler.cpp
--- a/Engine/Source/Runtime/Engine/Private/ShaderCompiler/ShaderCompiler.cpp
+++ b/Engine/Source/Runtime/Engine/Private/ShaderCompiler/ShaderCompiler.cpp
@@ -8058,7 +8058,7 @@ void GlobalBeginCompileShader(
 	}
 
 	// Reserve space in maps to prevent reallocation and rehashing in AddUniformBufferIncludesToEnvironment -- plus one at the end is for GeneratedInstancedStereo.ush
-	const int32 UniformBufferReserveNum = Input.Environment.UniformBufferMap.Num() + ShaderType->GetReferencedUniformBufferNames().Num() + (VFType ? VFType->GetReferencedUniformBufferNames().Num() : 0) + 1;
+	const int32 UniformBufferReserveNum = Input.Environment.UniformBufferMap.Num() + ShaderType->GetReferencedUniformBuffers().Num() + (VFType ? VFType->GetReferencedUniformBuffers().Num() : 0) + 1;
 	Input.Environment.UniformBufferMap.Reserve(UniformBufferReserveNum);
 	Input.Environment.IncludeVirtualPathToSharedContentsMap.Reserve(UniformBufferReserveNum);
 
@@ -9580,6 +9580,7 @@ static FString GetGlobalShaderCacheFilename(EShaderPlatform Platform)
 
 static FString GetGlobalShaderMapKeyString(const FGlobalShaderMapId& ShaderMapId, EShaderPlatform Platform, TArray<FShaderTypeDependency> const& Dependencies)
 {
+	TRACE_CPUPROFILER_EVENT_SCOPE(GetGlobalShaderMapKeyString);
 	FName Format = LegacyShaderPlatformToShaderFormat(Platform);
 	FString ShaderMapKeyString = Format.ToString() + TEXT("_") + FString(FString::FromInt(GetTargetPlatformManagerRef().ShaderFormatVersion(Format))) + TEXT("_");
 	ShaderMapAppendKeyString(Platform, ShaderMapKeyString);
diff --git a/Engine/Source/Runtime/RenderCore/Private/Shader.cpp b/Engine/Source/Runtime/RenderCore/Private/Shader.cpp
--- a/Engine/Source/Runtime/RenderCore/Private/Shader.cpp
+++ b/Engine/Source/Runtime/RenderCore/Private/Shader.cpp
@@ -480,8 +480,8 @@ bool FShaderType::ValidateCompiledResult(EShaderPlatform Platform, const FShader
 
 void FShaderType::UpdateReferencedUniformBufferNames(const TMap<FString, TArray<const TCHAR*>>& ShaderFileToUniformBufferVariables)
 {
-	ReferencedUniformBufferNames.Empty();
-	GenerateReferencedUniformBufferNames(SourceFilename, Name, ShaderFileToUniformBufferVariables, ReferencedUniformBufferNames);
+	ReferencedUniformBuffers.Empty();
+	GenerateReferencedUniformBuffers(SourceFilename, Name, ShaderFileToUniformBufferVariables, ReferencedUniformBuffers);
 }
 #endif // WITH_EDITOR
 
diff --git a/Engine/Source/Runtime/RenderCore/Private/ShaderCore.cpp b/Engine/Source/Runtime/RenderCore/Private/ShaderCore.cpp
--- a/Engine/Source/Runtime/RenderCore/Private/ShaderCore.cpp
+++ b/Engine/Source/Runtime/RenderCore/Private/ShaderCore.cpp
@@ -2999,11 +2999,11 @@ void UpdateReferencedUniformBufferNames(
 	}
 }
 
-void GenerateReferencedUniformBufferNames(
+void GenerateReferencedUniformBuffers(
 	const TCHAR* SourceFilename,
 	const TCHAR* ShaderTypeName,
 	const TMap<FString, TArray<const TCHAR*> >& ShaderFileToUniformBufferVariables,
-	TSet<const TCHAR*, TStringPointerSetKeyFuncs_DEPRECATED<const TCHAR*>>& UniformBufferNames)
+	TSet<const FShaderParametersMetadata*>& UniformBuffers)
 {
 	TArray<FString> FilesToSearch;
 	GetShaderIncludes(SourceFilename, SourceFilename, FilesToSearch, GMaxRHIShaderPlatform);
@@ -3014,7 +3014,11 @@ void GenerateReferencedUniformBufferNames(
 		const TArray<const TCHAR*>& FoundUniformBufferVariables = ShaderFileToUniformBufferVariables.FindChecked(FileToSearch);
 		for (const TCHAR* UniformBufferName : FoundUniformBufferVariables)
 		{
-			UniformBufferNames.Emplace(UniformBufferName);
+			const FShaderParametersMetadata* UniformBufferStruct = FindUniformBufferStructByName(UniformBufferName);
+			if (UniformBufferStruct)
+			{
+				UniformBuffers.Emplace(UniformBufferStruct);
+			}
 		}
 	}
 }
@@ -3102,7 +3106,10 @@ void AppendKeyStringShaderDependencies(
 	FString& OutKeyString,
 	bool bIncludeSourceHashes)
 {
-	TSet<const TCHAR*, TStringPointerSetKeyFuncs_DEPRECATED<const TCHAR*>> ReferencedUniformBufferNames;
+	FMemMark MemMark(FMemStack::Get());
+	using FMemStackSetAllocator = TSetAllocator<TSparseArrayAllocator<TMemStackAllocator<>, TMemStackAllocator<>>, TMemStackAllocator<>>;
+	TSet<const FShaderParametersMetadata*, DefaultKeyFuncs<const FShaderParametersMetadata*>, FMemStackSetAllocator> ReferencedUniformBuffers;
+	ReferencedUniformBuffers.Reserve(128);
 
 	for (const FShaderTypeDependency& ShaderTypeDependency : ShaderTypeDependencies)
 	{
@@ -3132,9 +3139,9 @@ void AppendKeyStringShaderDependencies(
 		const FSHAHash LayoutHash = GetShaderTypeLayoutHash(ShaderType->GetLayout(), LayoutParams);
 		LayoutHash.AppendString(OutKeyString);
 
-		for (const TCHAR* UniformBufferName : ShaderType->GetReferencedUniformBufferNames())
+		for (const FShaderParametersMetadata* UniformBuffer : ShaderType->GetReferencedUniformBuffers())
 		{
-			ReferencedUniformBufferNames.Add(UniformBufferName);
+			ReferencedUniformBuffers.Add(UniformBuffer);
 		}
 	}
 
@@ -3159,9 +3166,9 @@ void AppendKeyStringShaderDependencies(
 				ParameterStructMetadata->AppendKeyString(OutKeyString);
 			}
 
-			for (const TCHAR* UniformBufferName : ShaderType->GetReferencedUniformBufferNames())
+			for (const FShaderParametersMetadata* UniformBuffer : ShaderType->GetReferencedUniformBuffers())
 			{
-				ReferencedUniformBufferNames.Add(UniformBufferName);
+				ReferencedUniformBuffers.Add(UniformBuffer);
 			}
 		}
 	}
@@ -3189,22 +3196,26 @@ void AppendKeyStringShaderDependencies(
 			}
 		}
 
-		for (const TCHAR* UniformBufferName : VertexFactoryType->GetReferencedUniformBufferNames())
+		for (const FShaderParametersMetadata* UniformBuffer : VertexFactoryType->GetReferencedUniformBuffers())
 		{
-			ReferencedUniformBufferNames.Add(UniformBufferName);
+			ReferencedUniformBuffers.Add(UniformBuffer);
 		}
 	}
 
+	struct FUbSortByLayoutSignature
 	{
-		TArray<const TCHAR*> SortedUniformBufferNames = ReferencedUniformBufferNames.Array();
-		Algo::Sort(SortedUniformBufferNames, FUniformBufferNameSortOrder());
-
-		// Save uniform buffer member info so we can detect when layout has changed
-		for (const TCHAR* UniformBufferName : SortedUniformBufferNames)
+		bool operator()(const FShaderParametersMetadata& A, const FShaderParametersMetadata& B) const
 		{
-			FShaderParametersMetadata* UniformBufferMetadata = FindUniformBufferStructByName(UniformBufferName);
-			UniformBufferMetadata->AppendKeyString(OutKeyString);
+			return A.GetLayoutSignature() < B.GetLayoutSignature();
 		}
+	};
+	// sort the referenced uniform buffers by the stable layout signature; for ddc keys we care about stability not alphabetical ordering by name
+	ReferencedUniformBuffers.StableSort(FUbSortByLayoutSignature());
+
+	// Save uniform buffer member info so we can detect when layout has changed
+	for (const FShaderParametersMetadata* UniformBufferMetadata : ReferencedUniformBuffers)
+	{
+		UniformBufferMetadata->AppendKeyString(OutKeyString);
 	}
 }
 
diff --git a/Engine/Source/Runtime/RenderCore/Private/ShaderParameters.cpp b/Engine/Source/Runtime/RenderCore/Private/ShaderParameters.cpp
--- a/Engine/Source/Runtime/RenderCore/Private/ShaderParameters.cpp
+++ b/Engine/Source/Runtime/RenderCore/Private/ShaderParameters.cpp
@@ -412,39 +412,18 @@ FString UE::ShaderParameters::CreateUniformBufferShaderDeclaration(const TCHAR*
 	return CreateHLSLUniformBufferDeclaration(UniformBufferName, UniformBufferStruct);
 }
 
-static FShaderParametersMetadata* FindShaderParametersMetadataWithVariableName(uint32 InVariableNameHash, FStringView InVariableNameView)
-{
-#if WITH_EDITOR
-	TMap<FString, FShaderParametersMetadata*>& StringStructMap = FShaderParametersMetadata::GetStringStructMap();
-
-	FShaderParametersMetadata** FoundMetadata = StringStructMap.FindByHash(InVariableNameHash, InVariableNameView);
-	return FoundMetadata ? *FoundMetadata : nullptr;
-#else // WITH_EDITOR
-	for (FShaderParametersMetadata* Metadata : *FShaderParametersMetadata::GetStructList())
-	{
-		if (InVariableName == Metadata->GetShaderVariableName())
-		{
-			return Metadata;
-		}
-	}
-	return nullptr;
-#endif
-}
-
-void UE::ShaderParameters::AddUniformBufferIncludesToEnvironment(FShaderCompilerEnvironment& OutEnvironment, const TSet<const TCHAR*, TStringPointerSetKeyFuncs_DEPRECATED<const TCHAR*>>& InUniformBufferNames)
+void UE::ShaderParameters::AddUniformBufferIncludesToEnvironment(FShaderCompilerEnvironment& OutEnvironment, const TSet<const FShaderParametersMetadata*>& InUniformBuffers)
 {
 	TRACE_CPUPROFILER_EVENT_SCOPE(UE::ShaderParameters::AddUniformBufferIncludesToEnvironment);
 
 	FString UniformBufferIncludes;
 
-	for (const TCHAR* UniformBufferName : InUniformBufferNames)
+	for (const FShaderParametersMetadata* Metadata : InUniformBuffers)
 	{
-		FStringView UniformBufferNameView(UniformBufferName);
+		FStringView UniformBufferNameView(Metadata->GetShaderVariableName());
 		uint32 UniformBufferNameHash = GetTypeHash(UniformBufferNameView);
 		if (!OutEnvironment.UniformBufferMap.FindByHash(UniformBufferNameHash, UniformBufferNameView))
 		{
-			if (const FShaderParametersMetadata* Metadata = FindShaderParametersMetadataWithVariableName(UniformBufferNameHash, UniformBufferNameView))
-			{
 				const FThreadSafeSharedAnsiStringPtr UniformBufferDeclaration = Metadata->GetUniformBufferDeclarationAnsiPtr();
 
 				check(UniformBufferDeclaration.Get() != NULL);
@@ -455,7 +434,6 @@ void UE::ShaderParameters::AddUniformBufferIncludesToEnvironment(FShaderCompiler
 				OutEnvironment.IncludeVirtualPathToSharedContentsMap.AddByHash(Metadata->GetUniformBufferPathHash(), Metadata->GetUniformBufferPath(), UniformBufferDeclaration);
 
 				Metadata->AddResourceTableEntries(OutEnvironment.ResourceTableMap, OutEnvironment.UniformBufferMap);
-			}
 		}
 	}
 
@@ -467,7 +445,7 @@ void FShaderType::AddUniformBufferIncludesToEnvironment(FShaderCompilerEnvironme
 {
 	TRACE_CPUPROFILER_EVENT_SCOPE(FShaderType::AddReferencedUniformBufferIncludes);
 
-	UE::ShaderParameters::AddUniformBufferIncludesToEnvironment(OutEnvironment, ReferencedUniformBufferNames);
+	UE::ShaderParameters::AddUniformBufferIncludesToEnvironment(OutEnvironment, ReferencedUniformBuffers);
 }
 
 #endif // WITH_EDITOR
@@ -533,12 +511,9 @@ void FShaderType::GetShaderStableKeyParts(FStableShaderKeyAndValue& SaveKeyVal)
 	SaveKeyVal.ShaderType = FName(GetName() ? GetName() : TEXT("null"));
 }
 
-void FVertexFactoryType::FlushShaderFileCache(const TMap<FString, TArray<const TCHAR*> >& ShaderFileToUniformBufferVariables)
-{
-}
 void FVertexFactoryType::AddUniformBufferIncludesToEnvironment(FShaderCompilerEnvironment& OutEnvironment, EShaderPlatform Platform) const
 {
-	UE::ShaderParameters::AddUniformBufferIncludesToEnvironment(OutEnvironment, ReferencedUniformBufferNames);
+	UE::ShaderParameters::AddUniformBufferIncludesToEnvironment(OutEnvironment, ReferencedUniformBuffers);
 }
 
 #endif // WITH_EDITOR
diff --git a/Engine/Source/Runtime/RenderCore/Private/VertexFactory.cpp b/Engine/Source/Runtime/RenderCore/Private/VertexFactory.cpp
--- a/Engine/Source/Runtime/RenderCore/Private/VertexFactory.cpp
+++ b/Engine/Source/Runtime/RenderCore/Private/VertexFactory.cpp
@@ -55,8 +55,8 @@ FVertexFactoryType* FVertexFactoryType::GetVFByName(const FHashedName& VFName)
 #if WITH_EDITOR
 void FVertexFactoryType::UpdateReferencedUniformBufferNames(const TMap<FString, TArray<const TCHAR*>>& ShaderFileToUniformBufferVariables)
 {
-	ReferencedUniformBufferNames.Empty();
-	GenerateReferencedUniformBufferNames(ShaderFilename, Name, ShaderFileToUniformBufferVariables, ReferencedUniformBufferNames);
+	ReferencedUniformBuffers.Empty();
+	GenerateReferencedUniformBuffers(ShaderFilename, Name, ShaderFileToUniformBufferVariables, ReferencedUniformBuffers);
 }
 #endif
 
diff --git a/Engine/Source/Runtime/RenderCore/Public/Shader.h b/Engine/Source/Runtime/RenderCore/Public/Shader.h
--- a/Engine/Source/Runtime/RenderCore/Public/Shader.h
+++ b/Engine/Source/Runtime/RenderCore/Public/Shader.h
@@ -1439,9 +1439,16 @@ public:
 	}
 
 #if WITH_EDITOR
+	UE_DEPRECATED(5.5, "GetReferencedUniformBufferNames is deprecated, use GetReferencedUniformBuffers instead")
 	inline const TSet<const TCHAR*, TStringPointerSetKeyFuncs_DEPRECATED<const TCHAR*>>& GetReferencedUniformBufferNames() const
 	{
-		return ReferencedUniformBufferNames;
+		static TSet<const TCHAR*, TStringPointerSetKeyFuncs_DEPRECATED<const TCHAR*>> EmptySet;
+		return EmptySet;
+	}
+
+	const TSet<const FShaderParametersMetadata*>& GetReferencedUniformBuffers() const
+	{
+		return ReferencedUniformBuffers;
 	}
 
 	/** Adds include statements for uniform buffers that this shader type references. */
@@ -1487,11 +1494,11 @@ private:
 #if WITH_EDITOR
 protected:
 	/**
-	* Cache of referenced uniform buffer includes.
+	* Cache of referenced uniform buffer structs.
 	* These are derived from source files so they need to be flushed when editing and recompiling shaders on the fly.
 	* FShaderType::Initialize will add the referenced uniform buffers, but this set may be updated by FlushShaderFileCache.
 	*/
-	TSet<const TCHAR*, TStringPointerSetKeyFuncs_DEPRECATED<const TCHAR*>> ReferencedUniformBufferNames;
+	TSet<const FShaderParametersMetadata*> ReferencedUniformBuffers;
 #endif // WITH_EDITOR
 };
 
diff --git a/Engine/Source/Runtime/RenderCore/Public/ShaderCore.h b/Engine/Source/Runtime/RenderCore/Public/ShaderCore.h
--- a/Engine/Source/Runtime/RenderCore/Public/ShaderCore.h
+++ b/Engine/Source/Runtime/RenderCore/Public/ShaderCore.h
@@ -1200,11 +1200,11 @@ RENDERCORE_API void UpdateReferencedUniformBufferNames(
 	TArrayView<const FShaderPipelineType*> OutdatedShaderPipelineTypes);
 
 /** Parses the given source file and its includes for references of uniform buffers. */
-extern void GenerateReferencedUniformBufferNames(
+extern void GenerateReferencedUniformBuffers(
 	const TCHAR* SourceFilename,
 	const TCHAR* ShaderTypeName,
 	const TMap<FString, TArray<const TCHAR*> >& ShaderFileToUniformBufferVariables,
-	TSet<const TCHAR*, TStringPointerSetKeyFuncs_DEPRECATED<const TCHAR*>>& UniformBufferNames);
+	TSet<const FShaderParametersMetadata*>& UniformBuffers);
 
 struct FUniformBufferNameSortOrder
 {
diff --git a/Engine/Source/Runtime/RenderCore/Public/ShaderParameters.h b/Engine/Source/Runtime/RenderCore/Public/ShaderParameters.h
--- a/Engine/Source/Runtime/RenderCore/Public/ShaderParameters.h
+++ b/Engine/Source/Runtime/RenderCore/Public/ShaderParameters.h
@@ -39,7 +39,10 @@ namespace UE::ShaderParameters
 	/** Creates a shader code declaration of this struct for the given shader platform. */
 	RENDERCORE_API FString CreateUniformBufferShaderDeclaration(const TCHAR* Name, const FShaderParametersMetadata& UniformBufferStruct);
 
-	RENDERCORE_API void AddUniformBufferIncludesToEnvironment(FShaderCompilerEnvironment& OutEnvironment, const TSet<const TCHAR*, TStringPointerSetKeyFuncs_DEPRECATED<const TCHAR*>>& InUniformBufferNames);
+	UE_DEPRECATED(5.5, "AddUniformBufferIncludesToEnviroment now takes a set of FShaderParametersMetadata pointers")
+	inline void AddUniformBufferIncludesToEnvironment(FShaderCompilerEnvironment& OutEnvironment, const TSet<const TCHAR*, TStringPointerSetKeyFuncs_DEPRECATED<const TCHAR*>>& InUniformBufferNames) {}
+
+	RENDERCORE_API void AddUniformBufferIncludesToEnvironment(FShaderCompilerEnvironment& OutEnvironment, const TSet<const FShaderParametersMetadata*>& InUniformBuffers);
 }
 #endif
 
diff --git a/Engine/Source/Runtime/RenderCore/Public/VertexFactory.h b/Engine/Source/Runtime/RenderCore/Public/VertexFactory.h
--- a/Engine/Source/Runtime/RenderCore/Public/VertexFactory.h
+++ b/Engine/Source/Runtime/RenderCore/Public/VertexFactory.h
@@ -429,18 +429,19 @@ public:
 	/** Adds include statements for uniform buffers that this shader type references, and builds a prefix for the shader file with the include statements. */
 	RENDERCORE_API void AddUniformBufferIncludesToEnvironment(FShaderCompilerEnvironment& OutEnvironment, EShaderPlatform Platform) const;
 
-	UE_DEPRECATED(5.2, "AddReferencedUniformBufferIncludes has moved to AddUniformBufferIncludesToEnvironment and no longer takes a prefix argument.")
-	inline void AddReferencedUniformBufferIncludes(FShaderCompilerEnvironment& OutEnvironment, FString& OutSourceFilePrefix, EShaderPlatform Platform) const
+	UE_DEPRECATED(5.5, "GetReferencedUniformBufferNames is deprecated; call GetReferencedUniformBuffers instead.")
+	inline const TSet<const TCHAR*, TStringPointerSetKeyFuncs_DEPRECATED<const TCHAR*>>& GetReferencedUniformBufferNames() const 
 	{
-		AddUniformBufferIncludesToEnvironment(OutEnvironment, Platform);
+		static TSet<const TCHAR*, TStringPointerSetKeyFuncs_DEPRECATED<const TCHAR*>> EmptySet;
+		return EmptySet; 
+	};
+
+	const TSet<const FShaderParametersMetadata*>& GetReferencedUniformBuffers() const
+	{
+		return ReferencedUniformBuffers;
 	}
 
 	RENDERCORE_API void UpdateReferencedUniformBufferNames(const TMap<FString, TArray<const TCHAR*>>& ShaderFileToUniformBufferVariables);
-
-	UE_DEPRECATED(5.2, "FlushShaderFileCache is deprecated. UpdateReferencedUniformBufferNames should be used to flush any uniform buffer changes")
-	RENDERCORE_API void FlushShaderFileCache(const TMap<FString, TArray<const TCHAR*> >& ShaderFileToUniformBufferVariables);
-
-	inline const TSet<const TCHAR*, TStringPointerSetKeyFuncs_DEPRECATED<const TCHAR*>>& GetReferencedUniformBufferNames() const { return ReferencedUniformBufferNames; };
 #endif // WITH_EDITOR
 
 private:
@@ -468,11 +469,11 @@ private:
 
 #if WITH_EDITOR
 	/** 
-	 * Cache of referenced uniform buffer includes.  
+	 * Cache of referenced uniform buffer structs.
 	 * These are derived from source files so they need to be flushed when editing and recompiling shaders on the fly. 
 	 * FShaderType::Initialize will add the referenced uniform buffers, but this set may be updated by FlushShaderFileCache
 	 */
-	TSet<const TCHAR*, TStringPointerSetKeyFuncs_DEPRECATED<const TCHAR*>> ReferencedUniformBufferNames;
+	TSet<const FShaderParametersMetadata*> ReferencedUniformBuffers;
 #endif // WITH_EDITOR
 };
 
