Improve error handling for failed shader preprocessing

Shader preprocessing currently uses the same exception handler as shader compilation. The exception handler only makes sense when we call into platform shader compiler code, which we don't do for preprocessing. Otherwise it is preferable to use the existing error handling that already exists in the engine (the current exception handler turns fatal errors into warnings, for example).

Furthermore, the exception handler assumes that the global error history always includes a callstack. That is not the case and can cause ugly failures during error handling. This CL inserts a check and falls back to the general mechanism for extracting the callstack if no callstack is found. It's worth noting that this general mechanism will return the callstack of the exception handler, and not the callstack that caused the exception: the act of handling the exception will unwind the callstack and destroy that information.

diff --git a/Engine/Source/Runtime/RenderCore/Private/ShaderCore.cpp b/Engine/Source/Runtime/RenderCore/Private/ShaderCore.cpp
--- a/Engine/Source/Runtime/RenderCore/Private/ShaderCore.cpp
+++ b/Engine/Source/Runtime/RenderCore/Private/ShaderCore.cpp
@@ -1183,16 +1183,22 @@ int HandleShaderCompileException(Windows::LPEXCEPTION_POINTERS Info, FString& Ou
 {
 	const DWORD AssertExceptionCode = 0x00004000;
 	FString ExCodeStr;
+	OutCallStack = "";
 	if (Info->ExceptionRecord->ExceptionCode == AssertExceptionCode)
 	{
 		// In the case of an assert the assert handler populates the GErrorHist global.
-		// This contains a readable assert message followed by a callstack; so we can use that to populate
+		// This contains a readable assert message that may be followed by a callstack; so we can use that to populate
 		// our message/callstack and save some time as well as getting the properly formatted assert message.
-		FString Assert = GErrorHist;
 		const TCHAR* CallstackStart = FCString::Strfind(GErrorHist, TEXT("0x"));
-
-		OutExMsg = FString(CallstackStart - GErrorHist, GErrorHist);
-		OutCallStack = CallstackStart;
+		if (CallstackStart && CallstackStart > GErrorHist)
+		{
+			OutExMsg = FString(CallstackStart - GErrorHist, GErrorHist);
+			OutCallStack = CallstackStart;
+		}
+		else
+		{
+			OutExMsg = GErrorHist;
+		}
 	}
 	else
 	{
@@ -1210,7 +1216,10 @@ int HandleShaderCompileException(Windows::LPEXCEPTION_POINTERS Info, FString& Ou
 				Info->ExceptionRecord->ExceptionCode,
 				(uint64)Info->ExceptionRecord->ExceptionAddress);
 		}
+	}
 
+	if (OutCallStack.Len() == 0)
+	{
 		ANSICHAR CallStack[32768];
 		FMemory::Memzero(CallStack);
 		FPlatformStackWalk::StackWalkAndDump(CallStack, ARRAYSIZE(CallStack), Info->ExceptionRecord->ExceptionAddress);
@@ -1234,18 +1243,7 @@ public:
 		FString& OutExceptionCallstack,
 		FString& OutExceptionMsg)
 	{
-#if PLATFORM_WINDOWS
-		__try
-#endif
-		{
-			return Backend->PreprocessShader(Input, Environment, Output);
-		}
-#if PLATFORM_WINDOWS
-		__except (HandleShaderCompileException(GetExceptionInformation(), OutExceptionMsg, OutExceptionCallstack))
-		{
-			return false;
-		}
-#endif
+		return Backend->PreprocessShader(Input, Environment, Output);
 	}
 
 	static void PreprocessShaderInternal(
