Fixed crash in RDG when more than 64k passes are used.

diff --git a/Engine/Source/Runtime/RenderCore/Public/RenderGraphResources.h b/Engine/Source/Runtime/RenderCore/Public/RenderGraphResources.h
--- a/Engine/Source/Runtime/RenderCore/Public/RenderGraphResources.h
+++ b/Engine/Source/Runtime/RenderCore/Public/RenderGraphResources.h
@@ -410,7 +410,7 @@ protected:
 	ERHIAccess EpilogueAccess = DefaultEpilogueAccess;
 
 private:
-	static const uint16 DeallocatedReferenceCount = ~0;
+	static const uint32 DeallocatedReferenceCount = ~0;
 
 	void SetRHI(FRHIResource* Resource)
 	{
