[shaders] hash environment modifications for global & material shader types and include in shadermap DDC key; the result of this is that both modifications to ShouldCompilePermutation and ModifyCompilationEnvironment functions on individual shader types now correctly invalidate shadermaps, eliminating the need for version bumps in these cases and preventing the possibility of cache pollution bugs due to individual shader cache keys already taking these types of modifications into account (since these always affect the generated preprocessed shader code)

In conjunction with the optimizations submitted in https://p4-swarm.epicgames.net/changes/32085931 and https://p4-swarm.epicgames.net/changes/32084297, key construction for global shadermaps is actually about 4x faster, while material shadermap key construction has gotten about 30% slower; however in terms of absolute time delta it's hundreds of milliseconds of thread time on a cold cook of QAGame (extrapolating to known large content sets this is tens of seconds, but parallelized, so should have minimal noticeable impact)

diff --git a/Engine/Source/Runtime/Engine/Private/Materials/MaterialShader.cpp b/Engine/Source/Runtime/Engine/Private/Materials/MaterialShader.cpp
--- a/Engine/Source/Runtime/Engine/Private/Materials/MaterialShader.cpp
+++ b/Engine/Source/Runtime/Engine/Private/Materials/MaterialShader.cpp
@@ -9,6 +9,7 @@
 #include "Materials/MaterialInterface.h"
 #include "ProfilingDebugging/DiagnosticTable.h"
 #include "MeshMaterialShaderType.h"
+#include "MeshMaterialShader.h"
 #include "MaterialDomain.h"
 #include "MaterialShaderMapLayout.h"
 #include "SceneInterface.h"
@@ -186,8 +187,9 @@ FString GetBlendModeString(EBlendMode BlendMode)
 
 #if WITH_EDITOR
 /** Creates a string key for the derived data cache given a shader map id. */
-FString GetMaterialShaderMapKeyString(const FMaterialShaderMapId& ShaderMapId, EShaderPlatform Platform, bool bIncludeKeyStringShaderDependencies)
+FString GetMaterialShaderMapKeyString(const FMaterialShaderMapId& ShaderMapId, const FMaterialShaderParameters& ShaderParameters, EShaderPlatform Platform, bool bIncludeKeyStringShaderDependencies)
 {
+	TRACE_CPUPROFILER_EVENT_SCOPE(GetMaterialShaderMapKeyString);
 	FName Format = LegacyShaderPlatformToShaderFormat(Platform);
 	FString ShaderMapKeyString;
 	ShaderMapKeyString.Reserve(16384);
@@ -203,6 +205,47 @@ FString GetMaterialShaderMapKeyString(const FMaterialShaderMapId& ShaderMapId, E
 
 	ShaderMapAppendKeyString(Platform, ShaderMapKeyString);
 	ShaderMapId.AppendKeyString(ShaderMapKeyString, true, bIncludeKeyStringShaderDependencies);
+
+	FMemoryHasherBlake3 EnvHasher;
+
+	// Hash any modifications applied to the compilation environment for each shadertype that will be compiled.
+	// When the shadermap ID is constructed, individual shadertype dependencies are added for any pipelines, hence
+	// we don't need to have equivalent hashing code for the pipeline dependencies.
+	for (const FShaderTypeDependency& TypeDep : ShaderMapId.ShaderTypeDependencies)
+	{
+		const FShaderType* ShaderType = FindShaderTypeByName(TypeDep.ShaderTypeName);
+		if (const FMaterialShaderType* MatShaderType = ShaderType->GetMaterialShaderType())
+		{
+			if (MatShaderType->ShouldCompilePermutation(Platform, ShaderParameters, TypeDep.PermutationId, ShaderMapId.GetPermutationFlags()))
+			{
+				FShaderCompilerEnvironment EnvModifications(EnvHasher);
+				MatShaderType->SetupCompileEnvironment(Platform, ShaderParameters, TypeDep.PermutationId, ShaderMapId.GetPermutationFlags(), EnvModifications);
+				EnvModifications.SerializeEverythingButFiles(EnvHasher);
+			}
+		}
+		else if (const FMeshMaterialShaderType* MeshMatShaderType = ShaderType->GetMeshMaterialShaderType())
+		{
+			for (const FVertexFactoryTypeDependency& VFDep : ShaderMapId.VertexFactoryTypeDependencies)
+			{
+				const FVertexFactoryType* VFType = FindVertexFactoryType(VFDep.VertexFactoryTypeName);
+				const bool bVFShouldCache = FMeshMaterialShaderType::ShouldCompileVertexFactoryPermutation(Platform, ShaderParameters, VFType, ShaderType, ShaderMapId.GetPermutationFlags());
+				const bool bShaderShouldCache = MeshMatShaderType->ShouldCompilePermutation(Platform, ShaderParameters, VFType, TypeDep.PermutationId, ShaderMapId.GetPermutationFlags());
+				if (bVFShouldCache && bShaderShouldCache)
+				{
+					FShaderCompilerEnvironment EnvModifications(EnvHasher);
+					VFType->ModifyCompilationEnvironment(FVertexFactoryShaderPermutationParameters(Platform, ShaderParameters, VFType, ShaderType, ShaderMapId.GetPermutationFlags()), EnvModifications);
+					MeshMatShaderType->SetupCompileEnvironment(Platform, ShaderParameters, VFType, TypeDep.PermutationId, ShaderMapId.GetPermutationFlags(), EnvModifications);
+					EnvModifications.SerializeEverythingButFiles(EnvHasher);
+				}
+			}
+		}
+	}
+	// * 2 for hex representation of hash; + 6 for tag/underscores
+	TStringBuilder<sizeof(TCHAR) * (sizeof(FBlake3Hash::ByteArray) * 2 + 6)> EnvHashString;
+	EnvHashString << "_EMH_" << EnvHasher.Finalize() << '_';
+	check(EnvHashString.GetAllocatedSize() == 0);
+	ShaderMapKeyString.Append(EnvHashString.ToView());
+
 	FMaterialAttributeDefinitionMap::AppendDDCKeyString(ShaderMapKeyString);
 	FShaderCompileUtilities::AppendGBufferDDCKeyString(Platform, ShaderMapKeyString);
 
@@ -1410,12 +1453,12 @@ TSharedRef<FMaterialShaderMap::FAsyncLoadContext> FMaterialShaderMap::BeginLoadF
 
 	struct FMaterialShaderMapAsyncLoadContext : public FMaterialShaderMap::FAsyncLoadContext
 	{
-		FString	        DataKey;
-		FSharedBuffer   CachedData;
-		FIoHash         CachedDataHash;
-		EShaderPlatform Platform;
-		FString         AssetName;
-		FRequestOwner   RequestOwner{UE::DerivedData::EPriority::Normal};
+		FString	         DataKey;
+		FSharedBuffer    CachedData;
+		FIoHash          CachedDataHash;
+		EShaderPlatform  Platform;
+		const FMaterial* Material = nullptr;
+		FRequestOwner    RequestOwner{UE::DerivedData::EPriority::Normal};
 		TRefCountPtr<FMaterialShaderMap> ShaderMap;
 
 		bool IsReady() const override
@@ -1447,7 +1490,9 @@ TSharedRef<FMaterialShaderMap::FAsyncLoadContext> FMaterialShaderMap::BeginLoadF
 				// Deserialize from the cached data
 				ShaderMap->Serialize(Ar);
 
-				const FString InDataKey = GetMaterialShaderMapKeyString(ShaderMap->GetShaderMapId(), Platform, true);
+				check(Material != nullptr);
+
+				const FString InDataKey = GetMaterialShaderMapKeyString(ShaderMap->GetShaderMapId(), FMaterialShaderParameters(Material), Platform, true);
 
 				if (InDataKey != DataKey)
 				{
@@ -1463,7 +1508,7 @@ TSharedRef<FMaterialShaderMap::FAsyncLoadContext> FMaterialShaderMap::BeginLoadF
 				GShaderCompilerStats->AddDDCHit(1);
 
 				FCacheKey DdcKey = GetMaterialShaderMapKey(DataKey);
-				UE_LOG(LogMaterial, Verbose, TEXT("Loaded shaders for %s from DDC (key hash: %s)"), *AssetName, *LexToString(DdcKey.Hash));
+				UE_LOG(LogMaterial, Verbose, TEXT("Loaded shaders for %s from DDC (key hash: %s)"), *Material->GetAssetName(), *LexToString(DdcKey.Hash));
 			}
 			else
 			{
@@ -1496,8 +1541,7 @@ TSharedRef<FMaterialShaderMap::FAsyncLoadContext> FMaterialShaderMap::BeginLoadF
 			SCOPE_SECONDS_COUNTER(MaterialDDCTime);
 			COOK_STAT(auto Timer = MaterialShaderCookStats::UsageStats.TimeSyncWork());
 			COOK_STAT(Timer.TrackCyclesOnly());
-
-			Result->DataKey = GetMaterialShaderMapKeyString(ShaderMapId, InPlatform);
+			Result->DataKey = GetMaterialShaderMapKeyString(ShaderMapId, FMaterialShaderParameters(Material), InPlatform);
 			FCacheKey CacheKey = GetMaterialShaderMapKey(Result->DataKey);
 			OutDDCKeyDesc = LexToString(CacheKey.Hash);
 
@@ -1553,7 +1597,7 @@ TSharedRef<FMaterialShaderMap::FAsyncLoadContext> FMaterialShaderMap::BeginLoadF
 				FCacheGetValueRequest Request;
 				Request.Name = GetMaterialShaderMapName(Material->GetFullPath(), ShaderMapId, InPlatform);
 				Request.Key = GetMaterialShaderMapKey(Result->DataKey);
-				Result->AssetName = Material->GetAssetName();
+				Result->Material = Material;
 				Result->Platform = InPlatform;
 
 				GetCache().GetValue({Request}, Result->RequestOwner, [Result, bCheckCache](FCacheGetValueResponse&& Response)
@@ -1575,7 +1619,7 @@ TSharedRef<FMaterialShaderMap::FAsyncLoadContext> FMaterialShaderMap::BeginLoadF
 	return Result;
 }
 
-void FMaterialShaderMap::SaveToDerivedDataCache()
+void FMaterialShaderMap::SaveToDerivedDataCache(const FMaterialShaderParameters& ShaderParameters)
 {
 	TRACE_CPUPROFILER_EVENT_SCOPE(FMaterialShaderMap::SaveToDerivedDataCache);
 	COOK_STAT(auto Timer = MaterialShaderCookStats::UsageStats.TimeSyncWork());
@@ -1586,7 +1630,7 @@ void FMaterialShaderMap::SaveToDerivedDataCache()
 	TRACE_COUNTER_ADD(Shaders_FMaterialShaderMapDDCBytesSent, SaveData.Num());
 	COOK_STAT(Timer.AddMiss(SaveData.Num()));
 
-	const FString DataKey = GetMaterialShaderMapKeyString(ShaderMapId, GetShaderPlatform());
+	const FString DataKey = GetMaterialShaderMapKeyString(ShaderMapId, ShaderParameters, GetShaderPlatform());
 
 	using namespace UE::DerivedData;
 	FCachePutValueRequest Request;
@@ -2214,7 +2258,7 @@ void FMaterialShaderMap::Compile(
 #if WITH_EDITOR
 				if (bIsPersistent)
 				{
-					SaveToDerivedDataCache();
+					SaveToDerivedDataCache(FMaterialShaderParameters(Material));
 				}
 #endif
 			}
diff --git a/Engine/Source/Runtime/Engine/Private/Materials/MaterialShared.cpp b/Engine/Source/Runtime/Engine/Private/Materials/MaterialShared.cpp
--- a/Engine/Source/Runtime/Engine/Private/Materials/MaterialShared.cpp
+++ b/Engine/Source/Runtime/Engine/Private/Materials/MaterialShared.cpp
@@ -3273,7 +3273,7 @@ bool FMaterial::Translate_Legacy(const FMaterialShaderMapId& ShaderMapId,
 	FMaterialCompilationOutput& OutCompilationOutput,
 	TRefCountPtr<FSharedShaderCompilerEnvironment>& OutMaterialEnvironment)
 {
-	FString MaterialTranslationDDCKeyString = GetMaterialShaderMapKeyString(ShaderMapId, InPlatform, false);
+	FString MaterialTranslationDDCKeyString = GetMaterialShaderMapKeyString(ShaderMapId, FMaterialShaderParameters(this), InPlatform, false);
 	FHLSLMaterialTranslator MaterialTranslator(this, OutCompilationOutput, InStaticParameters, InPlatform, GetQualityLevel(), ShaderMapId.FeatureLevel, InTargetPlatform, &ShaderMapId.SubstrateCompilationConfig, MoveTemp(MaterialTranslationDDCKeyString));
 	const bool bSuccess = MaterialTranslator.Translate();
 	if (bSuccess)
diff --git a/Engine/Source/Runtime/Engine/Private/ShaderCompiler/ShaderCompiler.cpp b/Engine/Source/Runtime/Engine/Private/ShaderCompiler/ShaderCompiler.cpp
--- a/Engine/Source/Runtime/Engine/Private/ShaderCompiler/ShaderCompiler.cpp
+++ b/Engine/Source/Runtime/Engine/Private/ShaderCompiler/ShaderCompiler.cpp
@@ -6763,6 +6763,11 @@ void FShaderCompilingManager::ProcessCompiledShaderMaps(
 #endif
 			int32 NumIncompleteMaterials = 0;
 			int32 MaterialIndex = 0;
+			
+			// all materials using the same shader map must necessarily have the same set of FMaterialShaderParameters, so just initialize from the first
+			check(MaterialDependencies.Num() > 0);
+			FMaterialShaderParameters ShaderParameters(MaterialDependencies[0]);
+
 			while (MaterialIndex < MaterialDependencies.Num())
 			{
 				FMaterial* Material = MaterialDependencies[MaterialIndex];
@@ -6892,7 +6897,7 @@ void FShaderCompilingManager::ProcessCompiledShaderMaps(
 					ShaderMapToUseForRendering->bCompilationFinalized = true;
 					if (ShaderMapToUseForRendering->bIsPersistent)
 					{
-						ShaderMapToUseForRendering->SaveToDerivedDataCache();
+						ShaderMapToUseForRendering->SaveToDerivedDataCache(ShaderParameters);
 					}
 				}
 
@@ -9585,6 +9590,31 @@ static FString GetGlobalShaderMapKeyString(const FGlobalShaderMapId& ShaderMapId
 	FString ShaderMapKeyString = Format.ToString() + TEXT("_") + FString(FString::FromInt(GetTargetPlatformManagerRef().ShaderFormatVersion(Format))) + TEXT("_");
 	ShaderMapAppendKeyString(Platform, ShaderMapKeyString);
 	ShaderMapId.AppendKeyString(ShaderMapKeyString, Dependencies);
+
+	const EShaderPermutationFlags PermutationFlags = ShaderMapId.GetShaderPermutationFlags();
+
+	// Construct a hash of all the environment modifications applied for each shader type & permutation
+	FMemoryHasherBlake3 Hasher;
+	for (const FShaderTypeDependency& ShaderTypeDep : Dependencies)
+	{
+		const FGlobalShaderType* GlobalShaderType = FindShaderTypeByName(ShaderTypeDep.ShaderTypeName)->AsGlobalShaderType();
+		for (int32 PermutationId = 0; PermutationId < GlobalShaderType->GetPermutationCount(); PermutationId++)
+		{
+			if (GlobalShaderType->ShouldCompilePermutation(Platform, PermutationId, PermutationFlags))
+			{
+				FShaderCompilerEnvironment Env(Hasher);
+				GlobalShaderType->SetupCompileEnvironment(Platform, PermutationId, PermutationFlags, Env);
+				Env.SerializeEverythingButFiles(Hasher);
+			}
+		}
+	}
+
+	// * 2 for hex representation of hash; + 6 for tag/underscores
+	TStringBuilder<sizeof(TCHAR) * (sizeof(FBlake3Hash::ByteArray) * 2 + 6)> EnvHashString;
+	EnvHashString << "_EMH_" << Hasher.Finalize() << "_";
+	check(EnvHashString.GetAllocatedSize() == 0);
+	ShaderMapKeyString.Append(EnvHashString.ToView());
+
 	return FString::Printf(TEXT("%s_%s_%s"), TEXT("GSM"), *GetGlobalShaderMapDDCKey(), *ShaderMapKeyString);
 }
 
diff --git a/Engine/Source/Runtime/Engine/Public/MaterialShaderType.h b/Engine/Source/Runtime/Engine/Public/MaterialShaderType.h
--- a/Engine/Source/Runtime/Engine/Public/MaterialShaderType.h
+++ b/Engine/Source/Runtime/Engine/Public/MaterialShaderType.h
@@ -52,7 +52,17 @@ extern ENGINE_API FString GetShadingModelFieldString(FMaterialShadingModelField
 extern ENGINE_API FString GetBlendModeString(EBlendMode BlendMode);
 
 /** Creates a string key for the derived data cache given a shader map id. */
-extern ENGINE_API FString GetMaterialShaderMapKeyString(const FMaterialShaderMapId& ShaderMapId, EShaderPlatform Platform, bool bIncludeKeyStringShaderDependencies = true);
+UE_DEPRECATED(5.5, "GetMaterialShaderMapKeyString now requires passing an FMaterialShaderParameters instance as input.")
+inline FString GetMaterialShaderMapKeyString(const FMaterialShaderMapId& ShaderMapId, EShaderPlatform Platform, bool bIncludeKeyStringShaderDependencies = true)
+{
+	return TEXT("INVALID_KEY");
+}
+
+extern ENGINE_API FString GetMaterialShaderMapKeyString(
+	const FMaterialShaderMapId& ShaderMapId,
+	const FMaterialShaderParameters& ShaderParameters,
+	EShaderPlatform Platform,
+	bool bIncludeKeyStringShaderDependencies = true);
 
 /** Called for every material shader to update the appropriate stats. */
 extern void UpdateMaterialShaderCompilingStats(const FMaterial* Material);
diff --git a/Engine/Source/Runtime/Engine/Public/MaterialShared.h b/Engine/Source/Runtime/Engine/Public/MaterialShared.h
--- a/Engine/Source/Runtime/Engine/Public/MaterialShared.h
+++ b/Engine/Source/Runtime/Engine/Public/MaterialShared.h
@@ -892,6 +892,88 @@
 	};
 }
 
+struct FMaterialShaderParameters
+{
+	EMaterialDomain MaterialDomain;
+	FMaterialShadingModelField ShadingModels;
+	EBlendMode BlendMode;
+	ERHIFeatureLevel::Type FeatureLevel;
+	EMaterialQualityLevel::Type QualityLevel;
+	int32 BlendableLocation;
+	int32 NumCustomizedUVs;
+	uint32 StencilCompare;
+	union
+	{
+		uint64 PackedFlags;
+		struct
+		{
+			uint64 bIsDefaultMaterial : 1;
+			uint64 bIsSpecialEngineMaterial : 1;
+			uint64 bIsMasked : 1;
+			uint64 bIsDitherMasked : 1;
+			uint64 bIsTwoSided : 1;
+			uint64 bIsThinSurface : 1;
+			uint64 bIsDistorted : 1;
+			uint64 bShouldCastDynamicShadows : 1;
+			uint64 bWritesEveryPixel : 1;
+			uint64 bWritesEveryPixelShadowPass : 1;
+			uint64 bHasDiffuseAlbedoConnected : 1;
+			uint64 bHasF0Connected : 1;
+			uint64 bHasBaseColorConnected : 1;
+			uint64 bHasNormalConnected : 1;
+			uint64 bHasRoughnessConnected : 1;
+			uint64 bHasSpecularConnected : 1;
+			uint64 bHasMetallicConnected : 1;
+			uint64 bHasEmissiveColorConnected : 1;
+			uint64 bHasAmbientOcclusionConnected : 1;
+			uint64 bHasAnisotropyConnected : 1;
+			uint64 bIsTessellationEnabled : 1;
+			uint64 bHasDisplacementConnected : 1;
+			uint64 bHasVertexPositionOffsetConnected : 1;
+			uint64 bHasPixelDepthOffsetConnected : 1;
+			uint64 bMaterialMayModifyMeshPosition : 1;
+			uint64 bIsUsedWithStaticLighting : 1;
+			uint64 bIsUsedWithParticleSprites : 1;
+			uint64 bIsUsedWithMeshParticles : 1;
+			uint64 bIsUsedWithNiagaraSprites : 1;
+			uint64 bIsUsedWithNiagaraMeshParticles : 1;
+			uint64 bIsUsedWithNiagaraRibbons : 1;
+			uint64 bIsUsedWithLandscape : 1;
+			uint64 bIsUsedWithBeamTrails : 1;
+			uint64 bIsUsedWithSplineMeshes : 1;
+			uint64 bIsUsedWithSkeletalMesh : 1;
+			uint64 bIsUsedWithMorphTargets : 1;
+			uint64 bIsUsedWithAPEXCloth : 1;
+			uint64 bIsUsedWithGeometryCache : 1;
+			uint64 bIsUsedWithGeometryCollections : 1;
+			uint64 bIsUsedWithHairStrands : 1;
+			uint64 bIsUsedWithWater : 1;
+			uint64 bIsTranslucencyWritingVelocity : 1;
+			uint64 bIsTranslucencyWritingCustomDepth : 1;
+			uint64 bIsDitheredLODTransition : 1;
+			uint64 bIsUsedWithInstancedStaticMeshes : 1;
+			uint64 bHasPerInstanceCustomData : 1;
+			uint64 bHasPerInstanceRandom : 1;
+			uint64 bHasVertexInterpolator : 1;
+			uint64 bHasRuntimeVirtualTextureOutput : 1;
+			uint64 bIsUsedWithLidarPointCloud : 1;
+			uint64 bIsUsedWithVirtualHeightfieldMesh : 1;
+			uint64 bIsUsedWithNeuralNetworks : 1;
+			uint64 bIsUsedWithNanite : 1;
+			uint64 bIsStencilTestEnabled : 1;
+			uint64 bIsTranslucencySurface : 1;
+			uint64 bShouldDisableDepthTest : 1;
+			uint64 bHasRenderTracePhysicalMaterialOutput : 1;
+			uint64 bIsUsedWithVolumetricCloud : 1;
+			uint64 bIsUsedWithHeterogeneousVolumes : 1;
+			uint64 bIsMobileSeparateTranslucencyEnabled : 1;
+			uint64 bAlwaysEvaluateWorldPositionOffset : 1;
+		};
+	};
+
+	FMaterialShaderParameters(const FMaterial* InMaterial);
+};
+
 /** Contains all the information needed to uniquely identify a FMaterialShaderMap. */
 class FMaterialShaderMapId
 {
@@ -1345,7 +1427,7 @@
 
 #if WITH_EDITOR
 	/** Saves this shader map to the derived data cache. */
-	void SaveToDerivedDataCache();
+	void SaveToDerivedDataCache(const FMaterialShaderParameters& ShaderParameters);
 #endif
 
 	/** Serializes a shader map to an archive (used with recompiling shaders for a remote console) */
@@ -3021,87 +3103,6 @@
 ENGINE_API uint8 GetRayTracingMaskFromMaterial(const EBlendMode BlendMode);
 
 //
-struct FMaterialShaderParameters
-{
-	EMaterialDomain MaterialDomain;
-	FMaterialShadingModelField ShadingModels;
-	EBlendMode BlendMode;
-	ERHIFeatureLevel::Type FeatureLevel;
-	EMaterialQualityLevel::Type QualityLevel;
-	int32 BlendableLocation;
-	int32 NumCustomizedUVs;
-	uint32 StencilCompare;
-	union
-	{
-		uint64 PackedFlags;
-		struct
-		{
-			uint64 bIsDefaultMaterial : 1;
-			uint64 bIsSpecialEngineMaterial : 1;
-			uint64 bIsMasked : 1;
-			uint64 bIsDitherMasked : 1;
-			uint64 bIsTwoSided : 1;
-			uint64 bIsThinSurface : 1;
-			uint64 bIsDistorted : 1;
-			uint64 bShouldCastDynamicShadows : 1;
-			uint64 bWritesEveryPixel : 1;
-			uint64 bWritesEveryPixelShadowPass : 1;
-			uint64 bHasDiffuseAlbedoConnected : 1;
-			uint64 bHasF0Connected : 1;
-			uint64 bHasBaseColorConnected : 1;
-			uint64 bHasNormalConnected : 1;
-			uint64 bHasRoughnessConnected : 1;
-			uint64 bHasSpecularConnected : 1;
-			uint64 bHasMetallicConnected : 1;
-			uint64 bHasEmissiveColorConnected : 1;
-			uint64 bHasAmbientOcclusionConnected : 1;
-			uint64 bHasAnisotropyConnected : 1;
-			uint64 bIsTessellationEnabled : 1;
-			uint64 bHasDisplacementConnected : 1;
-			uint64 bHasVertexPositionOffsetConnected : 1;
-			uint64 bHasPixelDepthOffsetConnected : 1;
-			uint64 bMaterialMayModifyMeshPosition : 1;
-			uint64 bIsUsedWithStaticLighting : 1;
-			uint64 bIsUsedWithParticleSprites : 1;
-			uint64 bIsUsedWithMeshParticles : 1;
-			uint64 bIsUsedWithNiagaraSprites : 1;
-			uint64 bIsUsedWithNiagaraMeshParticles : 1;
-			uint64 bIsUsedWithNiagaraRibbons : 1;
-			uint64 bIsUsedWithLandscape : 1;
-			uint64 bIsUsedWithBeamTrails : 1;
-			uint64 bIsUsedWithSplineMeshes : 1;
-			uint64 bIsUsedWithSkeletalMesh : 1;
-			uint64 bIsUsedWithMorphTargets : 1;
-			uint64 bIsUsedWithAPEXCloth : 1;
-			uint64 bIsUsedWithGeometryCache : 1;
-			uint64 bIsUsedWithGeometryCollections : 1;
-			uint64 bIsUsedWithHairStrands : 1;
-			uint64 bIsUsedWithWater : 1;
-			uint64 bIsTranslucencyWritingVelocity : 1;
-			uint64 bIsTranslucencyWritingCustomDepth : 1;
-			uint64 bIsDitheredLODTransition : 1;
-			uint64 bIsUsedWithInstancedStaticMeshes : 1;
-			uint64 bHasPerInstanceCustomData : 1;
-			uint64 bHasPerInstanceRandom : 1;
-			uint64 bHasVertexInterpolator : 1;
-			uint64 bHasRuntimeVirtualTextureOutput : 1;
-			uint64 bIsUsedWithLidarPointCloud : 1;
-			uint64 bIsUsedWithVirtualHeightfieldMesh : 1;
-			uint64 bIsUsedWithNeuralNetworks : 1;
-			uint64 bIsUsedWithNanite : 1;
-			uint64 bIsStencilTestEnabled : 1;
-			uint64 bIsTranslucencySurface : 1;
-			uint64 bShouldDisableDepthTest : 1;
-			uint64 bHasRenderTracePhysicalMaterialOutput : 1;
-			uint64 bIsUsedWithVolumetricCloud : 1;
-			uint64 bIsUsedWithHeterogeneousVolumes : 1;
-			uint64 bIsMobileSeparateTranslucencyEnabled : 1;
-			uint64 bAlwaysEvaluateWorldPositionOffset : 1;
-		};
-	};
 
-	FMaterialShaderParameters(const FMaterial* InMaterial);
-};
-
 inline bool ShouldIncludeMaterialInDefaultOpaquePass(const FMaterial& Material)
 {
diff --git a/Engine/Source/Runtime/RenderCore/Private/GlobalShader.cpp b/Engine/Source/Runtime/RenderCore/Private/GlobalShader.cpp
--- a/Engine/Source/Runtime/RenderCore/Private/GlobalShader.cpp
+++ b/Engine/Source/Runtime/RenderCore/Private/GlobalShader.cpp
@@ -387,7 +387,7 @@ FGlobalShaderMapId::FGlobalShaderMapId(EShaderPlatform Platform, const ITargetPl
 	IniPlatformName = FName(TargetPlatform->IniPlatformName());
 
 	LayoutParams.InitializeForPlatform(TargetPlatform);
-	const EShaderPermutationFlags PermutationFlags = GetShaderPermutationFlags(LayoutParams);
+	const EShaderPermutationFlags PermutationFlags = GetShaderPermutationFlags();
 	TArray<FShaderType*> ShaderTypes;
 	TArray<const FShaderPipelineType*> ShaderPipelineTypes;
 
@@ -627,12 +627,8 @@ bool FGlobalShaderMap::IsEmpty() const
 
 bool FGlobalShaderMap::IsComplete(const ITargetPlatform* TargetPlatform) const
 {
-	// TODO: store these in the shadermap before it's start to be compiled?
-	FPlatformTypeLayoutParameters LayoutParams;
-	LayoutParams.InitializeForPlatform(TargetPlatform);
-	const EShaderPermutationFlags PermutationFlags = GetShaderPermutationFlags(LayoutParams);
-
 	FGlobalShaderMapId ShaderMapId(Platform, TargetPlatform);
+	const EShaderPermutationFlags PermutationFlags = ShaderMapId.GetShaderPermutationFlags();
 
 	// traverse all global shader types
 	for (TLinkedList<FShaderType*>::TIterator ShaderTypeIt(FShaderType::GetTypeList()); ShaderTypeIt; ShaderTypeIt.Next())
diff --git a/Engine/Source/Runtime/RenderCore/Private/ShaderCore.cpp b/Engine/Source/Runtime/RenderCore/Private/ShaderCore.cpp
--- a/Engine/Source/Runtime/RenderCore/Private/ShaderCore.cpp
+++ b/Engine/Source/Runtime/RenderCore/Private/ShaderCore.cpp
@@ -856,6 +856,10 @@ FShaderCompilerEnvironment::FShaderCompilerEnvironment()
 	IncludeVirtualPathToContentsMap.Empty(15);
 }
 
+FShaderCompilerEnvironment::FShaderCompilerEnvironment(FMemoryHasherBlake3& Hasher) : Hasher(&Hasher)
+{
+}
+
 /** Initialization constructor. */
 PRAGMA_DISABLE_DEPRECATION_WARNINGS		// FShaderCompilerDefinitions will be made internal in the future, marked deprecated until then
 FShaderCompilerEnvironment::FShaderCompilerEnvironment(const FShaderCompilerDefinitions& InDefinitions)
@@ -888,6 +892,7 @@ void FShaderCompilerEnvironment::Merge(const FShaderCompilerEnvironment& Other)
 	ResourceTableMap.Append(Other.ResourceTableMap);
 	UniformBufferMap.Append(Other.UniformBufferMap);
 	PRAGMA_DISABLE_DEPRECATION_WARNINGS	// FShaderCompilerDefinitions will be made internal in the future, marked deprecated until then
+	checkf(Definitions.IsValid(), TEXT("Merge is not supported on FShaderCompilerEnvironment in hashing mode"));
 	Definitions->Merge(*Other.Definitions);
 	PRAGMA_ENABLE_DEPRECATION_WARNINGS
 	CompileArgs.Append(Other.CompileArgs);
@@ -897,6 +902,7 @@ void FShaderCompilerEnvironment::Merge(const FShaderCompilerEnvironment& Other)
 
 FString FShaderCompilerEnvironment::GetDefinitionsAsCommentedCode() const
 {
+	checkf(Definitions.IsValid(), TEXT("GetDefinitionsAsCommentedCode is not supported on FShaderCompilerEnvironment in hashing mode"));
 	TArray<FString> DefinesLines;
 	PRAGMA_DISABLE_DEPRECATION_WARNINGS // FShaderCompilerDefinitions will be made internal in the future, marked deprecated until then
 	DefinesLines.Reserve(Definitions->Num());
@@ -919,40 +925,246 @@ FString FShaderCompilerEnvironment::GetDefinitionsAsCommentedCode() const
 PRAGMA_DISABLE_DEPRECATION_WARNINGS		// FShaderCompilerDefinitions will be made internal in the future, marked deprecated until then
 
 // Pass through functions to definitions
-void FShaderCompilerEnvironment::SetDefine(const TCHAR* Name, const TCHAR* Value)	{ Definitions->SetDefine(Name, Value); }
-void FShaderCompilerEnvironment::SetDefine(const TCHAR* Name, const FString& Value)	{ Definitions->SetDefine(Name, Value); }
-void FShaderCompilerEnvironment::SetDefine(const TCHAR* Name, uint32 Value)			{ Definitions->SetDefine(Name, Value); }
-void FShaderCompilerEnvironment::SetDefine(const TCHAR* Name, int32 Value)			{ Definitions->SetDefine(Name, Value); }
-void FShaderCompilerEnvironment::SetDefine(const TCHAR* Name, bool Value)			{ Definitions->SetDefine(Name, Value); }
-void FShaderCompilerEnvironment::SetDefine(const TCHAR* Name, float Value)			{ Definitions->SetDefine(Name, Value); }
-
-void FShaderCompilerEnvironment::SetDefine(FName Name, const TCHAR* Value)		{ Definitions->SetDefine(Name, Value); }
-void FShaderCompilerEnvironment::SetDefine(FName Name, const FString& Value)	{ Definitions->SetDefine(Name, Value); }
-void FShaderCompilerEnvironment::SetDefine(FName Name, uint32 Value)			{ Definitions->SetDefine(Name, Value); }
-void FShaderCompilerEnvironment::SetDefine(FName Name, int32 Value)				{ Definitions->SetDefine(Name, Value); }
-void FShaderCompilerEnvironment::SetDefine(FName Name, bool Value)				{ Definitions->SetDefine(Name, Value); }
-void FShaderCompilerEnvironment::SetDefine(FName Name, float Value)				{ Definitions->SetDefine(Name, Value); }
-
-void FShaderCompilerEnvironment::SetDefine(FShaderCompilerDefineNameCache& Name, const TCHAR* Value)	{ Definitions->SetDefine(Name, Value); }
-void FShaderCompilerEnvironment::SetDefine(FShaderCompilerDefineNameCache& Name, const FString& Value)	{ Definitions->SetDefine(Name, Value); }
-void FShaderCompilerEnvironment::SetDefine(FShaderCompilerDefineNameCache& Name, uint32 Value)			{ Definitions->SetDefine(Name, Value); }
-void FShaderCompilerEnvironment::SetDefine(FShaderCompilerDefineNameCache& Name, int32 Value)			{ Definitions->SetDefine(Name, Value); }
-void FShaderCompilerEnvironment::SetDefine(FShaderCompilerDefineNameCache& Name, bool Value)			{ Definitions->SetDefine(Name, Value); }
-void FShaderCompilerEnvironment::SetDefine(FShaderCompilerDefineNameCache& Name, float Value)			{ Definitions->SetDefine(Name, Value); }
+void FShaderCompilerEnvironment::SetDefine(const TCHAR* Name, const TCHAR* Value)
+{
+	if (Definitions.IsValid())
+	{
+		Definitions->SetDefine(Name, Value);
+	}
+	else
+	{
+		Hasher->Serialize(const_cast<TCHAR*>(Name), FCString::Strlen(Name));
+		Hasher->Serialize(const_cast<TCHAR*>(Value), FCString::Strlen(Value));
+	}
+}
+
+void FShaderCompilerEnvironment::SetDefine(const TCHAR* Name, const FString& Value)
+{ 
+	if (Definitions.IsValid())
+	{
+		Definitions->SetDefine(Name, Value);
+	}
+	else
+	{
+		Hasher->Serialize(const_cast<TCHAR*>(Name), FCString::Strlen(Name));
+		*Hasher << const_cast<FString&>(Value);
+	}
+}
+
+void FShaderCompilerEnvironment::SetDefine(const TCHAR* Name, uint32 Value)
+{
+	if (Definitions.IsValid())
+	{
+		Definitions->SetDefine(Name, Value);
+	}
+	else
+	{
+		Hasher->Serialize(const_cast<TCHAR*>(Name), FCString::Strlen(Name));
+		*Hasher << Value;
+	}
+}
+void FShaderCompilerEnvironment::SetDefine(const TCHAR* Name, int32 Value)
+{
+	if (Definitions.IsValid())
+	{
+		Definitions->SetDefine(Name, Value);
+	}
+	else
+	{
+		Hasher->Serialize(const_cast<TCHAR*>(Name), FCString::Strlen(Name));
+		*Hasher << Value;
+	}
+}
+void FShaderCompilerEnvironment::SetDefine(const TCHAR* Name, bool Value)
+{
+	if (Definitions.IsValid())
+	{
+		Definitions->SetDefine(Name, Value);
+	}
+	else
+	{
+		Hasher->Serialize(const_cast<TCHAR*>(Name), FCString::Strlen(Name));
+		*Hasher << Value;
+	}
+}
+void FShaderCompilerEnvironment::SetDefine(const TCHAR* Name, float Value)
+{
+	if (Definitions.IsValid())
+	{
+		Definitions->SetDefine(Name, Value);
+	}
+	else
+	{
+		Hasher->Serialize(const_cast<TCHAR*>(Name), FCString::Strlen(Name));
+		*Hasher << Value;
+	}
+}
+
+void FShaderCompilerEnvironment::SetDefine(FName Name, const TCHAR* Value)
+{
+	if (Definitions.IsValid())
+	{
+		Definitions->SetDefine(Name, Value);
+	}
+	else
+	{
+		SetDefine(*Name.ToString(), Value);
+	}
+}
+
+void FShaderCompilerEnvironment::SetDefine(FName Name, const FString& Value)
+{
+	if (Definitions.IsValid())
+	{
+		Definitions->SetDefine(Name, Value);
+	}
+	else
+	{
+		SetDefine(*Name.ToString(), Value);
+	}
+}
+
+void FShaderCompilerEnvironment::SetDefine(FName Name, uint32 Value)
+{
+	if (Definitions.IsValid())
+	{
+		Definitions->SetDefine(Name, Value);
+	}
+	else
+	{
+		SetDefine(*Name.ToString(), Value);
+	}
+}
+
+void FShaderCompilerEnvironment::SetDefine(FName Name, int32 Value)
+{
+	if (Definitions.IsValid())
+	{
+		Definitions->SetDefine(Name, Value);
+	}
+	else
+	{
+		SetDefine(*Name.ToString(), Value);
+	}
+}
+
+void FShaderCompilerEnvironment::SetDefine(FName Name, bool Value)
+{
+	if (Definitions.IsValid())
+	{
+		Definitions->SetDefine(Name, Value);
+	}
+	else
+	{
+		SetDefine(*Name.ToString(), Value);
+	}
+}
+
+void FShaderCompilerEnvironment::SetDefine(FName Name, float Value)
+{
+	if (Definitions.IsValid())
+	{
+		Definitions->SetDefine(Name, Value);
+	}
+	else
+	{
+		SetDefine(*Name.ToString(), Value);
+	}
+}
+
+void FShaderCompilerEnvironment::SetDefine(FShaderCompilerDefineNameCache& Name, const TCHAR* Value) 
+{
+	if (Definitions.IsValid())
+	{
+		Definitions->SetDefine(Name, Value);
+	}
+	else
+	{
+		SetDefine(*Name.ToString(), Value);
+	}
+}
+void FShaderCompilerEnvironment::SetDefine(FShaderCompilerDefineNameCache& Name, const FString& Value)
+{
+	if (Definitions.IsValid())
+	{
+		Definitions->SetDefine(Name, Value);
+	}
+	else
+	{
+		SetDefine(*Name.ToString(), Value);
+	}
+}
+void FShaderCompilerEnvironment::SetDefine(FShaderCompilerDefineNameCache& Name, uint32 Value)
+{
+	if (Definitions.IsValid())
+	{
+		Definitions->SetDefine(Name, Value);
+	}
+	else
+	{
+		SetDefine(*Name.ToString(), Value);
+	}
+}
+void FShaderCompilerEnvironment::SetDefine(FShaderCompilerDefineNameCache& Name, int32 Value)
+{
+	if (Definitions.IsValid())
+	{
+		Definitions->SetDefine(Name, Value);
+	}
+	else
+	{
+		SetDefine(*Name.ToString(), Value);
+	}
+}
+void FShaderCompilerEnvironment::SetDefine(FShaderCompilerDefineNameCache& Name, bool Value)
+{
+	if (Definitions.IsValid())
+	{
+		Definitions->SetDefine(Name, Value);
+	}
+	else
+	{
+		SetDefine(*Name.ToString(), Value);
+	}
+}
+void FShaderCompilerEnvironment::SetDefine(FShaderCompilerDefineNameCache& Name, float Value)
+{
+	if (Definitions.IsValid())
+	{
+		Definitions->SetDefine(Name, Value);
+	}
+	else
+	{
+		SetDefine(*Name.ToString(), Value);
+	}
+}
 
 int32 FShaderCompilerEnvironment::GetIntegerValue(FName Name) const
 {
+	checkf(Definitions.IsValid(), TEXT("GetIntegerValue is not supported on FShaderCompilerEnvironment in hashing mode"));
 	return Definitions->GetIntegerValue(Name);
 }
 
 int32 FShaderCompilerEnvironment::GetIntegerValue(FShaderCompilerDefineNameCache& NameCache, int32 ResultIfNotFound) const
 {
+	checkf(Definitions.IsValid(), TEXT("GetIntegerValue is not supported on FShaderCompilerEnvironment in hashing mode"));
 	return Definitions->GetIntegerValue(NameCache, ResultIfNotFound);
 }
 
 bool FShaderCompilerEnvironment::ContainsDefinition(FName Name) const
 {
-	return Definitions->Contains(Name);
+	if (Definitions.IsValid())
+	{
+		return Definitions->Contains(Name);
+	}
+
+	// If we're in hashing mode only, always report "false" for contains definition.
+	// This is only used by SetDefineIfUnset and as such will just have a minor impact of potential over-invalidation
+	// from certain shader types which call the aforementioned function (i.e. they will re-set potentially already set
+	// defines to a new value, generating a slightly different hash for the shader type). There are very few calls to this
+	// at the time of writing and with per-shader DDC this will only serve to force reconstruction of the shadermap so it's 
+	// not significant enough to worry about.
+	return false;
 }
 
 PRAGMA_ENABLE_DEPRECATION_WARNINGS
@@ -960,7 +1172,11 @@ PRAGMA_ENABLE_DEPRECATION_WARNINGS
 /** This "core" serialization is also used for the hashing the compiler job (where files are handled differently). Should stay in sync with the ShaderCompileWorker. */
 void FShaderCompilerEnvironment::SerializeEverythingButFiles(FArchive& Ar)
 {
-	Ar << *Definitions;
+	// If we don't have a definitions object created then we're in hashing mode and the defines were already hashed on set.
+	if (Definitions.IsValid())
+	{
+		Ar << *Definitions;
+	}
 	Ar << CompileArgs;
 	Ar << CompilerFlags;
 	Ar << RenderTargetOutputFormatsMap;
diff --git a/Engine/Source/Runtime/RenderCore/Public/GlobalShader.h b/Engine/Source/Runtime/RenderCore/Public/GlobalShader.h
--- a/Engine/Source/Runtime/RenderCore/Public/GlobalShader.h
+++ b/Engine/Source/Runtime/RenderCore/Public/GlobalShader.h
@@ -52,6 +52,8 @@ public:
 
 	RENDERCORE_API bool WithEditorOnly() const;
 
+	inline EShaderPermutationFlags GetShaderPermutationFlags() const { return ::GetShaderPermutationFlags(LayoutParams); }
+
 private:
 	FPlatformTypeLayoutParameters LayoutParams;
 	EShaderPlatform ShaderPlatform;
diff --git a/Engine/Source/Runtime/RenderCore/Public/ShaderCore.h b/Engine/Source/Runtime/RenderCore/Public/ShaderCore.h
--- a/Engine/Source/Runtime/RenderCore/Public/ShaderCore.h
+++ b/Engine/Source/Runtime/RenderCore/Public/ShaderCore.h
@@ -452,6 +452,9 @@ struct FShaderCompilerEnvironment
 	/** Default constructor. */
 	RENDERCORE_API FShaderCompilerEnvironment();
 
+	/** Constructor used when enviroment is constructed temporarily purely for the purpose of hashing for inclusion in DDC keys. */
+	RENDERCORE_API FShaderCompilerEnvironment(FMemoryHasherBlake3& Hasher);
+
 	/** Initialization constructor. */
 	PRAGMA_DISABLE_DEPRECATION_WARNINGS
 	UE_DEPRECATED(5.4, "FShaderCompilerDefinitions is being made private in the future, do not use this constructor.")
@@ -508,12 +511,15 @@ struct FShaderCompilerEnvironment
 	RENDERCORE_API int32 GetIntegerValue(FName Name) const;
 	RENDERCORE_API int32 GetIntegerValue(FShaderCompilerDefineNameCache& NameCache, int32 ResultIfNotFound = 0) const;
 
+	UE_DEPRECATED(5.5, "ContainsDefinition will be made private in the future and should not be called by downstream code.")
 	RENDERCORE_API bool ContainsDefinition(FName Name) const;
 
 	template <typename ValueType> void SetDefineIfUnset(const TCHAR* Name, ValueType Value)
 	{
 		FName NameKey(Name);
+		PRAGMA_DISABLE_DEPRECATION_WARNINGS
 		if (!ContainsDefinition(NameKey))
+		PRAGMA_ENABLE_DEPRECATION_WARNINGS
 		{
 			SetDefine(NameKey, Value);
 		}
@@ -612,6 +618,8 @@ private:
 	TPimplPtr<FShaderCompilerDefinitions, EPimplPtrMode::DeepCopy> Definitions;
 	PRAGMA_ENABLE_DEPRECATION_WARNINGS
 
+	FMemoryHasherBlake3* Hasher = nullptr;
+
 	TMap<FString, TVariant<bool, float, int32, uint32, FString>> CompileArgs;
 
 	/** Unused data kept around for deprecated FShaderCompilerEnvironment::GetDefinitions call */
diff --git a/Engine/Source/Runtime/RenderCore/Public/ShaderParameterMetadata.h b/Engine/Source/Runtime/RenderCore/Public/ShaderParameterMetadata.h
--- a/Engine/Source/Runtime/RenderCore/Public/ShaderParameterMetadata.h
+++ b/Engine/Source/Runtime/RenderCore/Public/ShaderParameterMetadata.h
@@ -531,6 +531,11 @@ public:
 		: Name(InName), MapIndex(INDEX_NONE)
 	{}
 
+	FString ToString() const 
+	{
+		return Name.ToString();
+	}
+
 	operator FName() const
 	{
 		return Name;
