Visualize Texture:  Support for name of view in addition to ID of view.  In the editor, the ID changes every time the scene capture actor is edited, forcing the user to repeatedly set the ID, which is annoying.  ID support still makes sense for running the game (as opposed to the editor), or for visualizing a specific cube face (the name will only display the first cube face).

#rnx

diff --git a/Engine/Source/Runtime/RenderCore/Private/VisualizeTexture.cpp b/Engine/Source/Runtime/RenderCore/Private/VisualizeTexture.cpp
--- a/Engine/Source/Runtime/RenderCore/Private/VisualizeTexture.cpp
+++ b/Engine/Source/Runtime/RenderCore/Private/VisualizeTexture.cpp
@@ -20,6 +20,7 @@ void FVisualizeTexture::ParseCommands(const TCHAR* Cmd, FOutputDevice &Ar)
 	// Find out what command to do based on first parameter.
 	ECommand Command = ECommand::Unknown;
 	int32 ViewId = 0;
+	FString ViewName;
 	FString ResourceName;
 	TOptional<uint32> ResourceVersion;
 	TOptional<FWildcardString> ResourceListWildCard;
@@ -58,7 +59,16 @@ void FVisualizeTexture::ParseCommands(const TCHAR* Cmd, FOutputDevice &Ar)
 			else
 			{
 				Command = ECommand::SetViewId;
-				ViewId = FCString::Strtoi(&FirstParameter[5], nullptr, 0);
+
+				// Supports view ID or string name of view
+				TCHAR* EndOfNumber;
+				ViewId = FCString::Strtoi(&FirstParameter[5], &EndOfNumber, 0);
+				if (EndOfNumber - &FirstParameter[0] != FirstParameter.Len())
+				{
+					// Didn't parse as a number, treat it as a string
+					ViewId = 0;
+					ViewName = FirstParameter.Right(FirstParameter.Len() - 5);
+				}
 			}
 		}
 		else
@@ -266,6 +276,7 @@ void FVisualizeTexture::ParseCommands(const TCHAR* Cmd, FOutputDevice &Ar)
 	else if (Command == ECommand::SetViewId) //-V547
 	{
 		Requested.ViewUniqueId = ViewId;
+		Requested.ViewName = ViewName;
 	}
 	else
 	{
@@ -356,8 +367,8 @@ void FVisualizeTexture::DisplayHelp(FOutputDevice &Ar)
 	Ar.Logf(TEXT("  BYNAME   = sort pool list by name"));
 	Ar.Logf(TEXT("  BYSIZE   = show pool list by size"));
 	Ar.Logf(TEXT(""));
-	Ar.Logf(TEXT("VisualizeTexture/Vis view=N"));
-	Ar.Logf(TEXT("  Unique ID of view to visualize textures from, \"view=?\" to dump list of available views"));
+	Ar.Logf(TEXT("VisualizeTexture/Vis view=[ID/NAME]"));
+	Ar.Logf(TEXT("  Unique ID or name of view to visualize textures from, \"view=?\" to dump list of available views"));
 	Ar.Logf(TEXT(""));
 }
 
@@ -950,6 +961,13 @@ void FVisualizeTexture::BeginFrameRenderThread()
 	bFoundRequestedView = false;
 }
 
+static bool VisualizeTextureViewNameMatches(const FString& ViewName, const TCHAR* Description)
+{
+	// Description will be of the form "EditorName (FName)" or "Name", with EditorName being user facing.  Match name followed by space or null terminator.
+	int32 ViewNameLen = ViewName.Len();
+	return !FCString::Strnicmp(*ViewName, Description, ViewNameLen) && (Description[ViewNameLen] == ' ' || Description[ViewNameLen] == 0);
+}
+
 void FVisualizeTexture::BeginViewRenderThread(ERHIFeatureLevel::Type InFeatureLevel, int32 UniqueId, const TCHAR* Description, bool bIsSceneCapture)
 {
 	// Only support visualization for views with a unique ID
@@ -969,7 +987,7 @@ void FVisualizeTexture::BeginViewRenderThread(ERHIFeatureLevel::Type InFeatureLe
 
 	ViewDescriptionMap.FindOrAdd(UniqueId) = Description;
 
-	if (Requested.ViewUniqueId == UniqueId)
+	if (!Requested.ViewName.IsEmpty() ? VisualizeTextureViewNameMatches(Requested.ViewName, Description) : Requested.ViewUniqueId == UniqueId)
 	{
 		// Found the specific view we requested
 		bIsRequestedView = true;
diff --git a/Engine/Source/Runtime/RenderCore/Public/VisualizeTexture.h b/Engine/Source/Runtime/RenderCore/Public/VisualizeTexture.h
--- a/Engine/Source/Runtime/RenderCore/Public/VisualizeTexture.h
+++ b/Engine/Source/Runtime/RenderCore/Public/VisualizeTexture.h
@@ -204,6 +204,7 @@ private:
 	struct FRequested
 	{
 		uint32 ViewUniqueId = 0;				// View requested to be visualized -- zero visualizes the last non-scene-capture view
+		FString ViewName;						// Alternately, string name of view to visualize
 		FString Name;
 		TOptional<uint32> Version;
 	} Requested;
