Enable UAV overlap for all UAVs in an RDG pass. If a pass happens to use multiple UAVs pointing at the same resource (but different slices), only one is currently chosen. However, using multiple UAVs that happen to have overlap will break validation (due to multiple begin / end calls). Until that is fixed, just enable it for all UAVs in the pass.

diff --git a/Engine/Source/Runtime/RenderCore/Private/RenderGraphBuilder.cpp b/Engine/Source/Runtime/RenderCore/Private/RenderGraphBuilder.cpp
--- a/Engine/Source/Runtime/RenderCore/Private/RenderGraphBuilder.cpp
+++ b/Engine/Source/Runtime/RenderCore/Private/RenderGraphBuilder.cpp
@@ -10,41 +10,6 @@
 #include "ProfilingDebugging/CsvProfiler.h"
 #include "Async/ParallelFor.h"
 
-#if ENABLE_RHI_VALIDATION
-
-inline void GatherPassUAVsForOverlapValidation(const FRDGPass* Pass, TArray<FRHIUnorderedAccessView*, TInlineAllocator<MaxSimultaneousUAVs, FRDGArrayAllocator>>& OutUAVs)
-{
-	// RHI validation tracking of Begin/EndUAVOverlaps happens on the underlying resource, so we need to be careful about not
-	// passing multiple UAVs that refer to the same resource, otherwise we get double-Begin and double-End validation errors.
-	// Filter UAVs to only those with unique parent resources.
-	TArray<FRDGViewableResource*, TInlineAllocator<MaxSimultaneousUAVs, FRDGArrayAllocator>> UniqueParents;
-	Pass->GetParameters().Enumerate([&](FRDGParameter Parameter)
-	{
-		if (Parameter.IsUAV())
-		{
-			if (FRDGUnorderedAccessViewRef UAV = Parameter.GetAsUAV())
-			{
-				FRDGViewableResource* Parent = UAV->GetParent();
-
-				// Check if we've already seen this parent.
-				bool bFound = false;
-				for (int32 Index = 0; !bFound && Index < UniqueParents.Num(); ++Index)
-				{
-					bFound = UniqueParents[Index] == Parent;
-				}
-
-				if (!bFound)
-				{
-					UniqueParents.Add(Parent);
-					OutUAVs.Add(UAV->GetRHI());
-				}
-			}
-		}
-	});
-}
-
-#endif
-
 struct FParallelPassSet : public FRHICommandListImmediate::FQueuedCommandList
 {
 	FParallelPassSet() = default;
@@ -58,13 +23,7 @@ inline void BeginUAVOverlap(const FRDGPass* Pass, FRHIComputeCommandList& RHICmd
 #if ENABLE_RHI_VALIDATION
 	if (GRHIValidationEnabled)
 	{
-		TArray<FRHIUnorderedAccessView*, TInlineAllocator<MaxSimultaneousUAVs, FRDGArrayAllocator>> UAVs;
-		GatherPassUAVsForOverlapValidation(Pass, UAVs);
-
-		if (UAVs.Num())
-		{
-			RHICmdList.BeginUAVOverlap(UAVs);
-		}
+		RHICmdList.BeginUAVOverlap();
 	}
 #endif
 }
@@ -74,13 +33,7 @@ inline void EndUAVOverlap(const FRDGPass* Pass, FRHIComputeCommandList& RHICmdLi
 #if ENABLE_RHI_VALIDATION
 	if (GRHIValidationEnabled)
 	{
-		TArray<FRHIUnorderedAccessView*, TInlineAllocator<MaxSimultaneousUAVs, FRDGArrayAllocator>> UAVs;
-		GatherPassUAVsForOverlapValidation(Pass, UAVs);
-
-		if (UAVs.Num())
-		{
-			RHICmdList.EndUAVOverlap(UAVs);
-		}
+		RHICmdList.EndUAVOverlap();
 	}
 #endif
 }
