Add a delay of frame count (r.DumpGPU.FrameDelay) for DumpGPU for easier denoiser debug.

diff --git a/Engine/Source/Runtime/RenderCore/Private/DumpGPU.cpp b/Engine/Source/Runtime/RenderCore/Private/DumpGPU.cpp
--- a/Engine/Source/Runtime/RenderCore/Private/DumpGPU.cpp
+++ b/Engine/Source/Runtime/RenderCore/Private/DumpGPU.cpp
@@ -95,6 +95,11 @@ static TAutoConsoleVariable<float> GDumpGPUDelay(
 	TEXT("Delay in seconds before dumping the frame."),
 	ECVF_RenderThreadSafe);
 
+static TAutoConsoleVariable<int32> GDumpGPUFrameDelay(
+	TEXT("r.DumpGPU.FrameDelay"), 0,
+	TEXT("Delay in frames before dumping the frame."),
+	ECVF_RenderThreadSafe);
+
 static TAutoConsoleVariable<int32> GDumpGPUFrameCount(
 	TEXT("r.DumpGPU.FrameCount"), 1,
 	TEXT("Number of consecutive frames to dump (default=1)."),
@@ -2003,6 +2008,7 @@ static FRDGResourceDumpContext* GRDGResourceDumpContext_GameThread = nullptr;
 static FRDGResourceDumpContext* GRDGResourceDumpContext_RenderThread = nullptr;
 
 static float GNextDumpingRemainingTime = -1.0f;
+static int32 GNextDumpingRemaingFrames = -1;
 static FRDGResourceDumpContext* GNextRDGResourceDumpContext = nullptr;
 
 
@@ -2115,6 +2121,11 @@ FString FRDGBuilder::BeginResourceDump(const TCHAR* Cmd)
 		GNextDumpingRemainingTime = GDumpGPUDelay.GetValueOnGameThread();
 		UE_LOG(LogDumpGPU, Display, TEXT("DumpGPU to %s armed for %d frames starting in %f seconds."), *NewResourceDumpContext->DumpingDirectoryPath, NewResourceDumpContext->FrameCount, GNextDumpingRemainingTime);
 	}
+	else if (GDumpGPUFrameDelay.GetValueOnGameThread() > 0)
+	{
+		GNextDumpingRemaingFrames = GDumpGPUFrameDelay.GetValueOnGameThread();
+		UE_LOG(LogDumpGPU, Display, TEXT("DumpGPU to %s armed for %d frames starting in %d frames."), *NewResourceDumpContext->DumpingDirectoryPath, NewResourceDumpContext->FrameCount, GNextDumpingRemaingFrames);
+	}
 	else
 	{
 		UE_LOG(LogDumpGPU, Display, TEXT("Dump to %s armed for %d frames starting next frame."), *NewResourceDumpContext->DumpingDirectoryPath, NewResourceDumpContext->FrameCount);
@@ -2934,6 +2945,16 @@ void TickEndFrame()
 				GNextRDGResourceDumpContext->Start();
 			}
 		}
+		else if (GNextDumpingRemaingFrames > 0)
+		{
+			check(GNextRDGResourceDumpContext);
+			GNextDumpingRemaingFrames -= 1;
+			if (GNextDumpingRemaingFrames <= 0)
+			{
+				GNextDumpingRemaingFrames = -1;
+				GNextRDGResourceDumpContext->Start();
+			}
+		}
 		else
 		{
 			GNextRDGResourceDumpContext->Start();
