Enabling chunk discovery on all build configurations in order to support Zen Streaming on Test configuration, while relying on the runtime check
if (!IsRunningWithPakFile() && ShouldLookForLooseCookedChunks())
to not affect performance in Test and Shipping builds that are using pak files.

diff --git a/Engine/Source/Runtime/RenderCore/Private/ShaderCodeLibrary.cpp b/Engine/Source/Runtime/RenderCore/Private/ShaderCodeLibrary.cpp
--- a/Engine/Source/Runtime/RenderCore/Private/ShaderCodeLibrary.cpp
+++ b/Engine/Source/Runtime/RenderCore/Private/ShaderCodeLibrary.cpp
@@ -59,8 +59,8 @@ ShaderCodeLibrary.cpp: Bound shader state cache implementation.
 // allow introspection (e.g. dumping the contents) for easier debugging
 #define UE_SHADERLIB_WITH_INTROSPECTION			!UE_BUILD_SHIPPING
 
-// In some development-only scenario (e.g. LaunchOn), the library is chunked, but the build isn't pak'd. We need to find the chunk files manually then.
-#define UE_SHADERLIB_SUPPORT_CHUNK_DISCOVERY	(!UE_BUILD_SHIPPING && !UE_BUILD_TEST)
+// Enabled by default for all configurations, runtime check should prevent chuck discovery in non dev builds
+#define UE_SHADERLIB_SUPPORT_CHUNK_DISCOVERY	(1)
 
 DEFINE_LOG_CATEGORY(LogShaderLibrary);
 
