Bump up pipeline cache version as shader frequency addition invalidated the format.

diff --git a/Engine/Source/Runtime/RenderCore/Private/PipelineCacheUtilities.cpp b/Engine/Source/Runtime/RenderCore/Private/PipelineCacheUtilities.cpp
--- a/Engine/Source/Runtime/RenderCore/Private/PipelineCacheUtilities.cpp
+++ b/Engine/Source/Runtime/RenderCore/Private/PipelineCacheUtilities.cpp
@@ -69,8 +69,9 @@ namespace Private
 		{
 			AddingPipelineCacheVersion = 5,
 			AddingDepthBounds = 6,
+			AddedWorkGraphFrequencies = 7,
 
-			Current = AddingDepthBounds
+			Current = AddedWorkGraphFrequencies
 		};
 
 		/** Magic to reject other files */
@@ -602,10 +603,10 @@ bool UE::PipelineCacheUtilities::LoadStablePipelineCacheFile(const FString& File
 		return false;
 	}
 
-	// start restrictive, as the format isn't really forward compatible, nor needs to be
-	if (Header.Version > SupportedHeader.Version)
+	// start restrictive, as the format is neither forward compatible, nor backward compatible
+	if (Header.Version != SupportedHeader.Version)
 	{
-		UE_LOG(LogPipelineCacheUtilities, Warning, TEXT("Rejecting %s, version is too new (%d vs expected %d)."), *Filename, int(Header.Version), int(SupportedHeader.Version));
+		UE_LOG(LogPipelineCacheUtilities, Warning, TEXT("Rejecting %s, version %d is not supported (current supported version is %d)."), *Filename, int(Header.Version), int(SupportedHeader.Version));
 		return false;
 	}
 
