Add ShouldPrecachePermutation to FShaderType to check if shader permutation needs PSO precaching on startup - by default ShouldCompilePermutation is used

EShaderPermutationPrecacheRequest is returned to define if a permutation is required, not required or development only. Only required PSOs are precached.
DevelopmentOnly should be used for debug permutations and won't trigger an ensure when permutation is used at runtime which isn't requested for precaching.
Validation also runs in editor builds.

diff --git a/Engine/Plugins/Compositing/OpenColorIO/Source/OpenColorIO/Public/OpenColorIOShaderType.h b/Engine/Plugins/Compositing/OpenColorIO/Source/OpenColorIO/Public/OpenColorIOShaderType.h
--- a/Engine/Plugins/Compositing/OpenColorIO/Source/OpenColorIO/Public/OpenColorIOShaderType.h
+++ b/Engine/Plugins/Compositing/OpenColorIO/Source/OpenColorIO/Public/OpenColorIOShaderType.h
@@ -72,6 +72,7 @@ public:
 		ConstructSerializedType InConstructSerializedRef,
 		ConstructCompiledType InConstructCompiledRef,
 		ShouldCompilePermutationType InShouldCompilePermutationRef,
+		ShouldPrecachePermutationType InShouldPrecachePermutationRef,
 		GetRayTracingPayloadTypeType InGetRayTracingPayloadTypeRef,
 #if WITH_EDITOR
 		ModifyCompilationEnvironmentType InModifyCompilationEnvironmentRef,
@@ -84,6 +85,7 @@ public:
 			InConstructSerializedRef,
 			InConstructCompiledRef,
 			InShouldCompilePermutationRef,
+			InShouldPrecachePermutationRef,
 			InGetRayTracingPayloadTypeRef,
 #if WITH_EDITOR
 			InModifyCompilationEnvironmentRef,
diff --git a/Engine/Plugins/FX/Niagara/Source/NiagaraShader/Public/NiagaraShaderType.h b/Engine/Plugins/FX/Niagara/Source/NiagaraShader/Public/NiagaraShaderType.h
--- a/Engine/Plugins/FX/Niagara/Source/NiagaraShader/Public/NiagaraShaderType.h
+++ b/Engine/Plugins/FX/Niagara/Source/NiagaraShader/Public/NiagaraShaderType.h
@@ -89,6 +89,7 @@ public:
 		ConstructSerializedType InConstructSerializedRef,
 		ConstructCompiledType InConstructCompiledRef,
 		ShouldCompilePermutationType InShouldCompilePermutationRef,
+		ShouldPrecachePermutationType InShouldPrecachePermutationRef,
 		GetRayTracingPayloadTypeType InGetRayTracingPayloadTypeRef,
 #if WITH_EDITOR
 		ModifyCompilationEnvironmentType InModifyCompilationEnvironmentRef,
@@ -101,6 +102,7 @@ public:
 			InConstructSerializedRef,
 			InConstructCompiledRef,
 			InShouldCompilePermutationRef,
+			InShouldPrecachePermutationRef,
 			InGetRayTracingPayloadTypeRef,
 #if WITH_EDITOR
 			InModifyCompilationEnvironmentRef,
diff --git a/Engine/Plugins/Runtime/ComputeFramework/Source/ComputeFramework/Private/ComputeFramework/ComputeKernelShaderType.h b/Engine/Plugins/Runtime/ComputeFramework/Source/ComputeFramework/Private/ComputeFramework/ComputeKernelShaderType.h
--- a/Engine/Plugins/Runtime/ComputeFramework/Source/ComputeFramework/Private/ComputeFramework/ComputeKernelShaderType.h
+++ b/Engine/Plugins/Runtime/ComputeFramework/Source/ComputeFramework/Private/ComputeFramework/ComputeKernelShaderType.h
@@ -74,6 +74,7 @@ public:
 		ConstructSerializedType InConstructSerializedRef,
 		ConstructCompiledType InConstructCompiledRef,
 		ShouldCompilePermutationType InShouldCompilePermutationRef,
+		ShouldPrecachePermutationType InShouldPrecachePermutationRef,
 		GetRayTracingPayloadTypeType InGetRayTracingPayloadTypeRef,
 #if WITH_EDITOR
 		ModifyCompilationEnvironmentType InModifyCompilationEnvironmentRef,
@@ -93,6 +94,7 @@ public:
 			InConstructSerializedRef,
 			InConstructCompiledRef,
 			InShouldCompilePermutationRef,
+			InShouldPrecachePermutationRef,
 			InGetRayTracingPayloadTypeRef,
 #if WITH_EDITOR
 			InModifyCompilationEnvironmentRef,
diff --git a/Engine/Source/Runtime/Engine/Private/ShaderCompiler/ShaderCompiler.cpp b/Engine/Source/Runtime/Engine/Private/ShaderCompiler/ShaderCompiler.cpp
--- a/Engine/Source/Runtime/Engine/Private/ShaderCompiler/ShaderCompiler.cpp
+++ b/Engine/Source/Runtime/Engine/Private/ShaderCompiler/ShaderCompiler.cpp
@@ -9542,42 +9542,53 @@ void VerifyGlobalShaders(EShaderPlatform Platform, bool bLoadedFromCacheFile, co
 
 void PrecacheComputePipelineStatesForGlobalShaders(EShaderPlatform Platform, const ITargetPlatform* TargetPlatform)
 {
-	static IConsoleVariable* PrecacheGlobalComputeShadersCVar = IConsoleManager::Get().FindConsoleVariable(TEXT("r.PSOPrecache.GlobalComputeShaders"));
-	if (!PipelineStateCache::IsPSOPrecachingEnabled() || PrecacheGlobalComputeShadersCVar == nullptr || PrecacheGlobalComputeShadersCVar->GetInt() == 0)
+	if (!GRHISupportsPSOPrecaching)
 	{
 		return;
 	}
 
+	TRACE_CPUPROFILER_EVENT_SCOPE(PrecacheComputePipelineStatesForGlobalShaders);
+
 	FPlatformTypeLayoutParameters LayoutParams;
 	LayoutParams.InitializeForPlatform(TargetPlatform);
 	EShaderPermutationFlags PermutationFlags = GetShaderPermutationFlags(LayoutParams);
 
 	FGlobalShaderMap* GlobalShaderMap = GetGlobalShaderMap(Platform);
-
+	
 	// some RHIs (OpenGL) can only create shaders on the Render thread. Queue the creation instead of doing it here.
 	TArray<TShaderRef<FShader>> ComputeShadersToPrecache;
 	for (TLinkedList<FShaderType*>::TIterator ShaderTypeIt(FShaderType::GetTypeList()); ShaderTypeIt; ShaderTypeIt.Next())
 	{
 		FGlobalShaderType* GlobalShaderType = ShaderTypeIt->GetGlobalShaderType();
-		if (!GlobalShaderType)
+		if (!GlobalShaderType || GlobalShaderType->GetFrequency() != SF_Compute)
 		{
 			continue;
 		}
 
+		int32 ShaderPermutationPerGlobalShaderType = 0;
 		for (int32 PermutationId = 0; PermutationId < GlobalShaderType->GetPermutationCount(); PermutationId++)
 		{
-			if (GlobalShaderType->ShouldCompilePermutation(Platform, PermutationId, PermutationFlags))
+			if (GlobalShaderType->ShouldCompilePermutation(Platform, PermutationId, PermutationFlags) &&
+				GlobalShaderType->ShouldPrecachePermutation(Platform, PermutationId, PermutationFlags) == EShaderPermutationPrecacheRequest::Required)
 			{
 				TShaderRef<FShader> GlobalShader = GlobalShaderMap->GetShader(GlobalShaderType, PermutationId);
-				if (GlobalShader.IsValid() && GlobalShader->GetFrequency() == SF_Compute)
+				if (GlobalShader.IsValid())
 				{
 					ComputeShadersToPrecache.Add(GlobalShader);
+					ShaderPermutationPerGlobalShaderType++;
 				}
 			}
 		}
+
+		int32 PermutationCountLimit = 300;
+		ensureMsgf(
+			ShaderPermutationPerGlobalShaderType < PermutationCountLimit,
+			TEXT("Global shader %s has %i permutations to precache: probably more than it needs."),
+			GlobalShaderType->GetName(), ShaderPermutationPerGlobalShaderType);
 	}
 
-	if (ComputeShadersToPrecache.Num() > 0)
+	static IConsoleVariable* PrecacheGlobalComputeShadersCVar = IConsoleManager::Get().FindConsoleVariable(TEXT("r.PSOPrecache.GlobalComputeShaders"));
+	if (PipelineStateCache::IsPSOPrecachingEnabled() && PrecacheGlobalComputeShadersCVar != nullptr && PrecacheGlobalComputeShadersCVar->GetInt() > 0 && ComputeShadersToPrecache.Num() > 0)
 	{
 		ENQUEUE_RENDER_COMMAND(PrecachePSOsForGlobalShaders)(
 			[ComputeShadersToPrecache](FRHICommandListImmediate& RHICmdList)
@@ -9586,6 +9597,7 @@ void PrecacheComputePipelineStatesForGlobalShaders(EShaderPlatform Platform, con
 				{
 					const TCHAR* TypeName = GlobalShader.GetType()->GetName();
 					FRHIComputeShader* RHIComputeShader = GlobalShader.GetComputeShader();
+					check(RHIComputeShader);
 					PipelineStateCache::PrecacheComputePipelineState(RHIComputeShader, TypeName);
 				}
 			});
diff --git a/Engine/Source/Runtime/Engine/Public/MaterialShaderType.h b/Engine/Source/Runtime/Engine/Public/MaterialShaderType.h
--- a/Engine/Source/Runtime/Engine/Public/MaterialShaderType.h
+++ b/Engine/Source/Runtime/Engine/Public/MaterialShaderType.h
@@ -119,6 +119,7 @@ public:
 		ConstructSerializedType InConstructSerializedRef,
 		ConstructCompiledType InConstructCompiledRef,
 		ShouldCompilePermutationType InShouldCompilePermutationRef,
+		ShouldPrecachePermutationType InShouldPrecachePermutationRef,
 		GetRayTracingPayloadTypeType InGetRayTracingPayloadTypeRef,
 #if WITH_EDITOR
 		ModifyCompilationEnvironmentType InModifyCompilationEnvironmentRef,
@@ -131,6 +132,7 @@ public:
 			InConstructSerializedRef,
 			InConstructCompiledRef,
 			InShouldCompilePermutationRef,
+			InShouldPrecachePermutationRef,
 			InGetRayTracingPayloadTypeRef,
 #if WITH_EDITOR
 			InModifyCompilationEnvironmentRef,
diff --git a/Engine/Source/Runtime/Engine/Public/MeshMaterialShaderType.h b/Engine/Source/Runtime/Engine/Public/MeshMaterialShaderType.h
--- a/Engine/Source/Runtime/Engine/Public/MeshMaterialShaderType.h
+++ b/Engine/Source/Runtime/Engine/Public/MeshMaterialShaderType.h
@@ -53,6 +53,7 @@ public:
 		ConstructSerializedType InConstructSerializedRef,
 		ConstructCompiledType InConstructCompiledRef,
 		ShouldCompilePermutationType InShouldCompilePermutationRef,
+		ShouldPrecachePermutationType InShouldPrecachePermutationRef,
 		GetRayTracingPayloadTypeType InGetRayTracingPayloadTypeRef,
 #if WITH_EDITOR
 		ModifyCompilationEnvironmentType InModifyCompilationEnvironmentRef,
@@ -65,6 +66,7 @@ public:
 			InConstructSerializedRef,
 			InConstructCompiledRef,
 			InShouldCompilePermutationRef,
+			InShouldPrecachePermutationRef,
 			InGetRayTracingPayloadTypeRef,
 #if WITH_EDITOR
 			InModifyCompilationEnvironmentRef,
diff --git a/Engine/Source/Runtime/RenderCore/Private/Shader.cpp b/Engine/Source/Runtime/RenderCore/Private/Shader.cpp
--- a/Engine/Source/Runtime/RenderCore/Private/Shader.cpp
+++ b/Engine/Source/Runtime/RenderCore/Private/Shader.cpp
@@ -239,6 +239,7 @@ FShaderType::FShaderType(
 	ConstructSerializedType InConstructSerializedRef,
 	ConstructCompiledType InConstructCompiledRef,
 	ShouldCompilePermutationType InShouldCompilePermutationRef,
+	ShouldPrecachePermutationType InShouldPrecachePermutationRef,
 	GetRayTracingPayloadTypeType InGetRayTracingPayloadTypeRef,
 #if WITH_EDITOR
 	ModifyCompilationEnvironmentType InModifyCompilationEnvironmentRef,
@@ -261,6 +262,7 @@ FShaderType::FShaderType(
 	ConstructSerializedRef(InConstructSerializedRef),
 	ConstructCompiledRef(InConstructCompiledRef),
 	ShouldCompilePermutationRef(InShouldCompilePermutationRef),
+	ShouldPrecachePermutationRef(InShouldPrecachePermutationRef),
 	GetRayTracingPayloadTypeRef(InGetRayTracingPayloadTypeRef),
 #if WITH_EDITOR
 	ModifyCompilationEnvironmentRef(InModifyCompilationEnvironmentRef),
@@ -422,6 +424,11 @@ bool FShaderType::ShouldCompilePermutation(const FShaderPermutationParameters& P
 	return ShouldCompileShaderFrequency((EShaderFrequency)Frequency, Parameters.Platform) && (*ShouldCompilePermutationRef)(Parameters);
 }
 
+EShaderPermutationPrecacheRequest FShaderType::ShouldPrecachePermutation(const FShaderPermutationParameters& Parameters) const
+{
+	return ShouldCompileShaderFrequency((EShaderFrequency)Frequency, Parameters.Platform) ? (*ShouldPrecachePermutationRef)(Parameters) : EShaderPermutationPrecacheRequest::NotRequired;
+}
+
 #if WITH_EDITOR
 
 void FShaderType::ModifyCompilationEnvironment(const FShaderPermutationParameters& Parameters, FShaderCompilerEnvironment& OutEnvironment) const
@@ -1149,6 +1156,24 @@ bool FShaderPipelineType::ShouldCompilePermutation(const FShaderPermutationParam
 	return true;
 }
 
+EShaderPermutationPrecacheRequest FShaderPipelineType::ShouldPrecachePermutation(const FShaderPermutationParameters& Parameters) const
+{
+	EShaderPermutationPrecacheRequest Result = EShaderPermutationPrecacheRequest::NotRequired;
+	for (const FShaderType* ShaderType : Stages)
+	{
+		EShaderPermutationPrecacheRequest ShaderTypeRequest = ShaderType->ShouldPrecachePermutation(Parameters);
+		if (ShaderTypeRequest == EShaderPermutationPrecacheRequest::Required)
+		{
+			return ShaderTypeRequest;
+		}
+		else if (ShaderTypeRequest == EShaderPermutationPrecacheRequest::DevelopmentOnly)
+		{
+			Result = ShaderTypeRequest;
+		}
+	}
+	return Result;
+}
+
 void FShaderPipeline::AddShader(FShader* Shader, int32 PermutationId)
 {
 	const EShaderFrequency Frequency = Shader->GetFrequency();
diff --git a/Engine/Source/Runtime/RenderCore/Public/GlobalShader.h b/Engine/Source/Runtime/RenderCore/Public/GlobalShader.h
--- a/Engine/Source/Runtime/RenderCore/Public/GlobalShader.h
+++ b/Engine/Source/Runtime/RenderCore/Public/GlobalShader.h
@@ -97,6 +97,7 @@ public:
 		ConstructSerializedType InConstructSerializedRef,
 		ConstructCompiledType InConstructCompiledRef,
 		ShouldCompilePermutationType InShouldCompilePermutationRef,
+		ShouldPrecachePermutationType InShouldPrecachePermutationRef,
 		GetRayTracingPayloadTypeType InGetRayTracingPayloadTypeRef,
 #if WITH_EDITOR
 		ModifyCompilationEnvironmentType InModifyCompilationEnvironmentRef,
@@ -109,6 +110,7 @@ public:
 			InConstructSerializedRef,
 			InConstructCompiledRef,
 			InShouldCompilePermutationRef,
+			InShouldPrecachePermutationRef,
 			InGetRayTracingPayloadTypeRef,
 #if WITH_EDITOR
 			InModifyCompilationEnvironmentRef,
@@ -134,6 +136,11 @@ public:
 		return FShaderType::ShouldCompilePermutation(FGlobalShaderPermutationParameters(GetFName(), Platform, PermutationId, Flags));
 	}
 
+	EShaderPermutationPrecacheRequest ShouldPrecachePermutation(EShaderPlatform Platform, int32 PermutationId, EShaderPermutationFlags Flags) const
+	{
+		return FShaderType::ShouldPrecachePermutation(FGlobalShaderPermutationParameters(GetFName(), Platform, PermutationId, Flags));
+	}
+
 #if WITH_EDITOR
 	/**
 	 * Sets up the environment used to compile an instance of this shader type.
@@ -211,6 +218,11 @@ public:
 	{
 		TShaderRef<FShader> Shader = GetShader(&ShaderType::GetStaticType(), PermutationId);
 		checkf(Shader.IsValid(), TEXT("Failed to find shader type %s in Platform %s"), ShaderType::GetStaticType().GetName(), *LegacyShaderPlatformToShaderFormat(Platform).ToString());
+
+		// Validate that the permutation has been precached
+		ensureMsgf(ShaderType::ShouldPrecachePermutation(FGlobalShaderPermutationParameters(ShaderType::GetStaticType().GetFName(), Platform, PermutationId)) != EShaderPermutationPrecacheRequest::NotRequired,
+			TEXT("Using a global shader permutation of %s which hasn't been requested for precaching at runtime. Check the implementation of ShouldPrecachePermutation on the global shader and make sure all required or development only permutations are requested for PSO precaching"), ShaderType::GetStaticType().GetName());
+
 		return TShaderRef<ShaderType>::Cast(Shader);
 	}
 
diff --git a/Engine/Source/Runtime/RenderCore/Public/Shader.h b/Engine/Source/Runtime/RenderCore/Public/Shader.h
--- a/Engine/Source/Runtime/RenderCore/Public/Shader.h
+++ b/Engine/Source/Runtime/RenderCore/Public/Shader.h
@@ -805,6 +805,13 @@ enum class EShaderPermutationFlags : uint32
 };
 ENUM_CLASS_FLAGS(EShaderPermutationFlags);
 
+enum class EShaderPermutationPrecacheRequest : uint8
+{
+	Required,
+	DevelopmentOnly,
+	NotRequired,
+};
+
 RENDERCORE_API EShaderPermutationFlags GetShaderPermutationFlags(const FPlatformTypeLayoutParameters& LayoutParams);
 
 struct FShaderPermutationParameters
@@ -867,6 +874,9 @@ public:
 	/** Can be overridden by FShader subclasses to determine whether a specific permutation should be compiled. */
 	static bool ShouldCompilePermutation(const FShaderPermutationParameters&) { return true; }
 
+	/** Can be overridden by FShader subclasses to determine whether a specific permutation should be precached. */
+	static EShaderPermutationPrecacheRequest ShouldPrecachePermutation(const FShaderPermutationParameters& Parameters) { return ShouldCompilePermutation(Parameters) ? EShaderPermutationPrecacheRequest::Required : EShaderPermutationPrecacheRequest::NotRequired; }
+
 	/** Can be overridden by FShader subclasses to determine whether compilation is valid. */
 	static bool ValidateCompiledResult(EShaderPlatform InPlatform, const FShaderParameterMap& InParameterMap, TArray<FString>& OutError) { return true; }
 
@@ -1225,6 +1235,7 @@ public:
 	typedef class FShader* (*ConstructSerializedType)();
 	typedef FShader* (*ConstructCompiledType)(const FShader::CompiledShaderInitializerType& Initializer);
 	typedef bool (*ShouldCompilePermutationType)(const FShaderPermutationParameters&);
+	typedef EShaderPermutationPrecacheRequest (*ShouldPrecachePermutationType)(const FShaderPermutationParameters&);
 	typedef ERayTracingPayloadType(*GetRayTracingPayloadTypeType)(const int32 PermutationId);
 #if WITH_EDITOR
 	typedef void (*ModifyCompilationEnvironmentType)(const FShaderPermutationParameters&, FShaderCompilerEnvironment&);
@@ -1260,6 +1271,7 @@ public:
 		ConstructSerializedType InConstructSerializedRef,
 		ConstructCompiledType InConstructCompiledRef,
 		ShouldCompilePermutationType InShouldCompilePermutationRef,
+		ShouldPrecachePermutationType InShouldPrecachePermutationRef,
 		GetRayTracingPayloadTypeType InGetRayTracingPayloadTypeRef,
 #if WITH_EDITOR
 		ModifyCompilationEnvironmentType InModifyCompilationEnvironmentRef,
@@ -1276,6 +1288,7 @@ public:
 	RENDERCORE_API FShader* ConstructCompiled(const FShader::CompiledShaderInitializerType& Initializer) const;
 
 	RENDERCORE_API bool ShouldCompilePermutation(const FShaderPermutationParameters& Parameters) const;
+	RENDERCORE_API EShaderPermutationPrecacheRequest ShouldPrecachePermutation(const FShaderPermutationParameters& Parameters) const;
 
 #if WITH_EDITOR
 	RENDERCORE_API void ModifyCompilationEnvironment(const FShaderPermutationParameters& Parameters, FShaderCompilerEnvironment& OutEnvironment) const;
@@ -1482,6 +1495,7 @@ private:
 	ConstructSerializedType ConstructSerializedRef;
 	ConstructCompiledType ConstructCompiledRef;
 	ShouldCompilePermutationType ShouldCompilePermutationRef;
+	ShouldPrecachePermutationType ShouldPrecachePermutationRef;
 	GetRayTracingPayloadTypeType GetRayTracingPayloadTypeRef;
 #if WITH_EDITOR
 	ModifyCompilationEnvironmentType ModifyCompilationEnvironmentRef;
@@ -1578,6 +1592,8 @@ struct FShaderCompiledShaderInitializerType
 	{ return new ShaderClass(static_cast<const typename ShaderMetaType::CompiledShaderInitializerType&>(Initializer)); }\
 	static bool ShouldCompilePermutationImpl(const FShaderPermutationParameters& Parameters) \
 	{ return ShaderClass::ShouldCompilePermutation(static_cast<const typename ShaderClass::FPermutationParameters&>(Parameters)); } \
+	static EShaderPermutationPrecacheRequest ShouldPrecachePermutationImpl(const FShaderPermutationParameters& Parameters) \
+	{ return ShaderClass::ShouldPrecachePermutation(static_cast<const typename ShaderClass::FPermutationParameters&>(Parameters)); } \
 	SHADER_DECLARE_EDITOR_VTABLE(ShaderClass)
 
 
@@ -1624,6 +1640,7 @@ struct FShaderCompiledShaderInitializerType
 	ShaderClass::ConstructSerializedInstance, \
 	ShaderClass::ConstructCompiledInstance, \
 	ShaderClass::ShouldCompilePermutationImpl, \
+	ShaderClass::ShouldPrecachePermutationImpl, \
 	ShaderClass::GetRayTracingPayloadType \
 	SHADER_TYPE_EDITOR_VTABLE(ShaderClass)
 
@@ -1800,6 +1817,8 @@ public:
 
 	RENDERCORE_API bool ShouldCompilePermutation(const FShaderPermutationParameters& Parameters) const;
 
+	RENDERCORE_API EShaderPermutationPrecacheRequest ShouldPrecachePermutation(const FShaderPermutationParameters& Parameters) const;
+
 protected:
 	const TCHAR* const Name;
 	FName TypeName;
