Rename EShaderPermutationPrecacheRequest members and add comments to better describe the different return values

diff --git a/Engine/Source/Runtime/Engine/Private/ShaderCompiler/ShaderCompiler.cpp b/Engine/Source/Runtime/Engine/Private/ShaderCompiler/ShaderCompiler.cpp
--- a/Engine/Source/Runtime/Engine/Private/ShaderCompiler/ShaderCompiler.cpp
+++ b/Engine/Source/Runtime/Engine/Private/ShaderCompiler/ShaderCompiler.cpp
@@ -9569,7 +9569,7 @@ void PrecacheComputePipelineStatesForGlobalShaders(EShaderPlatform Platform, con
 		for (int32 PermutationId = 0; PermutationId < GlobalShaderType->GetPermutationCount(); PermutationId++)
 		{
 			if (GlobalShaderType->ShouldCompilePermutation(Platform, PermutationId, PermutationFlags) &&
-				GlobalShaderType->ShouldPrecachePermutation(Platform, PermutationId, PermutationFlags) == EShaderPermutationPrecacheRequest::Required)
+				GlobalShaderType->ShouldPrecachePermutation(Platform, PermutationId, PermutationFlags) == EShaderPermutationPrecacheRequest::Precached)
 			{
 				TShaderRef<FShader> GlobalShader = GlobalShaderMap->GetShader(GlobalShaderType, PermutationId);
 				if (GlobalShader.IsValid())
diff --git a/Engine/Source/Runtime/RenderCore/Private/Shader.cpp b/Engine/Source/Runtime/RenderCore/Private/Shader.cpp
--- a/Engine/Source/Runtime/RenderCore/Private/Shader.cpp
+++ b/Engine/Source/Runtime/RenderCore/Private/Shader.cpp
@@ -426,7 +426,7 @@ bool FShaderType::ShouldCompilePermutation(const FShaderPermutationParameters& P
 
 EShaderPermutationPrecacheRequest FShaderType::ShouldPrecachePermutation(const FShaderPermutationParameters& Parameters) const
 {
-	return ShouldCompileShaderFrequency((EShaderFrequency)Frequency, Parameters.Platform) ? (*ShouldPrecachePermutationRef)(Parameters) : EShaderPermutationPrecacheRequest::NotRequired;
+	return ShouldCompileShaderFrequency((EShaderFrequency)Frequency, Parameters.Platform) ? (*ShouldPrecachePermutationRef)(Parameters) : EShaderPermutationPrecacheRequest::NotUsed;
 }
 
 #if WITH_EDITOR
@@ -1158,15 +1158,15 @@ bool FShaderPipelineType::ShouldCompilePermutation(const FShaderPermutationParam
 
 EShaderPermutationPrecacheRequest FShaderPipelineType::ShouldPrecachePermutation(const FShaderPermutationParameters& Parameters) const
 {
-	EShaderPermutationPrecacheRequest Result = EShaderPermutationPrecacheRequest::NotRequired;
+	EShaderPermutationPrecacheRequest Result = EShaderPermutationPrecacheRequest::NotUsed;
 	for (const FShaderType* ShaderType : Stages)
 	{
 		EShaderPermutationPrecacheRequest ShaderTypeRequest = ShaderType->ShouldPrecachePermutation(Parameters);
-		if (ShaderTypeRequest == EShaderPermutationPrecacheRequest::Required)
+		if (ShaderTypeRequest == EShaderPermutationPrecacheRequest::Precached)
 		{
 			return ShaderTypeRequest;
 		}
-		else if (ShaderTypeRequest == EShaderPermutationPrecacheRequest::DevelopmentOnly)
+		else if (ShaderTypeRequest == EShaderPermutationPrecacheRequest::NotPrecached)
 		{
 			Result = ShaderTypeRequest;
 		}
diff --git a/Engine/Source/Runtime/RenderCore/Public/GlobalShader.h b/Engine/Source/Runtime/RenderCore/Public/GlobalShader.h
--- a/Engine/Source/Runtime/RenderCore/Public/GlobalShader.h
+++ b/Engine/Source/Runtime/RenderCore/Public/GlobalShader.h
@@ -220,7 +220,7 @@ public:
 		checkf(Shader.IsValid(), TEXT("Failed to find shader type %s in Platform %s"), ShaderType::GetStaticType().GetName(), *LegacyShaderPlatformToShaderFormat(Platform).ToString());
 
 		// Validate that the permutation has been precached
-		ensureMsgf(ShaderType::ShouldPrecachePermutation(FGlobalShaderPermutationParameters(ShaderType::GetStaticType().GetFName(), Platform, PermutationId)) != EShaderPermutationPrecacheRequest::NotRequired,
+		ensureMsgf(ShaderType::ShouldPrecachePermutation(FGlobalShaderPermutationParameters(ShaderType::GetStaticType().GetFName(), Platform, PermutationId)) != EShaderPermutationPrecacheRequest::NotUsed,
 			TEXT("Using a global shader permutation of %s which hasn't been requested for precaching at runtime. Check the implementation of ShouldPrecachePermutation on the global shader and make sure all required or development only permutations are requested for PSO precaching"), ShaderType::GetStaticType().GetName());
 
 		return TShaderRef<ShaderType>::Cast(Shader);
diff --git a/Engine/Source/Runtime/RenderCore/Public/Shader.h b/Engine/Source/Runtime/RenderCore/Public/Shader.h
--- a/Engine/Source/Runtime/RenderCore/Public/Shader.h
+++ b/Engine/Source/Runtime/RenderCore/Public/Shader.h
@@ -807,9 +807,9 @@ ENUM_CLASS_FLAGS(EShaderPermutationFlags);
 
 enum class EShaderPermutationPrecacheRequest : uint8
 {
-	Required,
-	DevelopmentOnly,
-	NotRequired,
+	Precached,			//< Permutation should be precached because it can be used at runtime
+	NotPrecached,		//< Permutation doesn't have to be precached (debug feature only for example)
+	NotUsed,			//< Permutation is not used with current cvar/configuration/feature level
 };
 
 RENDERCORE_API EShaderPermutationFlags GetShaderPermutationFlags(const FPlatformTypeLayoutParameters& LayoutParams);
@@ -875,7 +875,7 @@ public:
 	static bool ShouldCompilePermutation(const FShaderPermutationParameters&) { return true; }
 
 	/** Can be overridden by FShader subclasses to determine whether a specific permutation should be precached. */
-	static EShaderPermutationPrecacheRequest ShouldPrecachePermutation(const FShaderPermutationParameters& Parameters) { return ShouldCompilePermutation(Parameters) ? EShaderPermutationPrecacheRequest::Required : EShaderPermutationPrecacheRequest::NotRequired; }
+	static EShaderPermutationPrecacheRequest ShouldPrecachePermutation(const FShaderPermutationParameters& Parameters) { return ShouldCompilePermutation(Parameters) ? EShaderPermutationPrecacheRequest::Precached : EShaderPermutationPrecacheRequest::NotUsed; }
 
 	/** Can be overridden by FShader subclasses to determine whether compilation is valid. */
 	static bool ValidateCompiledResult(EShaderPlatform InPlatform, const FShaderParameterMap& InParameterMap, TArray<FString>& OutError) { return true; }
