ShaderCore: Add TryGetShaderFileHash. It is run in an iterative cook, and the error ignored, when the files being loaded were recorded during a previous cook but were deleted in between the two cooks.
#rnx

diff --git a/Engine/Source/Runtime/RenderCore/Private/ShaderCore.cpp b/Engine/Source/Runtime/RenderCore/Private/ShaderCore.cpp
--- a/Engine/Source/Runtime/RenderCore/Private/ShaderCore.cpp
+++ b/Engine/Source/Runtime/RenderCore/Private/ShaderCore.cpp
@@ -2427,11 +2427,16 @@ bool LoadShaderSourceFile(const TCHAR* InVirtualFilePath, EShaderPlatform Shader
 #endif // WITH_EDITORONLY_DATA
 }
 
+static FString FormatErrorCantFindSourceFile(const TCHAR* VirtualFilePath)
+{
+	return FString::Printf(TEXT("Couldn't find source file of virtual shader path \'%s\'"), VirtualFilePath);
+}
+
 void LoadShaderSourceFileChecked(const TCHAR* VirtualFilePath, EShaderPlatform ShaderPlatform, FString& OutFileContents, const FName* ShaderPlatformName)
 {
 	if (!LoadShaderSourceFile(VirtualFilePath, ShaderPlatform, &OutFileContents, nullptr, ShaderPlatformName))
 	{
-		UE_LOG(LogShaders, Fatal, TEXT("Couldn't find source file of virtual shader path \'%s\'"), VirtualFilePath);
+		UE_LOG(LogShaders, Fatal, TEXT("%s"), *FormatErrorCantFindSourceFile(VirtualFilePath));
 	}
 }
 
@@ -2871,7 +2876,8 @@ void HashShaderFileWithIncludes(FArchive& HashingArchive, const TCHAR* VirtualFi
 	}
 }
 
-static void UpdateSingleShaderFilehash(FSHA1& InOutHashState, const TCHAR* VirtualFilePath, EShaderPlatform ShaderPlatform)
+static bool TryUpdateSingleShaderFilehash(FSHA1& InOutHashState, const TCHAR* VirtualFilePath,
+	EShaderPlatform ShaderPlatform, FString* OutErrorMessage)
 {
 	// Get the list of includes this file contains
 	TArray<FString> IncludeVirtualFilePaths;
@@ -2886,7 +2892,14 @@ static void UpdateSingleShaderFilehash(FSHA1& InOutHashState, const TCHAR* Virtu
 	{
 		// Load the include file and hash it
 		FString IncludeFileContents;
-		LoadShaderSourceFileChecked(*IncludeVirtualFilePaths[IncludeIndex], ShaderPlatform, IncludeFileContents);
+		if (!LoadShaderSourceFile(*IncludeVirtualFilePaths[IncludeIndex], ShaderPlatform, &IncludeFileContents, nullptr))
+		{
+			if (OutErrorMessage)
+			{
+				*OutErrorMessage = FormatErrorCantFindSourceFile(*IncludeVirtualFilePaths[IncludeIndex]);
+			}
+			return false;
+		}
 		InOutHashState.UpdateWithString(*IncludeFileContents, IncludeFileContents.Len());
 #if WITH_EDITOR &&  !(UE_BUILD_SHIPPING || UE_BUILD_TEST)
 		if (UE_LOG_ACTIVE(LogShaders, Verbose))
@@ -2901,8 +2914,20 @@ static void UpdateSingleShaderFilehash(FSHA1& InOutHashState, const TCHAR* Virtu
 
 	// Load the source file and hash it
 	FString FileContents;
-	LoadShaderSourceFileChecked(VirtualFilePath, ShaderPlatform, FileContents);
+	if (!LoadShaderSourceFile(VirtualFilePath, ShaderPlatform, &FileContents, nullptr))
+	{
+		if (OutErrorMessage)
+		{
+			*OutErrorMessage = FormatErrorCantFindSourceFile(VirtualFilePath);
+		}
+		return false;
+	}
 	InOutHashState.UpdateWithString(*FileContents, FileContents.Len());
+	if (OutErrorMessage)
+	{
+		OutErrorMessage->Reset();
+	}
+	return true;
 }
 
 /** 
@@ -2911,12 +2936,21 @@ static void UpdateSingleShaderFilehash(FSHA1& InOutHashState, const TCHAR* Virtu
 */
 static FCriticalSection GShaderFileHashCalculationGuard;
 
-/**
- * Calculates a Hash for the given filename and its includes if it does not already exist in the Hash cache.
- * @param Filename - shader file to Hash
- * @param ShaderPlatform - shader platform to Hash
- */
 const FSHAHash& GetShaderFileHash(const TCHAR* VirtualFilePath, EShaderPlatform ShaderPlatform)
+{
+	FString ErrorMessage;
+	const FSHAHash* Hash = TryGetShaderFileHash(VirtualFilePath, ShaderPlatform, &ErrorMessage);
+	if (!Hash)
+	{
+		UE_LOG(LogShaders, Fatal, TEXT("%s"), *ErrorMessage);
+		static FSHAHash EmptyHash;
+		return EmptyHash;
+	}
+	return *Hash;
+}
+
+const FSHAHash* TryGetShaderFileHash(const TCHAR* VirtualFilePath, EShaderPlatform ShaderPlatform,
+	FString* OutErrorMessage)
 {
 	// Make sure we are only accessing GShaderHashCache from one thread
 	//check(IsInGameThread() || IsAsyncLoading());
@@ -2930,7 +2964,8 @@ const FSHAHash& GetShaderFileHash(const TCHAR* VirtualFilePath, EShaderPlatform
 			// If a hash for this filename has been cached, use that
 			if (CachedHash)
 			{
-				return *CachedHash;
+				if (OutErrorMessage) OutErrorMessage->Reset();
+				return CachedHash;
 			}
 		}
 
@@ -2944,11 +2979,16 @@ const FSHAHash& GetShaderFileHash(const TCHAR* VirtualFilePath, EShaderPlatform
 		const FSHAHash* CachedHash = GShaderHashCache.FindHash(ShaderPlatform, VirtualFilePath);
 		if (CachedHash)
 		{
-			return *CachedHash;
+			if (OutErrorMessage) OutErrorMessage->Reset();
+			return CachedHash;
 		}
 
 		FSHA1 HashState;
-		UpdateSingleShaderFilehash(HashState, VirtualFilePath, ShaderPlatform);
+		bool bSucceeded = TryUpdateSingleShaderFilehash(HashState, VirtualFilePath, ShaderPlatform, OutErrorMessage);
+		if (!bSucceeded)
+		{
+			return nullptr;
+		}
 		HashState.Final();
 
 		// Update the hash cache
@@ -2960,7 +3000,7 @@ const FSHAHash& GetShaderFileHash(const TCHAR* VirtualFilePath, EShaderPlatform
 		UE_LOG(LogShaders, Verbose, TEXT("Final hash for file %s, %s"), VirtualFilePath,*BytesToHex(&NewHash.Hash[0], 20));
 #endif
 		INC_FLOAT_STAT_BY(STAT_ShaderCompiling_HashingShaderFiles, (float)HashTime);
-		return NewHash;
+		return &NewHash;
 	}
 }
 
@@ -3011,7 +3051,11 @@ const FSHAHash& GetShaderFilesHash(const TArray<FString>& VirtualFilePaths, ESha
 		FSHA1 HashState;
 		for (const FString& VirtualFilePath : VirtualFilePaths)
 		{
-			UpdateSingleShaderFilehash(HashState, *VirtualFilePath, ShaderPlatform);
+			FString ErrorMessage;
+			if (!TryUpdateSingleShaderFilehash(HashState, *VirtualFilePath, ShaderPlatform, &ErrorMessage))
+			{
+				UE_LOG(LogShaders, Fatal, TEXT("%s"), *ErrorMessage);
+			}
 		}
 		HashState.Final();
 
diff --git a/Engine/Source/Runtime/RenderCore/Public/ShaderCore.h b/Engine/Source/Runtime/RenderCore/Public/ShaderCore.h
--- a/Engine/Source/Runtime/RenderCore/Public/ShaderCore.h
+++ b/Engine/Source/Runtime/RenderCore/Public/ShaderCore.h
@@ -1168,8 +1168,19 @@ extern RENDERCORE_API void GetShaderIncludes(const TCHAR* EntryPointVirtualFileP
  * Calculates a Hash for the given filename if it does not already exist in the Hash cache.
  * @param Filename - shader file to Hash
  * @param ShaderPlatform - shader platform to Hash
+ * @return Reference to the Hash created and stored for the file, or to an empty FSHAHash if not found.
+ * Logs an error if the file is not loadable.
  */
 extern RENDERCORE_API const class FSHAHash& GetShaderFileHash(const TCHAR* VirtualFilePath, EShaderPlatform ShaderPlatform);
+/**
+ * Calculates a Hash for the given filename if it does not already exist in the Hash cache.
+ * @param Filename - shader file to Hash
+ * @param ShaderPlatform - shader platform to Hash
+ * @param OutErrorMessage - If non-null, receives the errormessage if nullptr is returned.
+ * @return Pointer to the Hash created and stored for the file, or nullptr if not found.
+ */
+extern RENDERCORE_API const FSHAHash* TryGetShaderFileHash(const TCHAR* VirtualFilePath,
+	EShaderPlatform ShaderPlatform, FString* OutErrorMessage = nullptr);
 
 /**
  * Calculates a hash for the given source file and all files included from it.
