Replaced thread_local in RenderGraphAllocator by a TLS index as it's faster on some platforms

diff --git a/Engine/Source/Runtime/RenderCore/Private/RenderGraphAllocator.cpp b/Engine/Source/Runtime/RenderCore/Private/RenderGraphAllocator.cpp
--- a/Engine/Source/Runtime/RenderCore/Private/RenderGraphAllocator.cpp
+++ b/Engine/Source/Runtime/RenderCore/Private/RenderGraphAllocator.cpp
@@ -8,7 +8,7 @@
 thread_local int32 FRDGAllocator::NumAccessesTLS = 0;
 #endif
 
-thread_local FRDGAllocator* FRDGAllocator::AllocatorTLS = nullptr;
+uint32 FRDGAllocator::AllocatorTLSSlot = FPlatformTLS::AllocTlsSlot();
 
 FRDGAllocator::FRDGAllocator()
 #if !RDG_USE_MALLOC
@@ -50,8 +50,9 @@ FRDGAllocator::~FRDGAllocator()
 
 FRDGAllocator& FRDGAllocator::GetTLS()
 {
-	checkf(AllocatorTLS, TEXT("Attempted to access RDG allocator outside of FRDGAllocatorScope"));
-	return *AllocatorTLS;
+	void* Allocator = FPlatformTLS::GetTlsValue(AllocatorTLSSlot);
+	checkf(Allocator, TEXT("Attempted to access RDG allocator outside of FRDGAllocatorScope"));
+	return *(FRDGAllocator*)Allocator;
 }
 
 #if RDG_ALLOCATOR_DEBUG
@@ -96,14 +97,14 @@ void FRDGAllocator::ReleaseAll()
 }
 
 FRDGAllocatorScope::FRDGAllocatorScope(FRDGAllocator& Allocator)
-	: AllocatorToRestore(FRDGAllocator::AllocatorTLS)
+	: AllocatorToRestore(FPlatformTLS::GetTlsValue(FRDGAllocator::AllocatorTLSSlot))
 {
-	FRDGAllocator::AllocatorTLS = &Allocator;
+	FPlatformTLS::SetTlsValue(FRDGAllocator::AllocatorTLSSlot, &Allocator);
 }
 
 FRDGAllocatorScope::~FRDGAllocatorScope()
 {
-	FRDGAllocator::AllocatorTLS = AllocatorToRestore;
+	FPlatformTLS::SetTlsValue(FRDGAllocator::AllocatorTLSSlot, AllocatorToRestore);
 }
 
 FORCENOINLINE void UE::RenderCore::Private::OnInvalidRDGAllocatorNum(int32 NewNum, SIZE_T NumBytesPerElement)
diff --git a/Engine/Source/Runtime/RenderCore/Public/RenderGraphAllocator.h b/Engine/Source/Runtime/RenderCore/Public/RenderGraphAllocator.h
--- a/Engine/Source/Runtime/RenderCore/Public/RenderGraphAllocator.h
+++ b/Engine/Source/Runtime/RenderCore/Public/RenderGraphAllocator.h
@@ -139,7 +139,7 @@ private:
 #endif
 
 	friend class FRDGAllocatorScope;
-	static thread_local FRDGAllocator* AllocatorTLS;
+	static uint32 AllocatorTLSSlot;
 };
 
 class FRDGAllocatorScope
@@ -149,7 +149,7 @@ public:
 	RENDERCORE_API ~FRDGAllocatorScope();
 
 private:
-	FRDGAllocator* AllocatorToRestore;
+	void* AllocatorToRestore;
 };
 
 #define RDG_FRIEND_ALLOCATOR_FRIEND(Type) friend class FRDGAllocator::TObject<Type>
