Fix leaking shader code preloading entries.

- Credits to Rex Hill for working out the mechanics of the leak.

diff --git a/Engine/Source/Runtime/RenderCore/Private/ShaderCodeLibrary.cpp b/Engine/Source/Runtime/RenderCore/Private/ShaderCodeLibrary.cpp
--- a/Engine/Source/Runtime/RenderCore/Private/ShaderCodeLibrary.cpp
+++ b/Engine/Source/Runtime/RenderCore/Private/ShaderCodeLibrary.cpp
@@ -269,6 +269,11 @@ public:
 	virtual ~FShaderMapResource_SharedCode();
 
 	// FRenderResource interface.
+	virtual void ReleaseResource() override
+	{
+		FShaderMapResource::ReleaseResource();
+		ensureMsgf(!bShaderMapPreloaded && !LibraryInstance, TEXT("FShaderMapResource_SharedCode::ReleaseRHI() was not called on a shadermap resource owned by %s"), *GetOwnerName().ToString());
+	}
 	virtual void ReleaseRHI() override;
 
 	// FShaderMapResource interface
@@ -1111,6 +1116,7 @@ public:
 		{
 			Resource = new FShaderMapResource_SharedCode(this, ShaderMapIndex);
 			Resource->bShaderMapPreloaded = Library->PreloadShaderMap(ShaderMapIndex, AttachShaderReadRequestFunc);
+			BeginInitResource(Resource);
 		}
 		Resource->AddRef();
 	}
