Fix rare FShaderMapResource_SharedCode::ReleaseRHI crash

diff --git a/Engine/Source/Runtime/RenderCore/Private/ShaderCodeLibrary.cpp b/Engine/Source/Runtime/RenderCore/Private/ShaderCodeLibrary.cpp
--- a/Engine/Source/Runtime/RenderCore/Private/ShaderCodeLibrary.cpp
+++ b/Engine/Source/Runtime/RenderCore/Private/ShaderCodeLibrary.cpp
@@ -1241,7 +1241,7 @@ void FShaderMapResource_SharedCode::ReleasePreloadedShaderCode(int32 ShaderIndex
 
 void FShaderMapResource_SharedCode::ReleaseRHI()
 {
-	if (LibraryInstance)
+	if (LibraryInstance && ensureMsgf(LibraryInstance->Library, TEXT("LibraryInstance->Library pointer is expected to be valid as long as library's FShaderMapResource are alive.")))
 	{
 		const int32 NumShaders = GetNumShaders();
 		for (int32 i = 0; i < NumShaders; ++i)
