Disable Dynamic loop method ResolveDepth for mobile platforms, isn't compiling valid code

diff --git a/Engine/Source/Runtime/RenderCore/Private/ResolveShader.cpp b/Engine/Source/Runtime/RenderCore/Private/ResolveShader.cpp
--- a/Engine/Source/Runtime/RenderCore/Private/ResolveShader.cpp
+++ b/Engine/Source/Runtime/RenderCore/Private/ResolveShader.cpp
@@ -10,7 +10,6 @@ IMPLEMENT_SHADER_TYPE(, FResolveDepthPS, TEXT("/Engine/Private/ResolvePixelShade
 IMPLEMENT_SHADER_TYPE(, FResolveDepth2XPS, TEXT("/Engine/Private/ResolvePixelShader.usf"), TEXT("MainDepth"), SF_Pixel);
 IMPLEMENT_SHADER_TYPE(, FResolveDepth4XPS, TEXT("/Engine/Private/ResolvePixelShader.usf"), TEXT("MainDepth"), SF_Pixel);
 IMPLEMENT_SHADER_TYPE(, FResolveDepth8XPS, TEXT("/Engine/Private/ResolvePixelShader.usf"), TEXT("MainDepth"), SF_Pixel);
-IMPLEMENT_SHADER_TYPE(, FResolveDepthArrayPS, TEXT("/Engine/Private/ResolvePixelShader.usf"), TEXT("MainDepth"), SF_Pixel);
 IMPLEMENT_SHADER_TYPE(, FResolveDepthArray2XPS, TEXT("/Engine/Private/ResolvePixelShader.usf"), TEXT("MainDepth"), SF_Pixel);
 IMPLEMENT_SHADER_TYPE(, FResolveDepthArray4XPS, TEXT("/Engine/Private/ResolvePixelShader.usf"), TEXT("MainDepth"), SF_Pixel);
 IMPLEMENT_SHADER_TYPE(, FResolveDepthArray8XPS, TEXT("/Engine/Private/ResolvePixelShader.usf"), TEXT("MainDepth"), SF_Pixel);
@@ -37,6 +36,11 @@ void FResolveDepthPS::ModifyCompilationEnvironment(const FGlobalShaderPermutatio
 	}
 }
 
+bool FResolveDepthPS::ShouldCompilePermutation(const FGlobalShaderPermutationParameters& Parameters)
+{
+	return !IsOpenGLPlatform(Parameters.Platform);
+}
+
 void FResolveDepthPS::SetParameters(FRHIBatchedShaderParameters& BatchedParameters, FParameter)
 {
 }
@@ -88,46 +92,30 @@ void FResolveDepth8XPS::ModifyCompilationEnvironment(const FGlobalShaderPermutat
 	OutEnvironment.SetDefine(TEXT("DEPTH_RESOLVE_NUM_SAMPLES"), 8);
 }
 
-// FResolveDepthArrayPS
-
-FResolveDepthArrayPS::FResolveDepthArrayPS() = default;
-FResolveDepthArrayPS::FResolveDepthArrayPS(const ShaderMetaType::CompiledShaderInitializerType& Initializer)
-	: FResolveDepthPS(Initializer)
-{
-}
-
-bool FResolveDepthArrayPS::ShouldCompilePermutation(const FGlobalShaderPermutationParameters& Parameters)
-{
-	if (FResolveDepthPS::ShouldCompilePermutation(Parameters))
-	{
-		UE::StereoRenderUtils::FStereoShaderAspects Aspects(Parameters.Platform);
-		return Aspects.IsMobileMultiViewEnabled();
-	}
-	return false;
-}
 
-void FResolveDepthArrayPS::ModifyCompilationEnvironment(const FGlobalShaderPermutationParameters& Parameters, FShaderCompilerEnvironment& OutEnvironment)
+static bool IsMobileMultiViewEnabled(EShaderPlatform ShaderPlatform)
 {
-	FResolveDepthPS::ModifyCompilationEnvironment(Parameters, OutEnvironment);
-	OutEnvironment.SetDefine(TEXT("DEPTH_RESOLVE_TEXTUREARRAY"), 1);
+	UE::StereoRenderUtils::FStereoShaderAspects Aspects(ShaderPlatform);
+	return Aspects.IsMobileMultiViewEnabled();
 }
 
 // FResolveDepthArray2XPS
 
 FResolveDepthArray2XPS::FResolveDepthArray2XPS() = default;
 FResolveDepthArray2XPS::FResolveDepthArray2XPS(const ShaderMetaType::CompiledShaderInitializerType& Initializer)
-	: FResolveDepthArrayPS(Initializer)
+	: FResolveDepthPS(Initializer)
 {
 }
 
 bool FResolveDepthArray2XPS::ShouldCompilePermutation(const FGlobalShaderPermutationParameters& Parameters)
 {
-	return FResolveDepthArrayPS::ShouldCompilePermutation(Parameters);
+	return IsMobileMultiViewEnabled(Parameters.Platform);
 }
 
 void FResolveDepthArray2XPS::ModifyCompilationEnvironment(const FGlobalShaderPermutationParameters& Parameters, FShaderCompilerEnvironment& OutEnvironment)
 {
-	FResolveDepthArrayPS::ModifyCompilationEnvironment(Parameters, OutEnvironment);
+	FResolveDepthPS::ModifyCompilationEnvironment(Parameters, OutEnvironment);
+	OutEnvironment.SetDefine(TEXT("DEPTH_RESOLVE_TEXTUREARRAY"), 1);
 	OutEnvironment.SetDefine(TEXT("DEPTH_RESOLVE_NUM_SAMPLES"), 2);
 }
 
@@ -135,18 +123,19 @@ void FResolveDepthArray2XPS::ModifyCompilationEnvironment(const FGlobalShaderPer
 
 FResolveDepthArray4XPS::FResolveDepthArray4XPS() = default;
 FResolveDepthArray4XPS::FResolveDepthArray4XPS(const ShaderMetaType::CompiledShaderInitializerType& Initializer)
-	: FResolveDepthArrayPS(Initializer)
+	: FResolveDepthPS(Initializer)
 {
 }
 
 bool FResolveDepthArray4XPS::ShouldCompilePermutation(const FGlobalShaderPermutationParameters& Parameters)
 {
-	return FResolveDepthArrayPS::ShouldCompilePermutation(Parameters);
+	return IsMobileMultiViewEnabled(Parameters.Platform);
 }
 
 void FResolveDepthArray4XPS::ModifyCompilationEnvironment(const FGlobalShaderPermutationParameters& Parameters, FShaderCompilerEnvironment& OutEnvironment)
 {
-	FResolveDepthArrayPS::ModifyCompilationEnvironment(Parameters, OutEnvironment);
+	FResolveDepthPS::ModifyCompilationEnvironment(Parameters, OutEnvironment);
+	OutEnvironment.SetDefine(TEXT("DEPTH_RESOLVE_TEXTUREARRAY"), 1);
 	OutEnvironment.SetDefine(TEXT("DEPTH_RESOLVE_NUM_SAMPLES"), 4);
 }
 
@@ -154,23 +143,19 @@ void FResolveDepthArray4XPS::ModifyCompilationEnvironment(const FGlobalShaderPer
 
 FResolveDepthArray8XPS::FResolveDepthArray8XPS() = default;
 FResolveDepthArray8XPS::FResolveDepthArray8XPS(const ShaderMetaType::CompiledShaderInitializerType& Initializer)
-	: FResolveDepthArrayPS(Initializer)
+	: FResolveDepthPS(Initializer)
 {
 }
 
-bool FResolveDepthArray8XPS::ShouldCache(EShaderPlatform Platform)
-{
-	return GetMaxSupportedFeatureLevel(Platform) >= ERHIFeatureLevel::SM5;
-}
-
 bool FResolveDepthArray8XPS::ShouldCompilePermutation(const FGlobalShaderPermutationParameters& Parameters)
 {
-	return FResolveDepthArrayPS::ShouldCompilePermutation(Parameters);
+	return IsMobileMultiViewEnabled(Parameters.Platform);
 }
 
 void FResolveDepthArray8XPS::ModifyCompilationEnvironment(const FGlobalShaderPermutationParameters& Parameters, FShaderCompilerEnvironment& OutEnvironment)
 {
-	FResolveDepthArrayPS::ModifyCompilationEnvironment(Parameters, OutEnvironment);
+	FResolveDepthPS::ModifyCompilationEnvironment(Parameters, OutEnvironment);
+	OutEnvironment.SetDefine(TEXT("DEPTH_RESOLVE_TEXTUREARRAY"), 1);
 	OutEnvironment.SetDefine(TEXT("DEPTH_RESOLVE_NUM_SAMPLES"), 8);
 }
 
@@ -255,7 +240,6 @@ void CreateResolveShaders()
 	ForceInitGlobalShaderType<FResolveDepth2XPS>();
 	ForceInitGlobalShaderType<FResolveDepth4XPS>();
 	ForceInitGlobalShaderType<FResolveDepth8XPS>();
-	ForceInitGlobalShaderType<FResolveDepthArrayPS>();
 	ForceInitGlobalShaderType<FResolveDepthArray2XPS>();
 	ForceInitGlobalShaderType<FResolveDepthArray4XPS>();
 	ForceInitGlobalShaderType<FResolveDepthArray8XPS>();
diff --git a/Engine/Source/Runtime/RenderCore/Public/ResolveShader.h b/Engine/Source/Runtime/RenderCore/Public/ResolveShader.h
--- a/Engine/Source/Runtime/RenderCore/Public/ResolveShader.h
+++ b/Engine/Source/Runtime/RenderCore/Public/ResolveShader.h
@@ -27,6 +27,7 @@ public:
 	FResolveDepthPS(const ShaderMetaType::CompiledShaderInitializerType& Initializer);
 
 	static void ModifyCompilationEnvironment(const FGlobalShaderPermutationParameters& Parameters, FShaderCompilerEnvironment& OutEnvironment);
+	static bool ShouldCompilePermutation(const FGlobalShaderPermutationParameters& Parameters);
 
 	RENDERCORE_API void SetParameters(FRHIBatchedShaderParameters& BatchedParameters, FParameter);
 
@@ -44,6 +45,10 @@ public:
 	FResolveDepth2XPS(const ShaderMetaType::CompiledShaderInitializerType& Initializer);
 
 	static void ModifyCompilationEnvironment(const FGlobalShaderPermutationParameters& Parameters, FShaderCompilerEnvironment& OutEnvironment);
+	static bool ShouldCompilePermutation(const FGlobalShaderPermutationParameters& Parameters)
+	{
+		return true;
+	}
 };
 
 class FResolveDepth4XPS : public FResolveDepthPS
@@ -56,6 +61,10 @@ public:
 	FResolveDepth4XPS(const ShaderMetaType::CompiledShaderInitializerType& Initializer);
 
 	static void ModifyCompilationEnvironment(const FGlobalShaderPermutationParameters& Parameters, FShaderCompilerEnvironment& OutEnvironment);
+	static bool ShouldCompilePermutation(const FGlobalShaderPermutationParameters& Parameters)
+	{
+		return true;
+	}
 };
 
 class FResolveDepth8XPS : public FResolveDepthPS
@@ -69,22 +78,13 @@ public:
 
 	static bool ShouldCache(EShaderPlatform Platform);
 	static void ModifyCompilationEnvironment(const FGlobalShaderPermutationParameters& Parameters, FShaderCompilerEnvironment& OutEnvironment);
+	static bool ShouldCompilePermutation(const FGlobalShaderPermutationParameters& Parameters)
+	{
+		return true;
+	}
 };
 
-class FResolveDepthArrayPS : public FResolveDepthPS
-{
-	DECLARE_EXPORTED_GLOBAL_SHADER(FResolveDepthArrayPS, RENDERCORE_API);
-public:
-	typedef FDummyResolveParameter FParameter;
-
-	FResolveDepthArrayPS();
-	FResolveDepthArrayPS(const ShaderMetaType::CompiledShaderInitializerType& Initializer);
-
-	static bool ShouldCompilePermutation(const FGlobalShaderPermutationParameters& Parameters);
-	static void ModifyCompilationEnvironment(const FGlobalShaderPermutationParameters& Parameters, FShaderCompilerEnvironment& OutEnvironment);
-};
-
-class FResolveDepthArray2XPS : public FResolveDepthArrayPS
+class FResolveDepthArray2XPS : public FResolveDepthPS
 {
 	DECLARE_EXPORTED_GLOBAL_SHADER(FResolveDepthArray2XPS, RENDERCORE_API);
 public:
@@ -97,7 +97,7 @@ public:
 	static void ModifyCompilationEnvironment(const FGlobalShaderPermutationParameters& Parameters, FShaderCompilerEnvironment& OutEnvironment);
 };
 
-class FResolveDepthArray4XPS : public FResolveDepthArrayPS
+class FResolveDepthArray4XPS : public FResolveDepthPS
 {
 	DECLARE_EXPORTED_GLOBAL_SHADER(FResolveDepthArray4XPS, RENDERCORE_API);
 public:
@@ -110,7 +110,7 @@ public:
 	static void ModifyCompilationEnvironment(const FGlobalShaderPermutationParameters& Parameters, FShaderCompilerEnvironment& OutEnvironment);
 };
 
-class FResolveDepthArray8XPS : public FResolveDepthArrayPS
+class FResolveDepthArray8XPS : public FResolveDepthPS
 {
 	DECLARE_EXPORTED_GLOBAL_SHADER(FResolveDepthArray8XPS, RENDERCORE_API);
 public:
@@ -119,7 +119,6 @@ public:
 	FResolveDepthArray8XPS();
 	FResolveDepthArray8XPS(const ShaderMetaType::CompiledShaderInitializerType& Initializer);
 
-	RENDERCORE_API static bool ShouldCache(EShaderPlatform Platform);
 	static bool ShouldCompilePermutation(const FGlobalShaderPermutationParameters& Parameters);
 	static void ModifyCompilationEnvironment(const FGlobalShaderPermutationParameters& Parameters, FShaderCompilerEnvironment& OutEnvironment);
 };
