[shaders] add compression for FShaderSource objects and enable by default (Mermaid SuperFast compression). This halves the high watermark for the ShaderCompiler LLM scope in my testing (QAGame cold cook), with minimal measurable impact to preprocessing time (<3% on average, but this may just be noise).

diff --git a/Engine/Source/Runtime/RenderCore/Private/ShaderCore.cpp b/Engine/Source/Runtime/RenderCore/Private/ShaderCore.cpp
--- a/Engine/Source/Runtime/RenderCore/Private/ShaderCore.cpp
+++ b/Engine/Source/Runtime/RenderCore/Private/ShaderCore.cpp
@@ -1515,6 +1515,9 @@ public:
 			// if the preprocessed job cache is enabled we need to strip the preprocessed code, this removes comments, line directives
 			// and blank lines to improve deduplication (and populates data required to remap diagnostic messages to correct line numbers)
 			Job.PreprocessOutput.StripCode(Job.Input.NeedsOriginalShaderSource());
+
+			// always compress the code after stripping to minimize memory footprint
+			Job.PreprocessOutput.CompressCode();
 		}
 
 		Job.PreprocessOutput.ElapsedTime = FPlatformTime::Seconds() - StartPreprocessTime;
@@ -1581,6 +1584,9 @@ public:
 			PreprocessShaderInternal(Compiler, Job);
 		}
 
+		// decompress if necessary; this is a no-op if source is not compressed.
+		Job.PreprocessOutput.DecompressCode();
+
 		if (Job.PreprocessOutput.bSucceeded)
 		{
 			if (Job.SecondaryPreprocessOutput.IsValid())
@@ -4089,7 +4095,8 @@ void FShaderCompileJob::OnComplete()
 		// (jobs deserialized from the cache/wait list/ddc will have a compiletime of 0.0)
 		&& (CVarDumpDebugInfoForCacheHits.GetValueOnAnyThread() || Output.CompileTime > 0.0f))
 	{
-		
+		// need to decompress the source code before dumping debug data
+		PreprocessOutput.DecompressCode();
 		if (SecondaryPreprocessOutput.IsValid() && SecondaryOutput.IsValid())
 		{
 			ShaderFormat->OutputDebugData(Input, PreprocessOutput, *SecondaryPreprocessOutput, Output, *SecondaryOutput);
diff --git a/Engine/Source/Runtime/RenderCore/Private/ShaderSource.cpp b/Engine/Source/Runtime/RenderCore/Private/ShaderSource.cpp
--- a/Engine/Source/Runtime/RenderCore/Private/ShaderSource.cpp
+++ b/Engine/Source/Runtime/RenderCore/Private/ShaderSource.cpp
@@ -2,8 +2,24 @@
 
 #include "ShaderSource.h"
 #include "Containers/StringConv.h"
+#include "Compression/OodleDataCompression.h"
+#include "HAL/IConsoleManager.h"
 #include "HAL/UnrealMemory.h"
 
+static int32 GShaderSourceCompressionMethod = 2;
+static FAutoConsoleVariableRef CVarShaderSourceCompressionMethod(
+	TEXT("r.ShaderSource.CompressionMethod"),
+	GShaderSourceCompressionMethod,
+	TEXT("Compression method for shader source stored in memory. See FOodleDataCompression::ECompressor enum for supported values; defaults to Mermaid."),
+	ECVF_Default);
+
+static int32 GShaderSourceCompressionLevel = 1;
+FAutoConsoleVariableRef CVarShaderSourceCompressionLevel(
+	TEXT("r.ShaderSource.CompressionLevel"),
+	GShaderSourceCompressionLevel,
+	TEXT("Compression level for shader source stored in memory. See FOodleDataCompression::ECompressionLevel enum for supported values; default is SuperFast."),
+	ECVF_Default);
+
 FShaderSource::FShaderSource(FShaderSource::FViewType InSrc, int32 AdditionalSlack)
 {
 	Set(InSrc, AdditionalSlack);
@@ -27,6 +43,9 @@ void FShaderSource::Set(FAnsiStringView InSrc)
 
 FShaderSource& FShaderSource::operator=(FShaderSource::FStringType&& InSrc)
 {
+	SourceCompressed.Empty();
+	DecompressedSize = 0;
+
 	int32 InInitialLen = InSrc.Len();
 	// input char array already has a null terminator appended, so can add one less padding char
 	TArray<CharType>& InSrcArray = InSrc.GetCharArray();
@@ -35,9 +54,57 @@ FShaderSource& FShaderSource::operator=(FShaderSource::FStringType&& InSrc)
 	return *this;
 }
 
+void FShaderSource::Compress()
+{
+	FOodleDataCompression::ECompressor Compressor = static_cast<FOodleDataCompression::ECompressor>(GShaderSourceCompressionMethod);
+	if (Compressor == FOodleDataCompression::ECompressor::NotSet)
+	{
+		return;
+	}
+
+	checkf(!IsCompressed(), TEXT("FShaderSource is already compressed"));
+
+	FOodleDataCompression::ECompressionLevel CompressionLevel = static_cast<FOodleDataCompression::ECompressionLevel>(GShaderSourceCompressionLevel);
+	
+	DecompressedSize = Source.Num();
+	int32 CompressionBufferSize = FOodleDataCompression::CompressedBufferSizeNeeded(DecompressedSize);
+	SourceCompressed.SetNumUninitialized(CompressionBufferSize);
+
+	int32 CompressedSize = FOodleDataCompression::Compress(
+		SourceCompressed.GetData(), 
+		CompressionBufferSize, 
+		Source.GetData(), 
+		DecompressedSize, 
+		Compressor,
+		CompressionLevel);
+
+	SourceCompressed.SetNum(CompressedSize, EAllowShrinking::Yes);
+	Source.Empty();
+}
+
+void FShaderSource::Decompress()
+{
+	if (IsCompressed())
+	{
+		Source.SetNumUninitialized(DecompressedSize);
+		FOodleDataCompression::Decompress(Source.GetData(), DecompressedSize, SourceCompressed.GetData(), SourceCompressed.Num());
+		SourceCompressed.Empty();
+		DecompressedSize = 0;
+	}
+}
+
 FArchive& operator<<(FArchive& Ar, FShaderSource& Src)
 {
-	Ar << Src.Source;
+	Ar << Src.DecompressedSize;
+	if (Src.IsCompressed())
+	{
+		Ar << Src.SourceCompressed;
+	}
+	else
+	{
+		Ar << Src.Source;
+	}
+
 	return Ar;
 }
 
diff --git a/Engine/Source/Runtime/RenderCore/Public/ShaderCompilerCore.h b/Engine/Source/Runtime/RenderCore/Public/ShaderCompilerCore.h
--- a/Engine/Source/Runtime/RenderCore/Public/ShaderCompilerCore.h
+++ b/Engine/Source/Runtime/RenderCore/Public/ShaderCompilerCore.h
@@ -31,7 +31,7 @@ class FShaderPipelineCompileJob;
 typedef TSharedPtr<TArray<ANSICHAR>, ESPMode::ThreadSafe> FShaderSharedAnsiStringPtr;
 
 // this is for the protocol, not the data, bump if FShaderCompilerInput/FShaderPreprocessOutput serialization, SerializeWorkerInput or ProcessInputFromArchive changes.
-inline const int32 ShaderCompileWorkerInputVersion = 26;
+inline const int32 ShaderCompileWorkerInputVersion = 27;
 // this is for the protocol, not the data, bump if FShaderCompilerOutput or WriteToOutputArchive changes.
 inline const int32 ShaderCompileWorkerOutputVersion = 20;
 // this is for the protocol, not the data.
diff --git a/Engine/Source/Runtime/RenderCore/Public/ShaderPreprocessTypes.h b/Engine/Source/Runtime/RenderCore/Public/ShaderPreprocessTypes.h
--- a/Engine/Source/Runtime/RenderCore/Public/ShaderPreprocessTypes.h
+++ b/Engine/Source/Runtime/RenderCore/Public/ShaderPreprocessTypes.h
@@ -267,6 +267,16 @@ private:
 	// and saving the original source in the OriginalPreprocessedSource member
 	void StripCode(bool bCopyOriginalPreprocessdSource);
 
+	void CompressCode() 
+	{
+		PreprocessedSource.Compress();
+	}
+
+	void DecompressCode()
+	{
+		PreprocessedSource.Decompress();
+	}
+
 	void RemapErrors(FShaderCompilerOutput& Output) const;
 
 	// Output of preprocessing; should be set by IShaderFormat::PreprocessShader
diff --git a/Engine/Source/Runtime/RenderCore/Public/ShaderSource.h b/Engine/Source/Runtime/RenderCore/Public/ShaderSource.h
--- a/Engine/Source/Runtime/RenderCore/Public/ShaderSource.h
+++ b/Engine/Source/Runtime/RenderCore/Public/ShaderSource.h
@@ -27,6 +27,9 @@ private:
 	{
 		Source.SetNumUninitialized(Num + ShaderSourceSimdPadding, AllowShrinking);
 		FMemory::Memzero(Source.GetData() + Num, sizeof(CharType) * ShaderSourceSimdPadding);
+
+		SourceCompressed.Empty();
+		DecompressedSize = 0;
 	}
 
 public:
@@ -96,15 +99,23 @@ public:
 		ShrinkToLen(Num, bShrink ? EAllowShrinking::Yes : EAllowShrinking::No);
 	}
 
-	/* String view accessor.
+	/* String view accessor. Will decompress data if it was previously compressed.
 	 * @return a string view pointing to the source contents, excluding padding.
 	 */
-	inline FViewType GetView() const { return { Source.GetData(), Len() }; }
+	inline FViewType GetView() const
+	{ 
+		checkf(!IsCompressed(), TEXT("FShaderSource is compressed; must decompress prior to calling GetView"));
+		return { Source.GetData(), Len() }; 
+	}
 
-	/* Direct data pointer accessor
+	/* Direct data pointer accessor. Will decompress data if it was previously compressed.
 	 * @return a pointer to the source data; will be null terminated by the SIMD padding.
 	 */
-	inline CharType* GetData() { return Source.GetData(); }
+	inline CharType* GetData()
+	{
+		checkf(!IsCompressed(), TEXT("FShaderSource is compressed; must decompress prior to calling GetData"));
+		return Source.GetData();
+	}
 
 	/* IsEmpty predicate
 	 * @return true if this source object is empty excluding the SIMD padding, false otherwise.
@@ -114,13 +125,25 @@ public:
 	/* Length accessor.
 	 * @return the non-padded length of the source (also excluding null terminator)
 	 */
-	inline int32 Len() const { return Source.Num() - ShaderSourceSimdPadding; };
+	inline int32 Len() const
+	{ 
+		checkf(!IsCompressed(), TEXT("Len should not be called on compressed FShaderSource."))
+		return Source.Num() - ShaderSourceSimdPadding;
+	};
+
+	inline bool IsCompressed() const
+	{
+		return DecompressedSize != 0; // DecompressedSize member is only set when compression occurs
+	}
 
 	/* FArchive serialization operator. Note this currently serializes padding for simplicity's sake.
 	 * @param Ar The archive to serialize from/to
 	 * @param ShaderSource the source object to serialize
 	 */
 	friend FArchive& operator<<(FArchive& Ar, FShaderSource& ShaderSource);
+
+	RENDERCORE_API void Compress();
+	RENDERCORE_API void Decompress();
 	
 private:
 #if SHADER_SOURCE_ANSI
@@ -129,5 +152,7 @@ private:
 	static constexpr int32 ShaderSourceSimdPadding = 7;
 #endif
 	TArray<CharType> Source;
+	TArray<uint8> SourceCompressed;
+	int32 DecompressedSize = 0;
 };
 
