[shaders] clean up some vestigial code in the preprocessing path - this no longer handles exceptions so the invoke function, which was needed to avoid problems around vars requiring unwinding, is no longer necessary (and the out exception params are unused)

diff --git a/Engine/Source/Runtime/RenderCore/Private/ShaderCore.cpp b/Engine/Source/Runtime/RenderCore/Private/ShaderCore.cpp
--- a/Engine/Source/Runtime/RenderCore/Private/ShaderCore.cpp
+++ b/Engine/Source/Runtime/RenderCore/Private/ShaderCore.cpp
@@ -1451,30 +1451,17 @@ class FInternalShaderCompilerFunctions
 {
 public:
 
-	static bool InvokePreprocess(
-		const IShaderFormat* Backend,
-		const FShaderCompilerInput& Input,
-		const FShaderCompilerEnvironment& Environment,
-		FShaderPreprocessOutput& Output,
-		FString& OutExceptionCallstack,
-		FString& OutExceptionMsg)
-	{
-		return Backend->PreprocessShader(Input, Environment, Output);
-	}
-
 	static void PreprocessShaderInternal(
 		const IShaderFormat* Backend, 
 		FShaderCompileJob& Job, 
-		const FShaderCompilerEnvironment& Environment,
-		FString& OutExceptionCallstack,
-		FString& OutExceptionMsg)
+		const FShaderCompilerEnvironment& Environment)
 	{
-		Job.PreprocessOutput.bSucceeded = InvokePreprocess(Backend, Job.Input, Environment, Job.PreprocessOutput, OutExceptionCallstack, OutExceptionMsg);
+		Job.PreprocessOutput.bSucceeded = Backend->PreprocessShader(Job.Input, Environment, Job.PreprocessOutput);
 		if (Job.PreprocessOutput.bSucceeded && Backend->RequiresSecondaryCompile(Job.Input, Environment, Job.PreprocessOutput))
 		{
 			Job.SecondaryPreprocessOutput = MakeUnique<FShaderPreprocessOutput>();
 			Job.SecondaryPreprocessOutput->bIsSecondary = true;
-			Job.PreprocessOutput.bSucceeded &= InvokePreprocess(Backend, Job.Input, Environment, *Job.SecondaryPreprocessOutput, OutExceptionCallstack, OutExceptionMsg);
+			Job.PreprocessOutput.bSucceeded &= Backend->PreprocessShader(Job.Input, Environment, *Job.SecondaryPreprocessOutput);
 		}
 	}
 
@@ -1482,9 +1469,6 @@ public:
 	{
 		const double StartPreprocessTime = FPlatformTime::Seconds();
 
-		FString ExceptionCallstack;
-		FString ExceptionMsg;
-
 		// if preprocessed cache is disabled, this function is executed as part of the compile process
 		// which will have merged the environment inline prior to calling this, so we can skip merging here 
 		if (Job.Input.bCachePreprocessed && IsValidRef(Job.Input.SharedEnvironment))
@@ -1494,20 +1478,11 @@ public:
 			// and affect what is passed to the workers)
 			FShaderCompilerEnvironment MergedEnvironment = Job.Input.Environment;
 			MergedEnvironment.Merge(*Job.Input.SharedEnvironment);
-			PreprocessShaderInternal(Backend, Job, MergedEnvironment, ExceptionCallstack, ExceptionMsg);
+			PreprocessShaderInternal(Backend, Job, MergedEnvironment);
 		}
 		else
 		{
-			PreprocessShaderInternal(Backend, Job, Job.Input.Environment, ExceptionCallstack, ExceptionMsg);
-		}
-
-		if (!Job.PreprocessOutput.bSucceeded && (!ExceptionMsg.IsEmpty() || !ExceptionCallstack.IsEmpty()))
-		{
-			FString StrippedErrorMessage = FString::Printf(
-				TEXT("Exception encountered in platform compiler: %s\nException Callstack:\n%s"),
-				*ExceptionMsg,
-				*ExceptionCallstack);
-			Job.PreprocessOutput.LogError(MoveTemp(StrippedErrorMessage));
+			PreprocessShaderInternal(Backend, Job, Job.Input.Environment);
 		}
 
 		if (Job.PreprocessOutput.bSucceeded && Job.Input.bCachePreprocessed)
