[shaders] also compress the "original" (non-stripped) source if it's set. no downside and also will help memory consumption when doing full cooks with shader debug info enabled.

diff --git a/Engine/Source/Runtime/RenderCore/Public/ShaderPreprocessTypes.h b/Engine/Source/Runtime/RenderCore/Public/ShaderPreprocessTypes.h
--- a/Engine/Source/Runtime/RenderCore/Public/ShaderPreprocessTypes.h
+++ b/Engine/Source/Runtime/RenderCore/Public/ShaderPreprocessTypes.h
@@ -270,11 +270,13 @@ private:
 	void CompressCode() 
 	{
 		PreprocessedSource.Compress();
+		OriginalPreprocessedSource.Compress();
 	}
 
 	void DecompressCode()
 	{
 		PreprocessedSource.Decompress();
+		OriginalPreprocessedSource.Decompress();
 	}
 
 	void RemapErrors(FShaderCompilerOutput& Output) const;
