Scene capture cube:  Support for automatic mip map generation.  Upgraded FGenerateMips API to support generating cube mip maps in both compute and raster code paths.  Also needed to fix a spurious error in FRHITextureSRVCreateInfo::Validate.

diff --git a/Engine/Source/Runtime/Engine/Classes/Engine/TextureRenderTargetCube.h b/Engine/Source/Runtime/Engine/Classes/Engine/TextureRenderTargetCube.h
--- a/Engine/Source/Runtime/Engine/Classes/Engine/TextureRenderTargetCube.h
+++ b/Engine/Source/Runtime/Engine/Classes/Engine/TextureRenderTargetCube.h
@@ -44,6 +44,14 @@ class UTextureRenderTargetCube : public UTextureRenderTarget
 	UPROPERTY()
 	uint8 bForceLinearGamma:1;
 
+	/** Whether to support Mip maps for this render target texture */
+	UPROPERTY(EditAnywhere, BlueprintReadOnly, Category=TextureRenderTargetCube, AssetRegistrySearchable)
+	uint8 bAutoGenerateMips : 1;
+
+	/** Sampler filter type for AutoGenerateMips. Defaults to match texture filter. */
+	UPROPERTY(EditAnywhere, BlueprintReadOnly, Category=TextureRenderTargetCube, AssetRegistrySearchable, meta = (editcondition = "bAutoGenerateMips"))
+	TEnumAsByte<enum TextureFilter> MipsSamplerFilter;
+
 	/** 
 	* Initialize the settings needed to create a render target texture
 	* and create its resource
@@ -81,7 +89,7 @@ class UTextureRenderTargetCube : public UTextureRenderTarget
 	
 	FORCEINLINE int32 GetNumMips() const
 	{
-		return 1;
+		return NumMips;
 	}
 	//~ Begin UObject Interface
 #if WITH_EDITOR
@@ -100,6 +108,9 @@ class UTextureRenderTargetCube : public UTextureRenderTarget
 	virtual float GetDisplayGamma() const override;
 	virtual ETextureClass GetRenderTargetTextureClass() const override { return ETextureClass::Cube; }
 	//~ End UTextureRenderTarget Interface
+
+private:
+	int32	NumMips;
 };
 
 
diff --git a/Engine/Source/Runtime/Engine/Private/TextureRenderTargetCube.cpp b/Engine/Source/Runtime/Engine/Private/TextureRenderTargetCube.cpp
--- a/Engine/Source/Runtime/Engine/Private/TextureRenderTargetCube.cpp
+++ b/Engine/Source/Runtime/Engine/Private/TextureRenderTargetCube.cpp
@@ -16,6 +16,7 @@
 #include "Engine/TextureCube.h"
 #include "RHIUtilities.h"
 #include "ProfilingDebugging/AssetMetadataTrace.h"
+#include "GenerateMips.h"
 
 #if WITH_EDITOR
 #include "Components/SceneCaptureComponentCube.h"
@@ -35,6 +36,9 @@ UTextureRenderTargetCube::UTextureRenderTargetCube(const FObjectInitializer& Obj
 	ClearColor = FLinearColor(0.0f, 1.0f, 0.0f, 1.0f);
 	OverrideFormat = PF_Unknown;
 	bForceLinearGamma = true;
+	bAutoGenerateMips = false;
+	MipsSamplerFilter = Filter;
+	NumMips = 0;
 	// note bool SRGB not set
 }
 
@@ -128,6 +132,20 @@ void UTextureRenderTargetCube::GetResourceSizeEx(FResourceSizeEx& CumulativeReso
 
 FTextureResource* UTextureRenderTargetCube::CreateResource()
 {
+	if (bAutoGenerateMips)
+	{
+		NumMips = FMath::FloorLog2(SizeX) + 1;
+
+		if (RHIRequiresComputeGenerateMips())
+		{
+			bCanCreateUAV = 1;
+		}
+	}
+	else
+	{
+		NumMips = 1;
+	}
+
 	return new FTextureRenderTargetCubeResource(this);
 }
 
@@ -149,8 +167,9 @@ void UTextureRenderTargetCube::PostEditChangeProperty(FPropertyChangedEvent& Pro
 
 	// Notify any scene capture components that point to this texture that they may need to refresh
 	static const FName SizeXName = GET_MEMBER_NAME_CHECKED(UTextureRenderTargetCube, SizeX);
+	static const FName AutoGenerateMipsName = GET_MEMBER_NAME_CHECKED(UTextureRenderTargetCube, bAutoGenerateMips);
 
-	if (PropertyChangedEvent.GetPropertyName() == SizeXName)
+	if (PropertyChangedEvent.GetPropertyName() == SizeXName || PropertyChangedEvent.GetPropertyName() == AutoGenerateMipsName)
 	{
 		for (TObjectIterator<USceneCaptureComponentCube> It; It; ++It)
 		{
@@ -260,6 +279,21 @@ void FTextureRenderTargetCubeResource::InitRHI(FRHICommandListBase& RHICmdList)
 		{
 			TexCreateFlags |= ETextureCreateFlags::UAV;
 		}
+		
+		if (Owner->bAutoGenerateMips)
+		{
+			TexCreateFlags |= ETextureCreateFlags::GenerateMipCapable;
+
+			if (FGenerateMips::WillFormatSupportCompute(Owner->GetFormat()))
+			{
+				TexCreateFlags |= ETextureCreateFlags::UAV;
+			}
+			else
+			{
+				// Required in FGenerateMips::ExecuteRaster for FRenderTargetBinding to work on individual slice
+				TexCreateFlags |= ETextureCreateFlags::TargetArraySlicesIndependently;
+			}
+		}
 
 		{
 			const FRHITextureCreateDesc Desc =
diff --git a/Engine/Source/Runtime/RHI/Private/RHI.cpp b/Engine/Source/Runtime/RHI/Private/RHI.cpp
--- a/Engine/Source/Runtime/RHI/Private/RHI.cpp
+++ b/Engine/Source/Runtime/RHI/Private/RHI.cpp
@@ -1797,8 +1797,20 @@ bool FRHITextureSRVCreateInfo::Validate(const FRHITextureDesc& TextureDesc, cons
 		TEXT("Failed to create SRV at mips %d-%d: the texture %s has only %d mip levels."),
 		TextureSRVDesc.MipLevel, (TextureSRVDesc.MipLevel + TextureSRVDesc.NumMipLevels), TextureName, TextureDesc.NumMips);
 
-	// Validate the array sloces
-	if (TextureDesc.IsTextureArray())
+	// Validate the array slices
+	if (TextureDesc.IsTextureCube() && TextureSRVDesc.DimensionOverride == ETextureDimension::Texture2DArray)
+	{
+		// Either TextureCube or TextureCubeArray, compute array size appropriately
+		int32 CubeArraySize = TextureDesc.Dimension == ETextureDimension::TextureCube ? ECubeFace::CubeFace_MAX : ECubeFace::CubeFace_MAX * TextureDesc.ArraySize;
+
+		ValidateResourceDesc((TextureSRVDesc.FirstArraySlice + TextureSRVDesc.NumArraySlices) <= CubeArraySize,
+			TEXT("Failed to create SRV at array slices %d-%d: the cube map texture %s has only %d slices."),
+			TextureSRVDesc.FirstArraySlice,
+			(TextureSRVDesc.FirstArraySlice + TextureSRVDesc.NumArraySlices),
+			TextureName,
+			CubeArraySize);
+	}
+	else if (TextureDesc.IsTextureArray())
 	{
 		ValidateResourceDesc((TextureSRVDesc.FirstArraySlice + TextureSRVDesc.NumArraySlices) <= TextureDesc.ArraySize,
 			TEXT("Failed to create SRV at array slices %d-%d: the texture array %s has only %d slices."),
diff --git a/Engine/Source/Runtime/RenderCore/Private/GenerateMips.cpp b/Engine/Source/Runtime/RenderCore/Private/GenerateMips.cpp
--- a/Engine/Source/Runtime/RenderCore/Private/GenerateMips.cpp
+++ b/Engine/Source/Runtime/RenderCore/Private/GenerateMips.cpp
@@ -124,12 +124,26 @@ void FGenerateMips::ExecuteRaster(FRDGBuilder& GraphBuilder, ERHIFeatureLevel::T
 	check(Texture);
 	check(Sampler);
 
+	const FRDGTextureDesc& TextureDesc = Texture->Desc;
+
 	auto ShaderMap = GetGlobalShaderMap(FeatureLevel);
 	TShaderMapRef<FGenerateMipsVS> VertexShader(ShaderMap);
 	TShaderMapRef<FGenerateMipsPS> PixelShader(ShaderMap);
 
-	const FRDGTextureDesc& TextureDesc = Texture->Desc;
+	int32 SliceCount = 1;
 
+	FRDGTextureSRVDesc SRVDesc(Texture);
+	SRVDesc.NumMipLevels = 1;
+
+	if (TextureDesc.Dimension == ETextureDimension::TextureCube)
+	{
+		SliceCount = ECubeFace::CubeFace_MAX;
+
+		SRVDesc.DimensionOverride = ETextureDimension::Texture2DArray;
+		SRVDesc.NumArraySlices = 1;
+	}
+
+	// Loop through each level of the mips that require creation and add a dispatch pass per level.
 	for (uint8 MipLevel = 1, MipCount = TextureDesc.NumMips; MipLevel < MipCount; ++MipLevel)
 	{
 		const uint32 InputMipLevel = MipLevel - 1;
@@ -138,35 +152,42 @@ void FGenerateMips::ExecuteRaster(FRDGBuilder& GraphBuilder, ERHIFeatureLevel::T
 			FMath::Max(TextureDesc.Extent.X >> MipLevel, 1),
 			FMath::Max(TextureDesc.Extent.Y >> MipLevel, 1));
 
-		FGenerateMipsPS::FParameters* PassParameters = GraphBuilder.AllocParameters<FGenerateMipsPS::FParameters>();
-		PassParameters->HalfTexelSize = FVector2f(0.5f / DestTextureSize.X, 0.5f / DestTextureSize.Y);
-		PassParameters->Level = InputMipLevel;
-		PassParameters->MipInSRV = GraphBuilder.CreateSRV(FRDGTextureSRVDesc::CreateForMipLevel(Texture, InputMipLevel));
-		PassParameters->MipSampler = Sampler;
-		PassParameters->RenderTargets[0] = FRenderTargetBinding(Texture, ERenderTargetLoadAction::ELoad, MipLevel);
+		SRVDesc.MipLevel = InputMipLevel;
 
-		GraphBuilder.AddPass(
-			RDG_EVENT_NAME("GenerateMips DestMipLevel=%d", MipLevel),
-			PassParameters,
-			ERDGPassFlags::Raster,
-			[VertexShader, PixelShader, PassParameters, DestTextureSize](FRHICommandList& RHICmdList)
+		for (int32 SliceIndex = 0; SliceIndex < SliceCount; ++SliceIndex)
 		{
-			RHICmdList.SetViewport(0.0f, 0.0f, 0.0f, (float)DestTextureSize.X, (float)DestTextureSize.Y, 1.0f);
-
-			FGraphicsPipelineStateInitializer GraphicsPSOInit;
-			GraphicsPSOInit.DepthStencilState = TStaticDepthStencilState<false, CF_Always>::GetRHI();
-			GraphicsPSOInit.RasterizerState = TStaticRasterizerState<>::GetRHI();
-			GraphicsPSOInit.BlendState = TStaticBlendStateWriteMask<CW_RGBA, CW_NONE, CW_NONE, CW_NONE, CW_NONE, CW_NONE, CW_NONE, CW_NONE>::GetRHI();
-			GraphicsPSOInit.PrimitiveType = PT_TriangleStrip;
-			GraphicsPSOInit.BoundShaderState.VertexDeclarationRHI = GFilterVertexDeclaration.VertexDeclarationRHI;
-			GraphicsPSOInit.BoundShaderState.VertexShaderRHI = VertexShader.GetVertexShader();
-			GraphicsPSOInit.BoundShaderState.PixelShaderRHI = PixelShader.GetPixelShader();
-			RHICmdList.ApplyCachedRenderTargets(GraphicsPSOInit);
-			SetGraphicsPipelineState(RHICmdList, GraphicsPSOInit, 0);
-			SetShaderParameters(RHICmdList, PixelShader, PixelShader.GetPixelShader(), *PassParameters);
-
-			FPixelShaderUtils::DrawFullscreenTriangle(RHICmdList, 1);
-		});
+			SRVDesc.FirstArraySlice = SliceIndex;
+
+			FGenerateMipsPS::FParameters* PassParameters = GraphBuilder.AllocParameters<FGenerateMipsPS::FParameters>();
+			PassParameters->HalfTexelSize = FVector2f(0.5f / DestTextureSize.X, 0.5f / DestTextureSize.Y);
+			PassParameters->Level = InputMipLevel;
+			PassParameters->MipInSRV = GraphBuilder.CreateSRV(SRVDesc);
+			PassParameters->MipSampler = Sampler;
+			PassParameters->RenderTargets[0] = FRenderTargetBinding(Texture, ERenderTargetLoadAction::ELoad, MipLevel, SliceCount > 1 ? SliceIndex : INDEX_NONE);
+
+			GraphBuilder.AddPass(
+				RDG_EVENT_NAME("GenerateMips DestMipLevel=%d Slice=%d", MipLevel, SliceIndex),
+				PassParameters,
+				ERDGPassFlags::Raster,
+				[VertexShader, PixelShader, PassParameters, DestTextureSize](FRHICommandList& RHICmdList)
+			{
+				RHICmdList.SetViewport(0.0f, 0.0f, 0.0f, (float)DestTextureSize.X, (float)DestTextureSize.Y, 1.0f);
+
+				FGraphicsPipelineStateInitializer GraphicsPSOInit;
+				GraphicsPSOInit.DepthStencilState = TStaticDepthStencilState<false, CF_Always>::GetRHI();
+				GraphicsPSOInit.RasterizerState = TStaticRasterizerState<>::GetRHI();
+				GraphicsPSOInit.BlendState = TStaticBlendStateWriteMask<CW_RGBA, CW_NONE, CW_NONE, CW_NONE, CW_NONE, CW_NONE, CW_NONE, CW_NONE>::GetRHI();
+				GraphicsPSOInit.PrimitiveType = PT_TriangleStrip;
+				GraphicsPSOInit.BoundShaderState.VertexDeclarationRHI = GFilterVertexDeclaration.VertexDeclarationRHI;
+				GraphicsPSOInit.BoundShaderState.VertexShaderRHI = VertexShader.GetVertexShader();
+				GraphicsPSOInit.BoundShaderState.PixelShaderRHI = PixelShader.GetPixelShader();
+				RHICmdList.ApplyCachedRenderTargets(GraphicsPSOInit);
+				SetGraphicsPipelineState(RHICmdList, GraphicsPSOInit, 0);
+				SetShaderParameters(RHICmdList, PixelShader, PixelShader.GetPixelShader(), *PassParameters);
+
+				FPixelShaderUtils::DrawFullscreenTriangle(RHICmdList, 1);
+			});
+		}
 	}
 }
 
@@ -186,6 +207,22 @@ void FGenerateMips::ExecuteCompute(FRDGBuilder& GraphBuilder, ERHIFeatureLevel::
 	PermutationVector.Set<FGenerateMipsCS::FGenMipsSwizzle>(bMipsSwizzle);
 	TShaderMapRef<FGenerateMipsCS> ComputeShader(GetGlobalShaderMap(FeatureLevel), PermutationVector);
 
+	int32 SliceCount = 1;
+
+	FRDGTextureSRVDesc SRVDesc(Texture);
+	FRDGTextureUAVDesc UAVDesc(Texture);
+	SRVDesc.NumMipLevels = 1;
+
+	if (TextureDesc.Dimension == ETextureDimension::TextureCube)
+	{
+		SliceCount = ECubeFace::CubeFace_MAX;
+
+		SRVDesc.DimensionOverride = ETextureDimension::Texture2DArray;
+		SRVDesc.NumArraySlices = 1;
+		UAVDesc.DimensionOverride = ETextureDimension::Texture2DArray;
+		UAVDesc.NumArraySlices = 1;
+	}
+
 	// Loop through each level of the mips that require creation and add a dispatch pass per level.
 	for (uint8 MipLevel = 1, MipCount = TextureDesc.NumMips; MipLevel < MipCount; ++MipLevel)
 	{
@@ -193,18 +230,27 @@ void FGenerateMips::ExecuteCompute(FRDGBuilder& GraphBuilder, ERHIFeatureLevel::
 			FMath::Max(TextureDesc.Extent.X >> MipLevel, 1),
 			FMath::Max(TextureDesc.Extent.Y >> MipLevel, 1));
 
-		FGenerateMipsCS::FParameters* PassParameters = GraphBuilder.AllocParameters<FGenerateMipsCS::FParameters>();
-		PassParameters->TexelSize  = FVector2f(1.0f / DestTextureSize.X, 1.0f / DestTextureSize.Y);
-		PassParameters->MipInSRV   = GraphBuilder.CreateSRV(FRDGTextureSRVDesc::CreateForMipLevel(Texture, MipLevel - 1));
-		PassParameters->MipOutUAV  = GraphBuilder.CreateUAV(FRDGTextureUAVDesc(Texture, MipLevel));
-		PassParameters->MipSampler = Sampler;
+		SRVDesc.MipLevel = (int8)(MipLevel - 1);
+		UAVDesc.MipLevel = (int8)(MipLevel);
 
-		FComputeShaderUtils::AddPass(
-			GraphBuilder,
-			RDG_EVENT_NAME("GenerateMips DestMipLevel=%d", MipLevel),
-			ComputeShader,
-			PassParameters,
-			FComputeShaderUtils::GetGroupCount(DestTextureSize, FComputeShaderUtils::kGolden2DGroupSize));
+		for (int32 SliceIndex = 0; SliceIndex < SliceCount; ++SliceIndex)
+		{
+			SRVDesc.FirstArraySlice = SliceIndex;
+			UAVDesc.FirstArraySlice = SliceIndex;
+
+			FGenerateMipsCS::FParameters* PassParameters = GraphBuilder.AllocParameters<FGenerateMipsCS::FParameters>();
+			PassParameters->TexelSize = FVector2f(1.0f / DestTextureSize.X, 1.0f / DestTextureSize.Y);
+			PassParameters->MipInSRV = GraphBuilder.CreateSRV(SRVDesc);
+			PassParameters->MipOutUAV = GraphBuilder.CreateUAV(UAVDesc);
+			PassParameters->MipSampler = Sampler;
+
+			FComputeShaderUtils::AddPass(
+				GraphBuilder,
+				RDG_EVENT_NAME("GenerateMips DestMipLevel=%d Slice=%d", MipLevel, SliceIndex),
+				ComputeShader,
+				PassParameters,
+				FComputeShaderUtils::GetGroupCount(DestTextureSize, FComputeShaderUtils::kGolden2DGroupSize));
+		}
 	}
 }
 
diff --git a/Engine/Source/Runtime/Renderer/Private/SceneCaptureRendering.cpp b/Engine/Source/Runtime/Renderer/Private/SceneCaptureRendering.cpp
--- a/Engine/Source/Runtime/Renderer/Private/SceneCaptureRendering.cpp
+++ b/Engine/Source/Runtime/Renderer/Private/SceneCaptureRendering.cpp
@@ -1348,6 +1348,9 @@ void FScene::UpdateSceneCaptureContents(USceneCaptureComponentCube* CaptureCompo
 			}
 		}
 
+		const bool bGenerateMips = TextureTarget->bAutoGenerateMips;
+		FGenerateMipsParams GenerateMipsParams{ TextureTarget->MipsSamplerFilter == TF_Nearest ? SF_Point : (TextureTarget->MipsSamplerFilter == TF_Trilinear ? SF_Trilinear : SF_Bilinear), AM_Clamp, AM_Clamp };
+
 		FIntPoint CaptureSize(TextureTarget->GetSurfaceWidth(), TextureTarget->GetSurfaceHeight());
 		const float FOV = 90 * (float)PI / 360.0f;
 
@@ -1409,7 +1412,7 @@ void FScene::UpdateSceneCaptureContents(USceneCaptureComponentCube* CaptureCompo
 				UE::RenderCommandPipe::FSyncScope SyncScope;
 
 				ENQUEUE_RENDER_COMMAND(CaptureCommand)(
-					[SceneRenderer, TextureRenderTarget, EventName, TargetFace](FRHICommandListImmediate& RHICmdList)
+					[SceneRenderer, TextureRenderTarget, EventName, TargetFace, bGenerateMips, GenerateMipsParams](FRHICommandListImmediate& RHICmdList)
 					{
 #if WITH_EDITOR
 						// Scene renderer may be deleted in UpdateSceneCaptureContent_RenderThread, grab view state pointer first
@@ -1418,10 +1421,13 @@ void FScene::UpdateSceneCaptureContents(USceneCaptureComponentCube* CaptureCompo
 
 						VISUALIZE_TEXTURE_BEGIN_VIEW(SceneRenderer->FeatureLevel, SceneRenderer->Views[0].GetViewKey(), *EventName, true);
 
+						// We need to generate mips on last cube face
+						const bool bLastCubeFace = ((int32)TargetFace == (int32)ECubeFace::CubeFace_MAX - 1);
+
 						FRHICopyTextureInfo CopyInfo;
 						CopyInfo.DestSliceIndex = TargetFace;
 						UpdateSceneCaptureContent_RenderThread(RHICmdList, SceneRenderer, TextureRenderTarget, TextureRenderTarget, EventName,
-							TConstArrayView<FRHICopyTextureInfo>(&CopyInfo, 1), false, FGenerateMipsParams(), true, false);
+							TConstArrayView<FRHICopyTextureInfo>(&CopyInfo, 1), bGenerateMips && bLastCubeFace, GenerateMipsParams, true, false);
 
 						VISUALIZE_TEXTURE_END_VIEW();
 
@@ -1587,7 +1593,7 @@ void FScene::UpdateSceneCaptureContents(USceneCaptureComponentCube* CaptureCompo
 			UE::RenderCommandPipe::FSyncScope SyncScope;
 
 			ENQUEUE_RENDER_COMMAND(CaptureAllCubeFaces)(
-				[SceneRenderer, CubeFaceTarget, TextureRenderTarget, EventName, CaptureSize](FRHICommandListImmediate& RHICmdList)
+				[SceneRenderer, CubeFaceTarget, TextureRenderTarget, EventName, CaptureSize, bGenerateMips, GenerateMipsParams](FRHICommandListImmediate& RHICmdList)
 				{
 					TStaticArray<FRHICopyTextureInfo, (int32)ECubeFace::CubeFace_MAX> CopyInfos;
 					for (int32 faceidx = 0; faceidx < (int32)ECubeFace::CubeFace_MAX; faceidx++)
@@ -1613,7 +1619,7 @@ void FScene::UpdateSceneCaptureContents(USceneCaptureComponentCube* CaptureCompo
 					VISUALIZE_TEXTURE_BEGIN_VIEW(SceneRenderer->FeatureLevel, SceneRenderer->Views[0].GetViewKey(), *EventName, true);
 
 					UpdateSceneCaptureContent_RenderThread(RHICmdList, SceneRenderer, CubeFaceTarget, TextureRenderTarget, EventName,
-						CopyInfos, false, FGenerateMipsParams(), true, false);
+						CopyInfos, bGenerateMips, GenerateMipsParams, true, false);
 
 					VISUALIZE_TEXTURE_END_VIEW();
 
diff --git a/Engine/Source/Runtime/Windows/D3D11RHI/Private/D3D11Texture.cpp b/Engine/Source/Runtime/Windows/D3D11RHI/Private/D3D11Texture.cpp
--- a/Engine/Source/Runtime/Windows/D3D11RHI/Private/D3D11Texture.cpp
+++ b/Engine/Source/Runtime/Windows/D3D11RHI/Private/D3D11Texture.cpp
@@ -874,7 +874,7 @@ void FD3D11DynamicRHI::RHIGenerateMips(FRHITexture* TextureRHI)
 {
 	FD3D11Texture* Texture = ResourceCast(TextureRHI);
 	// Surface must have been created with D3D11_BIND_RENDER_TARGET for GenerateMips to work
-	check(Texture->GetShaderResourceView() && Texture->GetRenderTargetView(0, -1));
+	check(Texture->GetShaderResourceView() && Texture->GetRenderTargetView(0, 0));
 	Direct3DDeviceIMContext->GenerateMips(Texture->GetShaderResourceView());
 
 	GPUProfilingData.RegisterGPUWork(0);
