Shader Compiler:  Fix use-after-free in FShaderCompilerEnvironment::Merge.  Need to not overwrite MemberNameBuffer (FUniformBufferEntry) when appending items to UniformBufferMap, as items in the FUniformResourceEntry array contain internal pointers to the name buffer.  Loop over other container's items similar to TMap::Append, but check if the item is already present before adding.

#rnx

diff --git a/Engine/Source/Runtime/RenderCore/Private/ShaderCore.cpp b/Engine/Source/Runtime/RenderCore/Private/ShaderCore.cpp
--- a/Engine/Source/Runtime/RenderCore/Private/ShaderCore.cpp
+++ b/Engine/Source/Runtime/RenderCore/Private/ShaderCore.cpp
@@ -890,7 +890,18 @@ void FShaderCompilerEnvironment::Merge(const FShaderCompilerEnvironment& Other)
 
 	CompilerFlags.Append(Other.CompilerFlags);
 	ResourceTableMap.Append(Other.ResourceTableMap);
-	UniformBufferMap.Append(Other.UniformBufferMap);
+	{
+		// Append, but don't overwrite the value of existing elements, to preserve MemberNameBuffer which is pointed to by ResourceTableMap entries
+		UniformBufferMap.Reserve(UniformBufferMap.Num() + Other.UniformBufferMap.Num());
+		for (auto& Pair : Other.UniformBufferMap)
+		{
+			uint32 KeyHash = GetTypeHash(Pair.Key);
+			if (!UniformBufferMap.ContainsByHash(KeyHash, Pair.Key))
+			{
+				UniformBufferMap.AddByHash(KeyHash, Pair.Key, Pair.Value);
+			}
+		}
+	}
 	PRAGMA_DISABLE_DEPRECATION_WARNINGS	// FShaderCompilerDefinitions will be made internal in the future, marked deprecated until then
 	checkf(Definitions.IsValid(), TEXT("Merge is not supported on FShaderCompilerEnvironment in hashing mode"));
 	Definitions->Merge(*Other.Definitions);
