[shaders]
- modify the input hash debug dump mechanism to output an empty "debughash_<hash>" file instead of a txt file with the hash in contents, and always dump these files for the instance of the job that actually compiled
- the existing cvar will now just make it so these files are also dumped for jobs which hit in DDC or the job cache; we don't do this by default so there's only a single match for the debug hash for any given shader normally and it is inside the folder containing the full debug info, including those artifacts which are only output as a side effect of the compile step
- add the same hash as the first line in the stripped source code, so "debughash_<hash>" can be used as a search term in Everything to quickly find debug info associated with a shader (i.e. when looking at a capture in renderdoc or similar)

note: this is a resubmit with fixes for mac/linux issues (avoiding printfs and using string builders instead; since wchar_t and TCHAR are not the same size on these platforms using %ls does not work), and a further fix for problems encountered with source compression when using wide chars in preprocessing instead of ansi chars.

diff --git a/Engine/Source/Developer/ShaderCompilerCommon/Private/ShaderCompilerCommon.cpp b/Engine/Source/Developer/ShaderCompilerCommon/Private/ShaderCompilerCommon.cpp
--- a/Engine/Source/Developer/ShaderCompilerCommon/Private/ShaderCompilerCommon.cpp
+++ b/Engine/Source/Developer/ShaderCompilerCommon/Private/ShaderCompilerCommon.cpp
@@ -1847,10 +1847,29 @@ namespace UE::ShaderCompilerCommon
 				FFileHelper::SaveStringToFile(Merged, *GetDebugFileName(Input, Options, TEXT("Diagnostics.txt")), FFileHelper::EEncodingOptions::ForceAnsi);
 			}
 		}
+		
+		// delete old DebugHash_* files so we don't clutter the debug info folder (these change every time the deadstripped source code changes)
+		IFileManager::Get().IterateDirectory(*Input.DumpDebugInfoPath, [](const TCHAR* FilenameOrDirectory, bool bIsDirectory)
+			{
+				if (!bIsDirectory)
+				{
+					FStringView Filename = FPathViews::GetCleanFilename(FilenameOrDirectory);
+					if (Filename.StartsWith(GetShaderSourceDebugHashPrefixWide()))
+					{
+						IFileManager::Get().Delete(FilenameOrDirectory);
+					}
+				}
+				return true;
+			});
 
-		if (EnumHasAnyFlags(Input.DebugInfoFlags, EShaderDebugInfoFlags::InputHash))
+		if (EnumHasAnyFlags(Input.DebugInfoFlags, EShaderDebugInfoFlags::InputHash) || Output.CompileTime > 0.0f)
 		{
-			FFileHelper::SaveStringToFile(LexToString(Input.Hash), *GetDebugFileName(Input, Options, TEXT("InputHash.txt")), FFileHelper::EEncodingOptions::ForceAnsi);
+			// If compile time was > 0, this was the copy of the job that actually compiled; we write an empty file with the shader hash in it so it can be
+			// found easily using the ShaderHash comment printed on the first line of the stripped source code.
+			// We don't do this for cache hits or duplicate jobs unless explicitly requested (via the InputHash debug info flags), so that _all_ debug artifacts
+			// (including those only generated by the compile process) are available in the folder containing this file.
+			FString InputHashStr = LexToString(Input.Hash);
+			FFileHelper::SaveStringToFile(FStringView(), *GetDebugFileName(Input, Options, GetShaderSourceDebugHashPrefixWide().GetData(), *InputHashStr));
 		}
 
 		if (EnumHasAnyFlags(Input.DebugInfoFlags, EShaderDebugInfoFlags::ShaderCodeBinary))
diff --git a/Engine/Source/Runtime/Engine/Private/ShaderCompiler/ShaderCompiler.cpp b/Engine/Source/Runtime/Engine/Private/ShaderCompiler/ShaderCompiler.cpp
--- a/Engine/Source/Runtime/Engine/Private/ShaderCompiler/ShaderCompiler.cpp
+++ b/Engine/Source/Runtime/Engine/Private/ShaderCompiler/ShaderCompiler.cpp
@@ -404,10 +404,9 @@ private:
 /** Map element type for job cache */
 struct FShaderJobData
 {
-	using FJobInputHash = FBlake3Hash;
 	using FJobOutputHash = FBlake3Hash;
 
-	FJobInputHash InputHash;
+	FShaderCompilerInputHash InputHash;
 
 	/** Output hash will be zero if output data has not been written yet, or can be cleared if output data has been removed */
 	FJobOutputHash OutputHash;
@@ -455,8 +454,8 @@ public:
 		Reserve(FShaderJobDataBlock::BlockSize);
 	}
 
-	FShaderJobData* Find(const FShaderJobData::FJobInputHash& Key);
-	FShaderJobCacheRef FindOrAdd(const FShaderJobData::FJobInputHash& Key);
+	FShaderJobData* Find(const FShaderCompilerInputHash& Key);
+	FShaderJobCacheRef FindOrAdd(const FShaderCompilerInputHash& Key);
 
 	FORCEINLINE int32 Num() const
 	{
@@ -762,7 +761,6 @@ public:
 	int32 GetPendingJobs(EShaderCompilerWorkerType InWorkerType, EShaderCompileJobPriority InPriority, int32 MinNumJobs, int32 MaxNumJobs, TArray<FShaderCommonCompileJobPtr>& OutJobs);
 
 private:
-	using FJobInputHash = FShaderCommonCompileJob::FInputHash;
 	using FJobOutputHash = FBlake3Hash;
 	using FJobCachedOutput = FSharedBuffer;
 	using FStoredOutput = FShaderJobCacheStoredOutput;
@@ -778,10 +776,10 @@ private:
 	void InternalSetPriority(FShaderCommonCompileJob* Job, EShaderCompileJobPriority InPriority);
 
 	/** Looks for or adds an entry for the given hash in the cache.  Returns cached output if it exists, or may initialize DDC request if one has been issued. */
-	FShaderJobCacheRef FindOrAdd(const FJobInputHash& Hash, EShaderCompileJobPriority JobPriority, const bool bCheckDDC, TPimplPtr<UE::DerivedData::FRequestOwner>& InoutRequestOwner, FJobCachedOutput*& OutCachedOutput);
+	FShaderJobCacheRef FindOrAdd(const FShaderCompilerInputHash& Hash, EShaderCompileJobPriority JobPriority, const bool bCheckDDC, TPimplPtr<UE::DerivedData::FRequestOwner>& InoutRequestOwner, FJobCachedOutput*& OutCachedOutput);
 
 	/** Find an existing item in the cache. */
-	FShaderJobData* Find(const FJobInputHash& Hash);
+	FShaderJobData* Find(const FShaderCompilerInputHash& Hash);
 
 	/** Add a reference to a duplicate job (to the DuplicateJobs array) */
 	void AddDuplicateJob(FShaderCommonCompileJob* DuplicateJob);
@@ -790,7 +788,7 @@ private:
 	void RemoveDuplicateJob(FShaderCommonCompileJob* DuplicateJob);
 
 	/** Adds a job output to the cache */
-	void AddJobOutput(FShaderJobData& JobData, const FShaderCommonCompileJob* FinishedJob, const FJobInputHash& Hash, const FJobCachedOutput& Contents, int32 InitialHitCount, const bool bAddToDDC);
+	void AddJobOutput(FShaderJobData& JobData, const FShaderCommonCompileJob* FinishedJob, const FShaderCompilerInputHash& Hash, const FJobCachedOutput& Contents, int32 InitialHitCount, const bool bAddToDDC);
 
 	/** Returns memory used by the cache*/
 	uint64 GetAllocatedMemory() const;
@@ -994,7 +992,7 @@ static FShaderJobData& GetShaderJobData(const FShaderJobCacheRef& CacheRef)
 	return CacheRef.Block->Data[CacheRef.IndexInBlock];
 }
 
-FShaderJobData* FShaderJobDataMap::Find(const FShaderJobData::FJobInputHash& Key)
+FShaderJobData* FShaderJobDataMap::Find(const FShaderCompilerInputHash& Key)
 {
 	// Search for key with linear probing
 	for (uint32 TableIndex = GetTypeHash(Key) & HashTableMask; HashTable[TableIndex] != INDEX_NONE; TableIndex = (TableIndex + 1) & HashTableMask)
@@ -1007,7 +1005,7 @@ FShaderJobData* FShaderJobDataMap::Find(const FShaderJobData::FJobInputHash& Key
 	return nullptr;
 }
 
-FShaderJobCacheRef FShaderJobDataMap::FindOrAdd(const FShaderJobData::FJobInputHash& Key)
+FShaderJobCacheRef FShaderJobDataMap::FindOrAdd(const FShaderCompilerInputHash& Key)
 {
 	// Search for key with linear probing
 	uint32 TableIndex;
@@ -1621,7 +1619,7 @@ void FShaderJobCache::SubmitJob(FShaderCommonCompileJob* Job)
 	{
 		bJobCacheLocked = true;
 
-		const FShaderCommonCompileJob::FInputHash& InputHash = Job->GetInputHash();
+		const FShaderCompilerInputHash& InputHash = Job->GetInputHash();
 
 		const bool bCheckDDC = GShaderCompilerPerShaderDDCGlobal || !(Job->bIsDefaultMaterial || Job->bIsGlobalShader);
 
@@ -1704,7 +1702,7 @@ void FShaderJobCache::SubmitJob(FShaderCommonCompileJob* Job)
 	else if (ShaderCompiler::IsJobCacheDebugValidateEnabled())
 	{
 		FSharedBuffer* ExistingOutput;
-		const FShaderCommonCompileJob::FInputHash& InputHash = Job->GetInputHash();
+		const FShaderCompilerInputHash& InputHash = Job->GetInputHash();
 		const bool bCheckDDC = !(Job->bIsDefaultMaterial || Job->bIsGlobalShader);
 		JobLock.WriteLock();
 		Job->JobCacheRef = FindOrAdd(InputHash, Job->Priority, bCheckDDC, Job->RequestOwner, ExistingOutput);
@@ -1872,7 +1870,7 @@ void FShaderJobCache::AddToCacheAndProcessPending(FShaderCommonCompileJob* Finis
 
 	ensureMsgf(FinishedJob->bInputHashSet, TEXT("Finished job didn't have input hash set, was shader compiler jobs cache toggled runtime?"));
 
-	const FShaderCommonCompileJob::FInputHash& InputHash = FinishedJob->GetInputHash();
+	const FShaderCompilerInputHash& InputHash = FinishedJob->GetInputHash();
 	TArray<uint8> Output;
 	FMemoryWriter Writer(Output);
 	FinishedJob->SerializeOutput(Writer);
@@ -10867,7 +10865,7 @@ namespace
 }
 #endif
 
-FShaderJobCacheRef FShaderJobCache::FindOrAdd(const FJobInputHash& Hash, EShaderCompileJobPriority JobPriority, const bool bCheckDDC, TPimplPtr<UE::DerivedData::FRequestOwner>& InoutRequestOwner, FJobCachedOutput*& OutCachedOutput)
+FShaderJobCacheRef FShaderJobCache::FindOrAdd(const FShaderCompilerInputHash& Hash, EShaderCompileJobPriority JobPriority, const bool bCheckDDC, TPimplPtr<UE::DerivedData::FRequestOwner>& InoutRequestOwner, FJobCachedOutput*& OutCachedOutput)
 {
 	LLM_SCOPE_BYTAG(ShaderCompiler);
 
@@ -11146,7 +11144,7 @@ FShaderJobCacheRef FShaderJobCache::FindOrAdd(const FJobInputHash& Hash, EShader
 	return JobCacheRef;
 }
 
-FShaderJobData* FShaderJobCache::Find(const FJobInputHash& Hash)
+FShaderJobData* FShaderJobCache::Find(const FShaderCompilerInputHash& Hash)
 {
 	check(ShaderCompiler::IsJobCacheEnabled());
 
@@ -11207,7 +11205,7 @@ FShaderJobCache::~FShaderJobCache()
 	}
 }
 
-void FShaderJobCache::AddJobOutput(FShaderJobData& JobData, const FShaderCommonCompileJob* FinishedJob, const FJobInputHash& Hash, const FJobCachedOutput& Contents, int32 InitialHitCount, const bool bAddToDDC)
+void FShaderJobCache::AddJobOutput(FShaderJobData& JobData, const FShaderCommonCompileJob* FinishedJob, const FShaderCompilerInputHash& Hash, const FJobCachedOutput& Contents, int32 InitialHitCount, const bool bAddToDDC)
 {
 	TRACE_CPUPROFILER_EVENT_SCOPE(FShaderJobCache::Add);
 
diff --git a/Engine/Source/Runtime/RenderCore/Private/ShaderCore.cpp b/Engine/Source/Runtime/RenderCore/Private/ShaderCore.cpp
--- a/Engine/Source/Runtime/RenderCore/Private/ShaderCore.cpp
+++ b/Engine/Source/Runtime/RenderCore/Private/ShaderCore.cpp
@@ -1502,6 +1502,17 @@ public:
 			// and blank lines to improve deduplication (and populates data required to remap diagnostic messages to correct line numbers)
 			Job.PreprocessOutput.StripCode(Job.Input.NeedsOriginalShaderSource());
 
+			FShaderCompilerInputHash Hash = Job.GetInputHash();
+			// Replace the placeholder debug hash value appended in StripCode with the real job input hash
+			FShaderSource::FViewType DebugHashStr = GetShaderSourceDebugHashPrefix();
+			FShaderSource::FViewType SourceView = Job.PreprocessOutput.GetSourceView();
+			int32 DebugHashLoc = SourceView.Find(DebugHashStr) + DebugHashStr.Len();
+			int32 NewlineLoc = SourceView.Find(SHADER_SOURCE_VIEWLITERAL("\n"), DebugHashLoc);
+			FShaderSource::TStringBuilder<2 * sizeof(FShaderCompilerInputHash::ByteArray) + 1> HashStr;
+			HashStr << Hash;
+			check(NewlineLoc - DebugHashLoc == HashStr.Len());
+			FMemory::Memcpy(Job.PreprocessOutput.EditSource().GetData() + DebugHashLoc, HashStr.GetData(), sizeof(FShaderSource::CharType) * HashStr.Len());
+
 			// always compress the code after stripping to minimize memory footprint
 			Job.PreprocessOutput.CompressCode();
 		}
@@ -3758,26 +3769,26 @@ FArchive& operator<<(FArchive& Ar, FShaderCompilerInput& Input)
 	return Ar;
 }
 
-FShaderCommonCompileJob::FInputHash FShaderPipelineCompileJob::GetInputHash()
+FShaderCompilerInputHash FShaderPipelineCompileJob::GetInputHash()
 {
 	if (bInputHashSet)
 	{
 		return InputHash;
 	}
-	static_assert(sizeof(FShaderCommonCompileJob::FInputHash) == 32);
+	static_assert(sizeof(FShaderCompilerInputHash) == 32);
 	int256 CombinedHash = 0u;
 	for (int32 Index = 0; Index < StageJobs.Num(); ++Index)
 	{
 		if (StageJobs[Index])
 		{
-			const FShaderCommonCompileJob::FInputHash StageHash = StageJobs[Index]->GetInputHash();
-			const FShaderCommonCompileJob::FInputHash::ByteArray& StageHashBytes = StageHash.GetBytes();
+			const FShaderCompilerInputHash StageHash = StageJobs[Index]->GetInputHash();
+			const FShaderCompilerInputHash::ByteArray& StageHashBytes = StageHash.GetBytes();
 			static_assert(sizeof(StageHashBytes) == sizeof(int256));
 			CombinedHash += int256(StageHashBytes, sizeof(StageHashBytes));
 		}
 	}
 
-	InputHash = FShaderCommonCompileJob::FInputHash(reinterpret_cast<FShaderCommonCompileJob::FInputHash::ByteArray&>(*CombinedHash.GetBits()));
+	InputHash = FShaderCompilerInputHash(reinterpret_cast<FShaderCompilerInputHash::ByteArray&>(*CombinedHash.GetBits()));
 	bInputHashSet = true;
 	return InputHash;
 }
@@ -3818,7 +3829,7 @@ struct FShaderVirtualFileContents
 	{}
 };
 
-FShaderCommonCompileJob::FInputHash FShaderCompileJob::GetInputHash()
+FShaderCompilerInputHash FShaderCompileJob::GetInputHash()
 {
 	if (bInputHashSet)
 	{
@@ -4064,7 +4075,12 @@ void FShaderCompileJob::SerializeOutput(FArchive& Ar)
 
 void FShaderCompileJob::OnComplete()
 {
-	const IShaderFormat* ShaderFormat = GetTargetPlatformManagerRef().FindShaderFormat(Input.ShaderFormat);
+	if (Input.NeedsOriginalShaderSource())
+	{
+		// Decompress the code if needed by debug info or source extraction
+		PreprocessOutput.DecompressCode();
+	}
+
 	// For jobs using the preprocessed cache, we need to remap error messages whether or not the job was actually the one that ran
 	// the compilation step. In addition since we always run preprocessing we set the total preprocess time accordingly.
 	if (Input.bCachePreprocessed)
@@ -4073,12 +4089,6 @@ void FShaderCompileJob::OnComplete()
 		Output.PreprocessTime = PreprocessOutput.ElapsedTime;
 	}
 
-	if (Input.NeedsOriginalShaderSource())
-	{
-		// Decompress the code if needed by debug info or source extraction
-		PreprocessOutput.DecompressCode();
-	}
-
 	// dump debug info for the job at this point if the preprocessed cache is enabled
 	// this ensures we get debug output for all jobs, including those that were found in the job cache,
 	// or matched another in-flight job's hash and so could share its results
@@ -4088,6 +4098,7 @@ void FShaderCompileJob::OnComplete()
 		// (jobs deserialized from the cache/wait list/ddc will have a compiletime of 0.0)
 		&& (CVarDumpDebugInfoForCacheHits.GetValueOnAnyThread() || Output.CompileTime > 0.0f))
 	{
+		const IShaderFormat* ShaderFormat = GetTargetPlatformManagerRef().FindShaderFormat(Input.ShaderFormat);
 		if (SecondaryPreprocessOutput.IsValid() && SecondaryOutput.IsValid())
 		{
 			ShaderFormat->OutputDebugData(Input, PreprocessOutput, *SecondaryPreprocessOutput, Output, *SecondaryOutput);
diff --git a/Engine/Source/Runtime/RenderCore/Private/ShaderPreprocessTypes.cpp b/Engine/Source/Runtime/RenderCore/Private/ShaderPreprocessTypes.cpp
--- a/Engine/Source/Runtime/RenderCore/Private/ShaderPreprocessTypes.cpp
+++ b/Engine/Source/Runtime/RenderCore/Private/ShaderPreprocessTypes.cpp
@@ -6,13 +6,62 @@
 #include "Containers/AnsiString.h"
 #include "Templates/UnrealTemplate.h"
 
-const FShaderSource::CharType* FilenameSentinel = SHADER_SOURCE_LITERAL("__UE_FILENAME_SENTINEL__");
-static const int FilenameSentinelLen = FShaderSource::FCStringType::Strlen(FilenameSentinel);
-static const FShaderSource::FStringType LineDirectiveSentinel = FShaderSource::FStringType::Printf(SHADER_SOURCE_LITERAL("#line 1 \"%s\"\n"), FilenameSentinel);
+struct FPreprocessConstants
+{
+	FPreprocessConstants()
+	{
+		// Add a placeholder comment to the header to allocate enough space for the "debug hash" (which
+		// will be used to associate shader source back to debug output). This will be replaced inline
+		// by code calling into the preprocessing after the fact (necessary since the hash of the preprocessed
+		// code itself is used in constructing this debug hash).
+		StrippedCodeHeaderBuilder << DebugHashCommentStart << DebugHashPrefix << FShaderCompilerInputHash() << DebugHashCommentEnd;
+		// Also add a sentinel line directive so we get the correct line offsets when remapping error/warning messages.
+		StrippedCodeHeaderBuilder << LineDirectiveStart << FilenameSentinel << LineDirectiveEnd;
+
+		check(StrippedCodeHeaderBuilder.GetAllocatedSize() == 0);
+
+		StrippedCodeHeader = { StrippedCodeHeaderBuilder.GetData(), StrippedCodeHeaderBuilder.Len() };
+	}
+
+	static constexpr FShaderSource::FViewType DebugHashPrefix = SHADER_SOURCE_VIEWLITERAL("DebugHash_");
+	static constexpr FStringView DebugHashPrefixWide = TEXTVIEW("DebugHash_");
+	static constexpr FShaderSource::FViewType FilenameSentinel = SHADER_SOURCE_VIEWLITERAL("__UE_FILENAME_SENTINEL");
+
+	FShaderSource::FViewType StrippedCodeHeader;
+
+private:
+
+	static constexpr FShaderSource::FViewType DebugHashCommentStart = SHADER_SOURCE_VIEWLITERAL("// ");
+	static constexpr FShaderSource::FViewType DebugHashCommentEnd = SHADER_SOURCE_VIEWLITERAL("\n");
+	static constexpr FShaderSource::FViewType LineDirectiveStart = SHADER_SOURCE_VIEWLITERAL("#line 1 \"");
+	static constexpr FShaderSource::FViewType LineDirectiveEnd = SHADER_SOURCE_VIEWLITERAL("\"\n");
+
+	static constexpr int32 HeaderLen =
+		DebugHashCommentStart.Len() +
+		DebugHashPrefix.Len() +
+		2 * sizeof(FShaderCompilerInputHash::ByteArray) + // size of input hash converted to a hex string
+		LineDirectiveStart.Len() +
+		FilenameSentinel.Len() +
+		LineDirectiveEnd.Len();
+
+	FShaderSource::TStringBuilder<HeaderLen> StrippedCodeHeaderBuilder;
+};
+
+static FPreprocessConstants GPreprocessConstants;
+
+FShaderSource::FViewType GetShaderSourceDebugHashPrefix()
+{
+	return GPreprocessConstants.DebugHashPrefix;
+}
+
+FStringView GetShaderSourceDebugHashPrefixWide()
+{
+	return GPreprocessConstants.DebugHashPrefixWide;
+}
 
 void FShaderDiagnosticRemapper::Remap(FShaderCompilerError& Diagnostic) const
 {
-	if (!Diagnostic.ErrorLineString.IsEmpty() && !Diagnostic.ErrorVirtualFilePath.IsEmpty() && Diagnostic.ErrorVirtualFilePath == FilenameSentinel)
+	if (!Diagnostic.ErrorLineString.IsEmpty() && !Diagnostic.ErrorVirtualFilePath.IsEmpty() && Diagnostic.ErrorVirtualFilePath == GPreprocessConstants.FilenameSentinel.GetData())
 	{
 		// only attempt remapping if the filename sentinel matches; this will bypass remapping in the case where compilation was retried on FXC
 		// with a DXC precompile (this compilation will report error messages relative to an intermediate hlsl file, which is not a change in behaviour)
@@ -39,7 +88,7 @@ void FShaderDiagnosticRemapper::Remap(FShaderCompilerError& Diagnostic) const
 		Diagnostic.ErrorVirtualFilePath = RemapData.Filename;
 	}
 
-	int32 FilenameIndex = Diagnostic.StrippedErrorMessage.Find(FilenameSentinel);
+	int32 FilenameIndex = Diagnostic.StrippedErrorMessage.Find(GPreprocessConstants.FilenameSentinel.GetData());
 	// if we find our sentinel filename in the diagnostic message it needs remapping
 	// need to loop in case the message contains the filename multiple times
 	while (FilenameIndex >= 0)
@@ -47,7 +96,7 @@ void FShaderDiagnosticRemapper::Remap(FShaderCompilerError& Diagnostic) const
 		const FString& OriginalMessage = Diagnostic.StrippedErrorMessage;
 		// Parse the line number from the message and record start and end indices
 		// Code assumes that the next integral number found in the message string after the filename is always line number
-		int32 LineNumberStart = FilenameIndex + FilenameSentinelLen;
+		int32 LineNumberStart = FilenameIndex + GPreprocessConstants.FilenameSentinel.Len();
 		while (!FChar::IsDigit(OriginalMessage[LineNumberStart]))
 		{
 			++LineNumberStart;
@@ -63,18 +112,18 @@ void FShaderDiagnosticRemapper::Remap(FShaderCompilerError& Diagnostic) const
 
 		// Assume line number is the same number of digits for simplicity in reserve allocation size;
 		// it's an upper bound but in most practical cases worst case we're allocating an extra byte
-		int32 RemappedLen = OriginalMessage.Len() - FilenameSentinelLen + RemapData.Filename.Len();
+		int32 RemappedLen = OriginalMessage.Len() - GPreprocessConstants.FilenameSentinel.Len() + RemapData.Filename.Len();
 		FString RemappedMessage;
 		RemappedMessage.Reserve(RemappedLen);
 		RemappedMessage.AppendChars(OriginalMessage.GetCharArray().GetData(), FilenameIndex);
 		RemappedMessage.Append(RemapData.Filename);
-		int32 SeparatorStart = FilenameIndex + FilenameSentinelLen;
+		int32 SeparatorStart = FilenameIndex + GPreprocessConstants.FilenameSentinel.Len();
 		RemappedMessage.Append(OriginalMessage.GetCharArray().GetData() + SeparatorStart, LineNumberStart - SeparatorStart);
 		RemappedMessage.AppendInt(RemapData.LineNumber);
 		RemappedMessage.Append(OriginalMessage.GetCharArray().GetData() + LineNumberEnd, OriginalMessage.Len() - LineNumberEnd);
 
 		Diagnostic.StrippedErrorMessage = MoveTemp(RemappedMessage);
-		FilenameIndex = Diagnostic.StrippedErrorMessage.Find(FilenameSentinel);
+		FilenameIndex = Diagnostic.StrippedErrorMessage.Find(GPreprocessConstants.FilenameSentinel.GetData());
 	}
 }
 
@@ -153,9 +202,9 @@ inline void SkipNewLine(const FShaderSource::CharType*& Current, const FShaderSo
 void FShaderPreprocessOutput::StripCode(bool bCopyOriginalPreprocessdSource)
 {
 	// Reserve worst case slack (i.e. assuming there is nothing to strip) to avoid reallocation
-	FShaderSource PreprocessedSourceStripped(LineDirectiveSentinel.GetCharArray().GetData(), PreprocessedSource.Len());
+	FShaderSource PreprocessedSourceStripped(GPreprocessConstants.StrippedCodeHeader, PreprocessedSource.Len());
 	FShaderSource::CharType* OutStrippedData = PreprocessedSourceStripped.GetData();
-	FShaderSource::CharType* OutStripped = OutStrippedData + LineDirectiveSentinel.Len();
+	FShaderSource::CharType* OutStripped = OutStrippedData + GPreprocessConstants.StrippedCodeHeader.Len();
 
 	const FShaderSource::CharType* Begin = PreprocessedSource.GetData(), * Current = Begin;
 	const FShaderSource::CharType* End = Current + PreprocessedSource.Len();
diff --git a/Engine/Source/Runtime/RenderCore/Private/ShaderSource.cpp b/Engine/Source/Runtime/RenderCore/Private/ShaderSource.cpp
--- a/Engine/Source/Runtime/RenderCore/Private/ShaderSource.cpp
+++ b/Engine/Source/Runtime/RenderCore/Private/ShaderSource.cpp
@@ -44,7 +44,7 @@ void FShaderSource::Set(FAnsiStringView InSrc)
 FShaderSource& FShaderSource::operator=(FShaderSource::FStringType&& InSrc)
 {
 	SourceCompressed.Empty();
-	DecompressedSize = 0;
+	DecompressedCharCount = 0;
 
 	int32 InInitialLen = InSrc.Len();
 	// input char array already has a null terminator appended, so can add one less padding char
@@ -66,15 +66,15 @@ void FShaderSource::Compress()
 
 	FOodleDataCompression::ECompressionLevel CompressionLevel = static_cast<FOodleDataCompression::ECompressionLevel>(GShaderSourceCompressionLevel);
 	
-	DecompressedSize = Source.Num();
-	int32 CompressionBufferSize = FOodleDataCompression::CompressedBufferSizeNeeded(DecompressedSize);
+	DecompressedCharCount = Source.Num();
+	int32 CompressionBufferSize = FOodleDataCompression::CompressedBufferSizeNeeded(GetDecompressedSize());
 	SourceCompressed.SetNumUninitialized(CompressionBufferSize);
 
 	int32 CompressedSize = FOodleDataCompression::Compress(
 		SourceCompressed.GetData(), 
 		CompressionBufferSize, 
 		Source.GetData(), 
-		DecompressedSize, 
+		GetDecompressedSize(),
 		Compressor,
 		CompressionLevel);
 
@@ -86,16 +86,16 @@ void FShaderSource::Decompress()
 {
 	if (IsCompressed())
 	{
-		Source.SetNumUninitialized(DecompressedSize);
-		FOodleDataCompression::Decompress(Source.GetData(), DecompressedSize, SourceCompressed.GetData(), SourceCompressed.Num());
+		Source.SetNumUninitialized(DecompressedCharCount);
+		FOodleDataCompression::Decompress(Source.GetData(), GetDecompressedSize(), SourceCompressed.GetData(), SourceCompressed.Num());
 		SourceCompressed.Empty();
-		DecompressedSize = 0;
+		DecompressedCharCount = 0;
 	}
 }
 
 FArchive& operator<<(FArchive& Ar, FShaderSource& Src)
 {
-	Ar << Src.DecompressedSize;
+	Ar << Src.DecompressedCharCount;
 	if (Src.IsCompressed())
 	{
 		Ar << Src.SourceCompressed;
diff --git a/Engine/Source/Runtime/RenderCore/Public/ShaderCompilerCore.h b/Engine/Source/Runtime/RenderCore/Public/ShaderCompilerCore.h
--- a/Engine/Source/Runtime/RenderCore/Public/ShaderCompilerCore.h
+++ b/Engine/Source/Runtime/RenderCore/Public/ShaderCompilerCore.h
@@ -229,6 +229,8 @@ enum class EShaderDebugInfoFlags : uint8
 };
 ENUM_CLASS_FLAGS(EShaderDebugInfoFlags)
 
+using FShaderCompilerInputHash = FBlake3Hash;
+
 /** Struct that gathers all readonly inputs needed for the compilation of a single shader. */
 struct FShaderCompilerInput
 {
@@ -281,7 +283,7 @@ struct FShaderCompilerInput
 	FString DebugDescription;
 
 	// Hash of this input (used as the key for the shader job cache)
-	FBlake3Hash Hash;
+	FShaderCompilerInputHash Hash;
 
 	// Compilation Environment
 	FShaderCompilerEnvironment Environment;
diff --git a/Engine/Source/Runtime/RenderCore/Public/ShaderCompilerJobTypes.h b/Engine/Source/Runtime/RenderCore/Public/ShaderCompilerJobTypes.h
--- a/Engine/Source/Runtime/RenderCore/Public/ShaderCompilerJobTypes.h
+++ b/Engine/Source/Runtime/RenderCore/Public/ShaderCompilerJobTypes.h
@@ -79,7 +79,7 @@ public:
 	FShaderCommonCompileJob* NextLink = nullptr;
 	FShaderCommonCompileJob** PrevLink = nullptr;
 
-	using FInputHash = FBlake3Hash;
+	using FInputHash UE_DEPRECATED(5.5, "Use FShaderCompilerInputHash instead of FShaderCommonCompileJob::FInputHash") = FShaderCompilerInputHash;
 
 	FPendingShaderMapCompileResultsPtr PendingShaderMap;
 
@@ -113,7 +113,7 @@ public:
 	/** Whether or not we are a global shader. */
 	uint8 bIsGlobalShader : 1;
 	/** Hash of all the job inputs */
-	FInputHash InputHash;
+	FShaderCompilerInputHash InputHash;
 
 	/** In-engine timestamp of being added to a pending queue. Not set for jobs that are satisfied from the jobs cache */
 	double TimeAddedToPendingQueue = 0.0;
@@ -147,7 +147,7 @@ public:
 	}
 
 	/** Returns hash of all inputs for this job (needed for caching). */
-	virtual FInputHash GetInputHash() { return FInputHash(); }
+	virtual FShaderCompilerInputHash GetInputHash() { return FShaderCompilerInputHash(); }
 
 	/** Serializes (and deserializes) the output for caching purposes. */
 	virtual void SerializeOutput(FArchive& Ar) {}
@@ -251,7 +251,7 @@ public:
 	// List of pipelines that are sharing this job.
 	TMap<const FVertexFactoryType*, TArray<const FShaderPipelineType*>> SharingPipelines;
 
-	virtual RENDERCORE_API FInputHash GetInputHash() override;
+	virtual RENDERCORE_API FShaderCompilerInputHash GetInputHash() override;
 	virtual RENDERCORE_API void SerializeOutput(FArchive& Ar) override;
 
 	virtual RENDERCORE_API void OnComplete() override;
@@ -312,7 +312,7 @@ public:
 	FShaderPipelineCompileJobKey Key;
 	TArray<TRefCountPtr<FShaderCompileJob>> StageJobs;
 
-	virtual RENDERCORE_API FInputHash GetInputHash() override;
+	virtual RENDERCORE_API FShaderCompilerInputHash GetInputHash() override;
 	virtual RENDERCORE_API void SerializeOutput(FArchive& Ar) override;
 	virtual RENDERCORE_API void OnComplete() override;
 	virtual RENDERCORE_API void AppendDebugName(FStringBuilderBase& OutName) const override;
diff --git a/Engine/Source/Runtime/RenderCore/Public/ShaderPreprocessTypes.h b/Engine/Source/Runtime/RenderCore/Public/ShaderPreprocessTypes.h
--- a/Engine/Source/Runtime/RenderCore/Public/ShaderPreprocessTypes.h
+++ b/Engine/Source/Runtime/RenderCore/Public/ShaderPreprocessTypes.h
@@ -13,6 +13,9 @@ struct FShaderCompilerEnvironment;
 struct FShaderCompilerError;
 struct FShaderCompilerInput;
 
+RENDERCORE_API FShaderSource::FViewType GetShaderSourceDebugHashPrefix();
+RENDERCORE_API FStringView GetShaderSourceDebugHashPrefixWide();
+
 /*
 * Helper class used to remap compiler diagnostic messages from stripped preprocessed source (i.e. source with all whitespace normalized
 * and comments and line directives removed) back to line numbers/locations from the original source. 
@@ -68,6 +71,11 @@ public:
 		return Empty;
 	}
 
+	FShaderSource::FViewType GetSourceView() const
+	{
+		return PreprocessedSource.GetView();
+	}
+
 	FAnsiStringView GetSourceViewAnsi() const
 	{
 #if SHADER_SOURCE_ANSI
diff --git a/Engine/Source/Runtime/RenderCore/Public/ShaderSource.h b/Engine/Source/Runtime/RenderCore/Public/ShaderSource.h
--- a/Engine/Source/Runtime/RenderCore/Public/ShaderSource.h
+++ b/Engine/Source/Runtime/RenderCore/Public/ShaderSource.h
@@ -6,6 +6,7 @@
 #include "Containers/Array.h"
 #include "Containers/StringView.h"
 #include "Containers/UnrealString.h"
+#include "Misc/StringBuilder.h"
 
 #define SHADER_SOURCE_ANSI 1
 
@@ -29,7 +30,7 @@ private:
 		FMemory::Memzero(Source.GetData() + Num, sizeof(CharType) * ShaderSourceSimdPadding);
 
 		SourceCompressed.Empty();
-		DecompressedSize = 0;
+		DecompressedCharCount = 0;
 	}
 
 public:
@@ -38,11 +39,17 @@ public:
 	typedef FAnsiStringView FViewType;
 	typedef FAnsiString FStringType;
 	typedef FCStringAnsi FCStringType;
+
+	template <int NumChars>
+	using TStringBuilder = TAnsiStringBuilder<NumChars>;
 #else
 	typedef TCHAR CharType;
 	typedef FStringView FViewType;
 	typedef FString FStringType;
 	typedef FCString FCStringType;
+
+	template <int NumChars>
+	using TStringBuilder = TStringBuilder<NumChars * 2>; // Buffer size for string builder is in bytes, not characters, hence *2
 #endif
 
 	/** Constexpr predicate indicating whether wide or ansi chars are used */
@@ -133,7 +140,12 @@ public:
 
 	inline bool IsCompressed() const
 	{
-		return DecompressedSize != 0; // DecompressedSize member is only set when compression occurs
+		return DecompressedCharCount != 0; // DecompressedCharCount member is only set when compression occurs
+	}
+
+	inline int32 GetDecompressedSize() const
+	{
+		return DecompressedCharCount * sizeof(FShaderSource::CharType);
 	}
 
 	/* FArchive serialization operator. Note this currently serializes padding for simplicity's sake.
@@ -153,6 +165,6 @@ private:
 #endif
 	TArray<CharType> Source;
 	TArray<uint8> SourceCompressed;
-	int32 DecompressedSize = 0;
+	int32 DecompressedCharCount = 0;
 };
 
