[shaders] fixing stripped code header length calculation, and account for auto-appended null terminator

#rnx

diff --git a/Engine/Source/Runtime/RenderCore/Private/ShaderPreprocessTypes.cpp b/Engine/Source/Runtime/RenderCore/Private/ShaderPreprocessTypes.cpp
--- a/Engine/Source/Runtime/RenderCore/Private/ShaderPreprocessTypes.cpp
+++ b/Engine/Source/Runtime/RenderCore/Private/ShaderPreprocessTypes.cpp
@@ -20,7 +20,7 @@ struct FPreprocessConstants
 
 		check(StrippedCodeHeaderBuilder.GetAllocatedSize() == 0);
 
-		StrippedCodeHeader = { StrippedCodeHeaderBuilder.GetData(), StrippedCodeHeaderBuilder.Len() };
+		StrippedCodeHeader = StrippedCodeHeaderBuilder.ToView();
 	}
 
 	static constexpr FShaderSource::FViewType DebugHashPrefix = SHADER_SOURCE_VIEWLITERAL("DebugHash_");
@@ -39,10 +39,12 @@ private:
 	static constexpr int32 HeaderLen =
 		DebugHashCommentStart.Len() +
 		DebugHashPrefix.Len() +
+		DebugHashCommentEnd.Len() +
 		2 * sizeof(FShaderCompilerInputHash::ByteArray) + // size of input hash converted to a hex string
 		LineDirectiveStart.Len() +
 		FilenameSentinel.Len() +
-		LineDirectiveEnd.Len();
+		LineDirectiveEnd.Len() +
+		1; // +1 for null terminator, even though we don't need it builder always ensures there's space for one
 
 	FShaderSource::TStringBuilder<HeaderLen> StrippedCodeHeaderBuilder;
 };
