Shader Keys: Add a keygenerator struct FShaderKeyGenerator than can generate either full string keys or hashed keys. Hashed keys are faster to create and will be used for incremental cook package dependencies.
#rnx

diff --git a/Engine/Source/Editor/UnrealEd/Private/CookOnTheFlyServer.cpp b/Engine/Source/Editor/UnrealEd/Private/CookOnTheFlyServer.cpp
--- a/Engine/Source/Editor/UnrealEd/Private/CookOnTheFlyServer.cpp
+++ b/Engine/Source/Editor/UnrealEd/Private/CookOnTheFlyServer.cpp
@@ -7648,8 +7648,8 @@ void GetAdditionalCurrentIniVersionStrings( const UCookOnTheFlyServer* CookOnThe
 		IniVersionMap.Add(UE4EngineVersionCompatibleName, UE4EngineVersionCompatible);
 	}*/
 
-	IniVersionMap.Add(TEXT("MaterialShaderMapDDCVersion"), *GetMaterialShaderMapDDCKey());
-	IniVersionMap.Add(TEXT("GlobalDDCVersion"), *GetGlobalShaderMapDDCKey());
+	IniVersionMap.Add(TEXT("MaterialShaderMapDDCVersion"), *GetMaterialShaderMapDDCGuid().ToString());
+	IniVersionMap.Add(TEXT("GlobalDDCVersion"), *GetGlobalShaderMapDDCGuid().ToString());
 
 	UProjectPackagingSettings* PackagingSettings = Cast<UProjectPackagingSettings>(UProjectPackagingSettings::StaticClass()->GetDefaultObject());
 	IniVersionMap.Add(TEXT("IsUsingShaderCodeLibrary"), FString::Printf(TEXT("%d"), PackagingSettings->bShareMaterialShaderCode && CookOnTheFlyServer->IsUsingShaderCodeLibrary()));
diff --git a/Engine/Source/Runtime/Core/Private/Serialization/MemoryImage.cpp b/Engine/Source/Runtime/Core/Private/Serialization/MemoryImage.cpp
--- a/Engine/Source/Runtime/Core/Private/Serialization/MemoryImage.cpp
+++ b/Engine/Source/Runtime/Core/Private/Serialization/MemoryImage.cpp
@@ -1,6 +1,8 @@
 // Copyright Epic Games, Inc. All Rights Reserved.
 
 #include "Serialization/MemoryImage.h"
+
+#include "Containers/StringView.h"
 #include "Containers/UnrealString.h"
 #include "Misc/SecureHash.h"
 #include "Misc/StringBuilder.h"
@@ -12,6 +14,7 @@
 #include "Misc/ScopeRWLock.h"
 #include "Misc/DataDrivenPlatformInfoRegistry.h"
 #include "Serialization/Archive.h"
+#include "Serialization/ShaderKeyGenerator.h"
 #include "Async/ParallelFor.h"
 
 DEFINE_LOG_CATEGORY_STATIC(LogMemoryImage, Log, All);
@@ -109,7 +112,24 @@ FArchive& FPlatformTypeLayoutParameters::Serialize(FArchive& Ar)
 
 void FPlatformTypeLayoutParameters::AppendKeyString(FString& KeyString) const
 {
-	KeyString += FString::Printf(TEXT("FL_%08x_MFA_%08x_V_" MEMORYIMAGE_DERIVEDDATA_VER_ANSI), Flags, MaxFieldAlignment);
+	FShaderKeyGenerator KeyGen(KeyString);
+	Append(KeyGen);
+}
+
+void FPlatformTypeLayoutParameters::Append(FShaderKeyGenerator& KeyGen) const
+{
+	if (KeyGen.IsBinary())
+	{
+		KeyGen.BinaryAppend(&Flags, sizeof(Flags));
+		KeyGen.BinaryAppend(&MaxFieldAlignment, sizeof(MaxFieldAlignment));
+		FAnsiStringView VersionText(MEMORYIMAGE_DERIVEDDATA_VER_ANSI);
+		KeyGen.BinaryAppend(VersionText.GetData(), VersionText.Len() * sizeof(VersionText[0]));
+	}
+	else
+	{
+		KeyGen.TextGetResultString() +=
+			FString::Printf(TEXT("FL_%08x_MFA_%08x_V_" MEMORYIMAGE_DERIVEDDATA_VER_ANSI), Flags, MaxFieldAlignment);
+	}
 }
 
 // evaluated during static-initialization, so logging from regular check() macros won't work correctly
diff --git a/Engine/Source/Runtime/Core/Private/Serialization/ShaderKeyGenerator.cpp b/Engine/Source/Runtime/Core/Private/Serialization/ShaderKeyGenerator.cpp
new file mode 100644
--- /dev/null
+++ b/Engine/Source/Runtime/Core/Private/Serialization/ShaderKeyGenerator.cpp
@@ -0,0 +1,56 @@
+// Copyright Epic Games, Inc. All Rights Reserved.
+
+#include "Serialization/ShaderKeyGenerator.h"
+
+#include "Hash/Blake3.h"
+#include "Misc/Guid.h"
+#include "Misc/SecureHash.h"
+
+void FShaderKeyGenerator::Append(const FBlake3Hash& Value)
+{
+	switch (OutputType)
+	{
+	case EOutputType::Text:
+	{
+		CA_SUPPRESS(6260) /* warning C6260: sizeof * sizeof is usually wrong. */
+		constexpr int32 StringSize = sizeof(TCHAR) * sizeof(FBlake3Hash::ByteArray) * 2;
+		ResultString->Append(WriteToString<StringSize>(Value));
+		return;
+	}
+	case EOutputType::Binary:
+		ResultFunc(&Value, sizeof(Value));
+		return;
+	default:
+		checkNoEntry();
+	}
+}
+
+void FShaderKeyGenerator::Append(const FGuid& Value)
+{
+	switch (OutputType)
+	{
+	case EOutputType::Text:
+		Value.AppendString(*ResultString);
+		return;
+	case EOutputType::Binary:
+		ResultFunc(&Value, sizeof(Value));
+		return;
+	default:
+		checkNoEntry();
+	}
+}
+
+void FShaderKeyGenerator::Append(const FSHAHash& Value)
+{
+	switch (OutputType)
+	{
+	case EOutputType::Text:
+		Value.AppendString(*ResultString);
+		return;
+	case EOutputType::Binary:
+		ResultFunc(&Value.Hash, sizeof(Value.Hash));
+		return;
+	default:
+		checkNoEntry();
+	}
+}
diff --git a/Engine/Source/Runtime/Core/Public/Serialization/MemoryLayout.h b/Engine/Source/Runtime/Core/Public/Serialization/MemoryLayout.h
--- a/Engine/Source/Runtime/Core/Public/Serialization/MemoryLayout.h
+++ b/Engine/Source/Runtime/Core/Public/Serialization/MemoryLayout.h
@@ -15,6 +15,7 @@
 #include "Templates/UnrealTemplate.h"
 
 class FHashedName;
+class FShaderKeyGenerator;
 class FSHA1;
 class FMemoryImageWriter;
 class FMemoryUnfreezeContent;
@@ -854,6 +855,14 @@ struct FPlatformTypeLayoutParameters
 	 * even if binary layouts happen to be compatible.
 	 */
 	CORE_API void AppendKeyString(FString& KeyString) const;
+	CORE_API void Append(FShaderKeyGenerator& KeyGen) const;
+
+private:
+	// Hidden friend for FShaderKeyGenerator Append function
+	friend inline void Append(FShaderKeyGenerator& KeyGen, const FPlatformTypeLayoutParameters& Value)
+	{
+		Value.Append(KeyGen);
+	}
 };
 
 #if UE_ENABLE_INCLUDE_ORDER_DEPRECATED_IN_5_5
diff --git a/Engine/Source/Runtime/Core/Public/Serialization/ShaderKeyGenerator.h b/Engine/Source/Runtime/Core/Public/Serialization/ShaderKeyGenerator.h
new file mode 100644
--- /dev/null
+++ b/Engine/Source/Runtime/Core/Public/Serialization/ShaderKeyGenerator.h
@@ -0,0 +1,370 @@
+// Copyright Epic Games, Inc. All Rights Reserved.
+
+#pragma once
+
+#include "Containers/StringView.h"
+#include "Containers/UnrealString.h"
+#include "CoreTypes.h"
+#include "Misc/StringBuilder.h"
+#include "Templates/Function.h"
+#include "UObject/NameTypes.h"
+
+class FSHAHash;
+struct FBlake3Hash;
+struct FGuid;
+
+/**
+ * Output class passed to Append functions for shader data. It receives Append calls for the input data for shader
+ * compilation. These input data should trigger a recompile if they change, and they are therefore added into the
+ * key used for storage of shader data in DDC and incremental cooks.
+ * 
+ * Known types are appended through FShaderKeyGenerator Append member functions.
+ * For other types, the standard API (which is used by e.g. the append function for TArray) for types that can be
+ * appended to FShaderKeyGenerator is the function in global namespace:
+ *     void Append(FShaderKeyGenerator& KeyGen, const TypeName& Value);
+ * 
+ * The proper Append function can be called for any type using operator<<.
+ */
+class FShaderKeyGenerator
+{
+public:
+	/** Constructor that writes the appended data to a hash function. */
+	inline FShaderKeyGenerator(TUniqueFunction<void(const void* Data, uint64 Size)>&& InResultFunc);
+	/** Constructor that writes the appended data to a long human-readable debug string. */
+	inline FShaderKeyGenerator(FString& InResultString);
+
+	/**
+	 * True iff the KeyGen is writing to a hash function. Writing to a hash function also implies that debugtext and
+	 * separators will be skipped in the output.
+	 */
+	inline bool IsBinary() const;
+	/** True iff the KeyGen is writing to a human-readable debug string. */
+	inline bool IsText() const;
+
+	/** Must not be called unless IsBinary is true. Append data directly to the hash function. */
+	inline void BinaryAppend(const void* Data, uint64 Size);
+	/** Must not be called unless IsText is true. Return a modifiable reference to the string being written. */
+	inline FString& TextGetResultString();
+
+	/** Append arbitrary text to the output string or hash function. */
+	inline void Append(FStringView Value);
+	/**
+	 * Append arbitrary text to the output string or hash function. Null-terminated character pointer.
+	 * Null pointer is a noop.
+	 */
+	inline void Append(const TCHAR* Value);
+	/** Convert the FName to text (case-sensitive) and append it to the output string or hash function. */
+	inline void Append(FName Value);
+	/** Appendf the integer to the output string or pass it as binary data to the hash function. */
+	inline void Append(int64 Value);
+	/** Appendf the integer to the output string or pass it as binary data to the hash function. */
+	inline void Append(uint64 Value);
+	/** Appendf the integer to the output string or pass it as binary data to the hash function. */
+	inline void Append(int32 Value);
+	/** Appendf the integer to the output string or pass it as binary data to the hash function. */
+	inline void Append(uint32 Value);
+	/** Appendf the integer to the output string, using %X format, or pass it as binary data to the hash function. */
+	inline void AppendHex(uint32 Value);
+	/** Appendf 0 or 1 to the output string or pass a 0 or 1 uint8 the hash function. */
+	inline void AppendBoolInt(bool Value);
+	/** Append Value to the output string (equivalent to LexToString) or pass it as binary data to the hash function. */
+	CORE_API void Append(const FBlake3Hash& Value);
+	/** Append Value to the output string (EGuidFormats::Digits) or pass it as binary data to the hash function. */
+	CORE_API void Append(const FGuid& Value);
+	/** Append Value to the output string (equivalent to LexToString) or pass it as binary data to the hash function. */
+	CORE_API void Append(const FSHAHash& Value);
+
+	/** Append arbitrary text to the output human-readable string. Noop if !IsText. */
+	inline void AppendDebugText(FStringView Value);
+	/** Append the separator character '_' to the output human-readable string. Noop if !IsText. */
+	inline void AppendSeparator();
+
+private:
+	enum class EOutputType
+	{
+		Text,
+		Binary,
+	};
+
+private:
+	TUniqueFunction<void(const void* Data, uint64 Size)> ResultFunc;
+	FString* ResultString = nullptr;
+	EOutputType OutputType = EOutputType::Text;
+};
+
+
+///////////////////////////////////////////////////////
+// Global functions taking FShaderKeyGenerator argument
+///////////////////////////////////////////////////////
+
+/** Append an ArrayView of handled types to the KeyGen. Calls operator<<(KeyGen, TypeName) on each element.*/
+template <typename T>
+void Append(FShaderKeyGenerator& KeyGen, TConstArrayView<T> Value);
+
+/** Append an Array of handled types to the KeyGen. */
+template <typename T>
+void Append(FShaderKeyGenerator& KeyGen, const TArray<T>& Value);
+
+/**
+ * (FShaderKeyGenerator&)KeyGen << (const T&) Value;
+ * 
+ * Append the given Value to the KeyGen, using FShaderKeyGenerator::Append for known types, or
+ * Append(FShaderKeyGenerator&, const T&) for unknown types.
+ * Forward declare not possible, see implementation below.
+ */
+// template <typename T>
+// FShaderKeyGenerator& operator<<(FShaderKeyGenerator& KeyGen, const T& Value);
+
+
+///////////////////////////////////////////////////////
+// Inline implementations
+///////////////////////////////////////////////////////
+
+inline FShaderKeyGenerator::FShaderKeyGenerator(TUniqueFunction<void(const void* Data, uint64 Size)>&& InResultFunc)
+	: ResultFunc(MoveTemp(InResultFunc))
+	, OutputType(EOutputType::Binary)
+{
+}
+
+inline FShaderKeyGenerator::FShaderKeyGenerator(FString& InResultString)
+	: ResultString(&InResultString)
+	, OutputType(EOutputType::Text)
+{
+}
+
+inline bool FShaderKeyGenerator::IsBinary() const
+{
+	return OutputType == EOutputType::Binary;
+}
+
+inline bool FShaderKeyGenerator::IsText() const
+{
+	return OutputType == EOutputType::Text;
+}
+
+inline void FShaderKeyGenerator::BinaryAppend(const void* Data, uint64 Size)
+{
+	check(IsBinary());
+	ResultFunc(Data, Size);
+}
+
+inline FString& FShaderKeyGenerator::TextGetResultString()
+{
+	check(IsText());
+	return *ResultString;
+}
+
+inline void FShaderKeyGenerator::Append(FStringView Value)
+{
+	switch (OutputType)
+	{
+	case EOutputType::Text:
+		ResultString->Append(Value);
+		return;
+	case EOutputType::Binary:
+		ResultFunc(Value.GetData(), Value.Len() * sizeof(Value.GetData()[0]));
+		return;
+	default:
+		checkNoEntry();
+	}
+}
+
+inline void FShaderKeyGenerator::Append(const TCHAR* Value)
+{
+	Append(FStringView(Value));
+}
+
+inline void FShaderKeyGenerator::Append(FName Value)
+{
+	switch (OutputType)
+	{
+	case EOutputType::Text:
+		Value.AppendString(*ResultString);
+		return;
+	case EOutputType::Binary:
+	{
+		TStringBuilder<FName::StringBufferSize> Builder(InPlace, Value);
+		ResultFunc(Builder.GetData(), Builder.Len() * sizeof(Builder.GetData()[0]));
+		return;
+	}
+	default:
+		checkNoEntry();
+	}
+}
+
+inline void FShaderKeyGenerator::Append(int64 Value)
+{
+	switch (OutputType)
+	{
+	case EOutputType::Text:
+		ResultString->Appendf(TEXT(INT64_FMT), Value);
+		return;
+	case EOutputType::Binary:
+		ResultFunc(&Value, sizeof(Value));
+		return;
+	default:
+		checkNoEntry();
+	}
+}
+
+inline void FShaderKeyGenerator::Append(uint64 Value)
+{
+	switch (OutputType)
+	{
+	case EOutputType::Text:
+		ResultString->Appendf(TEXT(UINT64_FMT), Value);
+		return;
+	case EOutputType::Binary:
+		ResultFunc(&Value, sizeof(Value));
+		return;
+	default:
+		checkNoEntry();
+	}
+}
+
+inline void FShaderKeyGenerator::Append(int32 Value)
+{
+	switch (OutputType)
+	{
+	case EOutputType::Text:
+		ResultString->Appendf(TEXT("%d"), Value);
+		return;
+	case EOutputType::Binary:
+		ResultFunc(&Value, sizeof(Value));
+		return;
+	default:
+		checkNoEntry();
+	}
+}
+
+inline void FShaderKeyGenerator::Append(uint32 Value)
+{
+	switch (OutputType)
+	{
+	case EOutputType::Text:
+		ResultString->Appendf(TEXT("%u"), Value);
+		return;
+	case EOutputType::Binary:
+		ResultFunc(&Value, sizeof(Value));
+		return;
+	default:
+		checkNoEntry();
+	}
+}
+
+inline void FShaderKeyGenerator::AppendHex(uint32 Value)
+{
+	switch (OutputType)
+	{
+	case EOutputType::Text:
+		ResultString->Appendf(TEXT("%X"), Value);
+		return;
+	case EOutputType::Binary:
+		ResultFunc(&Value, sizeof(Value));
+		return;
+	default:
+		checkNoEntry();
+	}
+}
+
+inline void FShaderKeyGenerator::AppendBoolInt(bool Value)
+{
+	switch (OutputType)
+	{
+	case EOutputType::Text:
+		ResultString->AppendChar(Value ? '1' : '0');
+		return;
+	case EOutputType::Binary:
+	{
+		uint8 SizedValue = Value ? 1 : 0;
+		ResultFunc(&SizedValue, sizeof(SizedValue));
+		return;
+	}
+	default:
+		checkNoEntry();
+	}
+}
+
+inline void FShaderKeyGenerator::AppendDebugText(FStringView Value)
+{
+	switch (OutputType)
+	{
+	case EOutputType::Text:
+		ResultString->Append(Value);
+		return;
+	case EOutputType::Binary:
+		// Binary output ignores DebugText
+		return;
+	default:
+		checkNoEntry();
+	}
+}
+
+inline void FShaderKeyGenerator::AppendSeparator()
+{
+	switch (OutputType)
+	{
+	case EOutputType::Text:
+		ResultString->AppendChar('_');
+		return;
+	case EOutputType::Binary:
+		// Binary output ignores DebugText; Separator is a type of DebugText
+		return;
+	default:
+		checkNoEntry();
+	}
+}
+
+template <typename T>
+void Append(FShaderKeyGenerator& KeyGen, const TArray<T>& Value)
+{
+	Append(KeyGen, TConstArrayView<T>(Value));
+}
+
+template <typename T>
+void Append(FShaderKeyGenerator& KeyGen, TConstArrayView<T> Value)
+{
+	for (const T& Element : Value)
+	{
+		KeyGen << Element;
+	}
+}
+
+/**
+ * Template override for struct used in std::enable_if to report whether a type is a FShaderKeyGenerator known type
+ * with an Append member function. This template override provides the false value for non-known types.
+ */
+template <typename T, typename = int>
+struct IsFShaderKeyGeneratorKnownType
+ : std::false_type
+{
+};
+
+/**
+ * Template override for struct used in std::enable_if to report whether a type is a FShaderKeyGenerator known type
+ * with an Append member function. This template override provides the true value for known types.
+ */
+template <typename T>
+struct IsFShaderKeyGeneratorKnownType<T, decltype(std::declval<FShaderKeyGenerator>().Append(std::declval<T>()), 0)>
+ : std::true_type
+{
+};
+
+/** Template overide of FShaderKeyGenerator&& operator<<(FShaderKeyGenerator&, const T&), for known types. */
+template <typename T>
+typename std::enable_if<IsFShaderKeyGeneratorKnownType<T>::value, FShaderKeyGenerator&>::type
+operator<<(FShaderKeyGenerator& KeyGen, const T& Value)
+{
+	KeyGen.Append(Value);
+	return KeyGen;
+}
+
+/** Template overide of FShaderKeyGenerator&& operator<<(FShaderKeyGenerator&, const T&), for non-known types. */
+template <typename T>
+typename std::enable_if<!IsFShaderKeyGeneratorKnownType<T>::value, FShaderKeyGenerator&>::type
+operator<<(FShaderKeyGenerator& KeyGen, const T& Value)
+{
+	Append(KeyGen, Value);
+	return KeyGen;
+}
+
diff --git a/Engine/Source/Runtime/Engine/Classes/Materials/MaterialLayersFunctions.h b/Engine/Source/Runtime/Engine/Classes/Materials/MaterialLayersFunctions.h
--- a/Engine/Source/Runtime/Engine/Classes/Materials/MaterialLayersFunctions.h
+++ b/Engine/Source/Runtime/Engine/Classes/Materials/MaterialLayersFunctions.h
@@ -14,6 +14,7 @@
 #define LOCTEXT_NAMESPACE "MaterialLayersFunctions"
 
 class FArchive;
+class FShaderKeyGenerator;
 struct FMaterialLayersFunctions;
 
 UENUM()
@@ -48,6 +49,7 @@ struct FMaterialLayersFunctionsID
 
 	//TODO: Investigate whether this is really required given it is only used by FMaterialShaderMapId AND that one also uses UpdateHash
 	void AppendKeyString(FString& KeyString) const;
+	void Append(FShaderKeyGenerator& KeyGen) const;
 	#endif
 };
 
@@ -272,6 +274,7 @@ struct FMaterialLayersFunctions : public FMaterialLayersFunctionsRuntimeData
 
 	/** Gets a string representation of the ID */
 	ENGINE_API FString GetStaticPermutationString() const;
+	ENGINE_API void AppendStaticPermutationKey(FShaderKeyGenerator& KeyGen) const;
 
 	ENGINE_API void UnlinkLayerFromParent(int32 Index);
 	ENGINE_API bool IsLayerLinkedToParent(int32 Index) const;
diff --git a/Engine/Source/Runtime/Engine/Private/Materials/HLSLMaterialTranslator.cpp b/Engine/Source/Runtime/Engine/Private/Materials/HLSLMaterialTranslator.cpp
--- a/Engine/Source/Runtime/Engine/Private/Materials/HLSLMaterialTranslator.cpp
+++ b/Engine/Source/Runtime/Engine/Private/Materials/HLSLMaterialTranslator.cpp
@@ -40,6 +40,7 @@
 #include "StringTemplate.h"
 #include "ParameterCollection.h"
 #include "RenderUtils.h"
+#include "Serialization/ShaderKeyGenerator.h"
 #include "Stats/StatsMisc.h"
 #include "Stats/StatsTrace.h"
 #include "SubstrateDefinitions.h"
@@ -592,13 +593,19 @@ void FHLSLMaterialTranslator::FSubstrateCompilationContext::Initialise()
 }
 
 void FHLSLMaterialTranslator::AppendVersionString(FString& Output, EShaderPlatform Platform)
+{
+	FShaderKeyGenerator KeyGen(Output);
+	AppendVersion(KeyGen, Platform);
+}
+
+void FHLSLMaterialTranslator::AppendVersion(FShaderKeyGenerator& KeyGen, EShaderPlatform Platform)
 {
 	static const FGuid DDCVersion = FDevSystemGuids::GetSystemGuid(FDevSystemGuids::Get().MaterialTranslationDDCVersion);
-	Output.Append(TEXT("_MatTransl_"));
-	Output.Append(FMaterialSourceTemplate::Get().GetTemplateHashString(Platform));
-	Output.Append(TEXT("_"));
-	DDCVersion.AppendString(Output);
-	Output.Append(TEXT("_"));
+	KeyGen.AppendDebugText(TEXT("_MatTransl_"));
+	KeyGen.Append(FMaterialSourceTemplate::Get().GetTemplateHashString(Platform));
+	KeyGen.AppendSeparator();
+	KeyGen.Append(DDCVersion);
+	KeyGen.AppendSeparator();
 }
 
 FHLSLMaterialTranslator::FHLSLMaterialTranslator(FMaterial* InMaterial,
diff --git a/Engine/Source/Runtime/Engine/Private/Materials/HLSLMaterialTranslator.h b/Engine/Source/Runtime/Engine/Private/Materials/HLSLMaterialTranslator.h
--- a/Engine/Source/Runtime/Engine/Private/Materials/HLSLMaterialTranslator.h
+++ b/Engine/Source/Runtime/Engine/Private/Materials/HLSLMaterialTranslator.h
@@ -56,6 +56,7 @@
 
 class Error;
 class FAddUniformExpressionScope;
+class FShaderKeyGenerator;
 
 namespace UE::DerivedData
 {
@@ -575,6 +576,7 @@ public:
 	 * Returns a string representation that identify the translator version. Used to keep shader DDC keys valid when material translation internals change.
 	 */
 	static void AppendVersionString(FString& Output, EShaderPlatform Platform);
+	static void AppendVersion(FShaderKeyGenerator& KeyGen, EShaderPlatform Platform);
 
 	FHLSLMaterialTranslator(FMaterial* InMaterial,
 		FMaterialCompilationOutput& InMaterialCompilationOutput,
diff --git a/Engine/Source/Runtime/Engine/Private/Materials/MaterialAttributeDefinitionMap.cpp b/Engine/Source/Runtime/Engine/Private/Materials/MaterialAttributeDefinitionMap.cpp
--- a/Engine/Source/Runtime/Engine/Private/Materials/MaterialAttributeDefinitionMap.cpp
+++ b/Engine/Source/Runtime/Engine/Private/Materials/MaterialAttributeDefinitionMap.cpp
@@ -7,6 +7,7 @@
 #include "Materials/Material.h"
 #include "Materials/MaterialExpressionShadingModel.h"
 #include "Materials/MaterialExpressionSubstrate.h"
+#include "Serialization/ShaderKeyGenerator.h"
 #include "SubstrateDefinitions.h"
 
 #define LOCTEXT_NAMESPACE "MaterialShared"
@@ -483,6 +484,12 @@ FString FMaterialAttributeDefinitionMap::GetPinNameFromShadingModelField(FMateri
 }
 
 void FMaterialAttributeDefinitionMap::AppendDDCKeyString(FString& String)
+{
+	FShaderKeyGenerator KeyGen(String);
+	AppendDDCKey(KeyGen);
+}
+
+void FMaterialAttributeDefinitionMap::AppendDDCKey(FShaderKeyGenerator& KeyGen)
 {
 	FString& DDCString = GMaterialPropertyAttributesMap.AttributeDDCString;
 
@@ -513,7 +520,7 @@ void FMaterialAttributeDefinitionMap::AppendDDCKeyString(FString& String)
 		// TODO: In debug force re-generate DDC string and compare to catch invalid runtime changes
 	}
 
-	String.Append(DDCString);
+	KeyGen.Append(DDCString);
 }
 
 void FMaterialAttributeDefinitionMap::AddCustomAttribute(const FGuid& AttributeID, const FString& AttributeName, const FString& FunctionName, EMaterialValueType ValueType, const FVector4& DefaultValue, MaterialAttributeBlendFunction BlendFunction /*= nullptr*/)
diff --git a/Engine/Source/Runtime/Engine/Private/Materials/MaterialExpressions.cpp b/Engine/Source/Runtime/Engine/Private/Materials/MaterialExpressions.cpp
--- a/Engine/Source/Runtime/Engine/Private/Materials/MaterialExpressions.cpp
+++ b/Engine/Source/Runtime/Engine/Private/Materials/MaterialExpressions.cpp
@@ -296,6 +296,7 @@
 #include "Framework/Notifications/NotificationManager.h"
 #include "Widgets/Notifications/SNotificationList.h"
 #include "MaterialGraph/MaterialGraphSchema.h"
+#include "Serialization/ShaderKeyGenerator.h"
 #include "SubstrateMaterial.h"
 #else
 #include "Materials/MaterialExpressionVertexInterpolator.h"
@@ -16254,20 +16255,25 @@ void FMaterialLayersFunctionsID::UpdateHash(FSHA1& HashState) const
 	HashState.Update((const uint8*)LayerStates.GetData(), LayerStates.Num()*LayerStates.GetTypeSize());
 }
 
-
 void FMaterialLayersFunctionsID::AppendKeyString(FString& KeyString) const
+{
+	FShaderKeyGenerator KeyGen(KeyString);
+	Append(KeyGen);
+}
+
+void FMaterialLayersFunctionsID::Append(FShaderKeyGenerator& KeyGen) const
 {
 	for (const FGuid &Guid : LayerIDs)
 	{
-		KeyString += Guid.ToString();
+		KeyGen.Append(Guid);
 	}
 	for (const FGuid &Guid : BlendIDs)
 	{
-		KeyString += Guid.ToString();
+		KeyGen.Append(Guid);
 	}
 	for (bool State : LayerStates)
 	{
-		KeyString += FString::FromInt(State);
+		KeyGen.AppendBoolInt(State);
 	}
 }
 #endif // WITH_EDITOR
@@ -16351,10 +16357,16 @@ const FMaterialLayersFunctionsID FMaterialLayersFunctionsRuntimeData::GetID(cons
 FString FMaterialLayersFunctions::GetStaticPermutationString() const
 {
 	FString StaticKeyString;
-	GetID().AppendKeyString(StaticKeyString);
+	FShaderKeyGenerator KeyGen(StaticKeyString);
+	AppendStaticPermutationKey(KeyGen);
 	return StaticKeyString;
 }
 
+void FMaterialLayersFunctions::AppendStaticPermutationKey(FShaderKeyGenerator& KeyGen) const
+{
+	GetID().Append(KeyGen);
+}
+
 void FMaterialLayersFunctions::SerializeLegacy(FArchive& Ar)
 {
 	FString KeyString_DEPRECATED;
diff --git a/Engine/Source/Runtime/Engine/Private/Materials/MaterialShader.cpp b/Engine/Source/Runtime/Engine/Private/Materials/MaterialShader.cpp
--- a/Engine/Source/Runtime/Engine/Private/Materials/MaterialShader.cpp
+++ b/Engine/Source/Runtime/Engine/Private/Materials/MaterialShader.cpp
@@ -1,6 +1,7 @@
 // Copyright Epic Games, Inc. All Rights Reserved.
 
 #include "MaterialShader.h"
+
 #include "DerivedDataCache.h"
 #include "DerivedDataRequestOwner.h"
 #include "RenderUtils.h"
@@ -24,6 +25,7 @@
 #include "ProfilingDebugging/LoadTimeTracker.h"
 #include "Misc/PathViews.h"
 #include "SceneTexturesConfig.h"
+#include "Serialization/ShaderKeyGenerator.h"
 #include "ShaderCodeLibrary.h"
 #include "UObject/UE5ReleaseStreamObjectVersion.h"
 #include "Materials/Material.h"
@@ -187,24 +189,38 @@ FString GetBlendModeString(EBlendMode BlendMode)
 
 #if WITH_EDITOR
 /** Creates a string key for the derived data cache given a shader map id. */
-FString GetMaterialShaderMapKeyString(const FMaterialShaderMapId& ShaderMapId, const FMaterialShaderParameters& ShaderParameters, EShaderPlatform Platform, bool bIncludeKeyStringShaderDependencies)
+FString GetMaterialShaderMapKeyString(const FMaterialShaderMapId& ShaderMapId,
+	const FMaterialShaderParameters& ShaderParameters, EShaderPlatform Platform,
+	bool bIncludeKeyStringShaderDependencies)
 {
-	TRACE_CPUPROFILER_EVENT_SCOPE(GetMaterialShaderMapKeyString);
-	FName Format = LegacyShaderPlatformToShaderFormat(Platform);
 	FString ShaderMapKeyString;
 	ShaderMapKeyString.Reserve(16384);
+	FShaderKeyGenerator KeyGen(ShaderMapKeyString);
+	GetMaterialShaderMapKey(KeyGen, ShaderMapId, ShaderParameters, Platform,
+		bIncludeKeyStringShaderDependencies);
+	return ShaderMapKeyString;
+}
+
+void GetMaterialShaderMapKey(FShaderKeyGenerator& KeyGen, const FMaterialShaderMapId& ShaderMapId,
+	const FMaterialShaderParameters& ShaderParameters, EShaderPlatform Platform,
+	bool bIncludeKeyStringShaderDependencies)
+{
+	TRACE_CPUPROFILER_EVENT_SCOPE(GetMaterialShaderMapKeyString);
+
+	FName Format = LegacyShaderPlatformToShaderFormat(Platform);
 	
-	ShaderMapKeyString.Append(TEXTVIEW("MATSM_"));
-	ShaderMapKeyString.Append(GetMaterialShaderMapDDCKey());
-	ShaderMapKeyString.AppendChar('_');
+	KeyGen.AppendDebugText(TEXTVIEW("MATSM_"));
+	KeyGen.Append(GetMaterialShaderMapDDCGuid());
+	KeyGen.AppendSeparator();
+
+	KeyGen.Append(Format);
+	KeyGen.AppendSeparator();
+	KeyGen.Append(GetTargetPlatformManagerRef().ShaderFormatVersion(Format));
+	KeyGen.AppendSeparator();
 
-	Format.AppendString(ShaderMapKeyString);
-	ShaderMapKeyString.AppendChar('_');
-	ShaderMapKeyString.AppendInt(GetTargetPlatformManagerRef().ShaderFormatVersion(Format));
-	ShaderMapKeyString.AppendChar('_');
+	ShaderMapAppendKey(Platform, KeyGen);
+	ShaderMapId.Append(KeyGen, true, bIncludeKeyStringShaderDependencies);
 
-	ShaderMapAppendKeyString(Platform, ShaderMapKeyString);
-	ShaderMapId.AppendKeyString(ShaderMapKeyString, true, bIncludeKeyStringShaderDependencies);
 
 	FMemoryHasherBlake3 EnvHasher;
 
@@ -240,18 +256,13 @@ FString GetMaterialShaderMapKeyString(const FMaterialShaderMapId& ShaderMapId, c
 			}
 		}
 	}
-	// * 2 for hex representation of hash; + 6 for tag/underscores
-	TStringBuilder<sizeof(TCHAR) * (sizeof(FBlake3Hash::ByteArray) * 2 + 6)> EnvHashString;
-	EnvHashString << "_EMH_" << EnvHasher.Finalize() << '_';
-	check(EnvHashString.GetAllocatedSize() == 0);
-	ShaderMapKeyString.Append(EnvHashString.ToView());
+	KeyGen.AppendDebugText(TEXTVIEW("_EMH_"));
+	KeyGen.Append(EnvHasher.Finalize());
+	KeyGen.AppendSeparator();
 
-	FMaterialAttributeDefinitionMap::AppendDDCKeyString(ShaderMapKeyString);
-	FShaderCompileUtilities::AppendGBufferDDCKeyString(Platform, ShaderMapKeyString);
-
-	FHLSLMaterialTranslator::AppendVersionString(ShaderMapKeyString, Platform);
-
-	return ShaderMapKeyString;
+	FMaterialAttributeDefinitionMap::AppendDDCKey(KeyGen);
+	FShaderCompileUtilities::AppendGBufferDDCKey(Platform, KeyGen);
+	FHLSLMaterialTranslator::AppendVersion(KeyGen, Platform);
 }
 
 static UE::DerivedData::FCacheKey GetMaterialShaderMapKey(const FStringView MaterialShaderMapKey)
@@ -311,6 +322,59 @@ void UpdateMaterialShaderCompilingStats(const FMaterial* Material)
 	}
 }
 
+void FStaticParameterBase::AppendKeyString(FString& KeyString) const
+{
+	FShaderKeyGenerator KeyGen(KeyString);
+	Append(KeyGen);
+}
+
+void FStaticParameterBase::Append(FShaderKeyGenerator& KeyGen) const
+{
+	ParameterInfo.Append(KeyGen);
+	KeyGen.AppendBoolInt(bOverride);
+	KeyGen.Append(ExpressionGUID);
+}
+
+void FStaticSwitchParameter::AppendKeyString(FString& KeyString) const
+{
+	FShaderKeyGenerator KeyGen(KeyString);
+	Append(KeyGen);
+}
+
+void FStaticSwitchParameter::Append(FShaderKeyGenerator& KeyGen) const
+{
+	FStaticParameterBase::Append(KeyGen);
+	KeyGen.Append(Value);
+}
+
+void FStaticComponentMaskParameter::AppendKeyString(FString& KeyString) const
+{
+	FShaderKeyGenerator KeyGen(KeyString);
+	Append(KeyGen);
+}
+
+void FStaticComponentMaskParameter::Append(FShaderKeyGenerator& KeyGen) const
+{
+	FStaticParameterBase::Append(KeyGen);
+	KeyGen.AppendBoolInt(R);
+	KeyGen.AppendBoolInt(G);
+	KeyGen.AppendBoolInt(B);
+	KeyGen.AppendBoolInt(A);
+}
+
+void FStaticTerrainLayerWeightParameter::AppendKeyString(FString& KeyString) const
+{
+	FShaderKeyGenerator KeyGen(KeyString);
+	Append(KeyGen);
+}
+
+void FStaticTerrainLayerWeightParameter::Append(FShaderKeyGenerator& KeyGen) const
+{
+	KeyGen.Append(LayerName);
+	KeyGen.Append(WeightmapIndex);
+	KeyGen.AppendBoolInt(bWeightBasedBlend);
+}
+
 FStaticParameterSet::FStaticParameterSet(const FStaticParameterSet& InValue) = default;
 
 FStaticParameterSet& FStaticParameterSet::operator=(const FStaticParameterSet& InValue)
@@ -622,11 +686,28 @@ void FStaticParameterSet::SetStaticComponentMaskParameterValue(const FMaterialPa
 FString FSubstrateCompilationConfig::GetShaderMapKeyString() const
 {
 	FString SubStrateCompStr;
-	if (bFullSimplify)					SubStrateCompStr += TEXT("_SBSTRFS");
-	if (BytesPerPixelOverride >= 0)		SubStrateCompStr += TEXT("_SBSTRBS");
-	if (ClosuresPerPixelOverride >= 0)	SubStrateCompStr += TEXT("_SBSTRCS");
+	FShaderKeyGenerator KeyGen(SubStrateCompStr);
+	Append(KeyGen);
 	return SubStrateCompStr;
 }
+void FSubstrateCompilationConfig::Append(FShaderKeyGenerator& KeyGen) const
+{
+	if (bFullSimplify)
+	{
+		KeyGen.AppendSeparator();
+		KeyGen.Append(TEXT("SBSTRFS"));
+	}
+	if (BytesPerPixelOverride >= 0)
+	{
+		KeyGen.AppendSeparator();
+		KeyGen.Append(TEXT("SBSTRBS"));
+	}
+	if (ClosuresPerPixelOverride >= 0)
+	{
+		KeyGen.AppendSeparator();
+		KeyGen.Append(TEXT("SBSTRCS"));
+	}
+}
 
 void FSubstrateCompilationConfig::UpdateHash(FSHA1& Hasher) const
 {
@@ -1115,116 +1196,121 @@ void FMaterialShaderMapId::UpdateFromParameterSet(const FStaticParameterSet& Sta
 uint32 GetTypeHash(FPlatformTypeLayoutParameters Params) { return HashCombine(Params.Flags, Params.MaxFieldAlignment); }
 
 void FMaterialShaderMapId::AppendStaticParametersString(FString& ParamsString) const
+{
+	FShaderKeyGenerator KeyGen(ParamsString);
+	AppendStaticParameters(KeyGen);
+}
+
+void FMaterialShaderMapId::AppendStaticParameters(FShaderKeyGenerator& KeyGen) const
 {
 	for (const FStaticSwitchParameter& StaticSwitchParameter : StaticSwitchParameters)
 	{
-		StaticSwitchParameter.AppendKeyString(ParamsString);
+		StaticSwitchParameter.Append(KeyGen);
 	}
 	for (const FStaticComponentMaskParameter& StaticComponentMaskParameter : StaticComponentMaskParameters)
 	{
-		StaticComponentMaskParameter.AppendKeyString(ParamsString);
+		StaticComponentMaskParameter.Append(KeyGen);
 	}
 	for (const FStaticTerrainLayerWeightParameter& StaticTerrainLayerWeightParameter : TerrainLayerWeightParameters)
 	{
-		StaticTerrainLayerWeightParameter.AppendKeyString(ParamsString);
+		StaticTerrainLayerWeightParameter.Append(KeyGen);
 	}
 }
 
 void FMaterialShaderMapId::AppendKeyString(FString& KeyString, bool bIncludeSourceAndMaterialState, bool bIncludeKeyStringShaderDependencies) const
+{
+	FShaderKeyGenerator KeyGen(KeyString);
+	Append(KeyGen, bIncludeSourceAndMaterialState, bIncludeKeyStringShaderDependencies);
+}
+
+void FMaterialShaderMapId::Append(FShaderKeyGenerator& KeyGen, bool bIncludeSourceAndMaterialState, bool bIncludeKeyStringShaderDependencies) const
 {
 	check(IsContentValid());
 	if (bIncludeSourceAndMaterialState)
 	{
-		BaseMaterialId.AppendString(KeyString);
-		KeyString.AppendChar('_');
+		KeyGen.Append(BaseMaterialId);
+		KeyGen.AppendSeparator();
 	}
 
-	GetMaterialQualityLevelFName(QualityLevel).AppendString(KeyString);
-	KeyString.AppendChar('_');
+	KeyGen.Append(GetMaterialQualityLevelFName(QualityLevel));
+	KeyGen.AppendSeparator();
 
 	FName FeatureLevelName;
 	GetFeatureLevelName(FeatureLevel, FeatureLevelName);
-	FeatureLevelName.AppendString(KeyString);
-	KeyString.AppendChar('_');
+	KeyGen.Append(FeatureLevelName);
+	KeyGen.AppendSeparator();
 
-	LayoutParams.AppendKeyString(KeyString);
+	LayoutParams.Append(KeyGen);
 
-	AppendStaticParametersString(KeyString);
+	AppendStaticParameters(KeyGen);
 
 	if (MaterialLayersId)
 	{
-		MaterialLayersId->AppendKeyString(KeyString);
+		MaterialLayersId->Append(KeyGen);
 	}
 
-	KeyString.AppendChar('_');
-	KeyString.AppendInt(Usage);
-	KeyString.AppendChar('_');
+	KeyGen.AppendSeparator();
+	KeyGen.Append(static_cast<int32>(Usage));
+	KeyGen.AppendSeparator();
 
 	if (Usage == EMaterialShaderMapUsage::MaterialExportCustomOutput)
 	{
-		KeyString += UsageCustomOutput;
-		KeyString.AppendChar('_');
+		KeyGen.Append(UsageCustomOutput);
+		KeyGen.AppendSeparator();
 	}
 
 	if (bIncludeSourceAndMaterialState)
 	{
-		// Add any referenced functions to the key so that we will recompile when they are changed
-		for (int32 FunctionIndex = 0; FunctionIndex < ReferencedFunctions.Num(); FunctionIndex++)
-		{
-			ReferencedFunctions[FunctionIndex].AppendString(KeyString);
-		}
+		::Append(KeyGen, ReferencedFunctions);
 	}
 
 	{
 		const FSHAHash LayoutHash = GetShaderTypeLayoutHash(StaticGetTypeLayoutDesc<FMaterialShaderMapContent>(), LayoutParams);
-		KeyString.AppendChar('_');
-		LayoutHash.AppendString(KeyString);
-		KeyString.AppendChar('_');
+		KeyGen.AppendSeparator();
+		KeyGen.Append(LayoutHash);
+		KeyGen.AppendSeparator();
 	}
 
-	KeyString.AppendChar('_');
+	KeyGen.AppendSeparator();
 
 	if (bIncludeSourceAndMaterialState)
 	{
-		for (int32 CollectionIndex = 0; CollectionIndex < ReferencedParameterCollections.Num(); CollectionIndex++)
-		{
-			ReferencedParameterCollections[CollectionIndex].AppendString(KeyString);
-		}
+		::Append(KeyGen, ReferencedParameterCollections);
 	}
 
 	// Add the inputs for any shaders that are stored inline in the shader map
 	if (bIncludeKeyStringShaderDependencies)
 	{
-		AppendKeyStringShaderDependencies(
+		AppendShaderDependencies(
+			KeyGen,
 			MakeArrayView(ShaderTypeDependencies),
 			MakeArrayView(ShaderPipelineTypeDependencies),
 			MakeArrayView(VertexFactoryTypeDependencies),
 			LayoutParams,
-			KeyString,
 			bIncludeSourceAndMaterialState);
 	}
 
-	BytesToHex(&TextureReferencesHash.Hash[0], sizeof(TextureReferencesHash.Hash), KeyString);
+	KeyGen.Append(TextureReferencesHash);
 
-	BytesToHex(&ExpressionIncludesHash.Hash[0], sizeof(ExpressionIncludesHash.Hash), KeyString);
+	KeyGen.Append(ExpressionIncludesHash);
 
-	BytesToHex(&BasePropertyOverridesHash.Hash[0], sizeof(BasePropertyOverridesHash.Hash), KeyString);
+	KeyGen.Append(BasePropertyOverridesHash);
 
 	if (bUsingNewHLSLGenerator)
 	{
-		KeyString += FString::Printf(TEXT("_NewHLSL%d"), FMaterialHLSLGenerator::Version);
+		KeyGen.AppendDebugText(TEXT("_NewHLSL"));
+		KeyGen.Append(FMaterialHLSLGenerator::Version);
 	}
 
-	KeyString += SubstrateCompilationConfig.GetShaderMapKeyString();
+	SubstrateCompilationConfig.Append(KeyGen);
 }
 
 void FMaterialShaderMapId::SetShaderDependencies(const TArray<FShaderType*>& ShaderTypes, const TArray<const FShaderPipelineType*>& ShaderPipelineTypes, const TArray<FVertexFactoryType*>& VFTypes, EShaderPlatform ShaderPlatform)
 {
 	if (!FPlatformProperties::RequiresCookedData() && AllowShaderCompiling())
 	{
-		for (int32 ShaderTypeIndex = 0; ShaderTypeIndex < ShaderTypes.Num(); ShaderTypeIndex++)
+		for (const FShaderType* ShaderType : ShaderTypes)
 		{
-			FShaderType* ShaderType = ShaderTypes[ShaderTypeIndex];
 			FShaderTypeDependency Dependency;
 			Dependency.ShaderTypeName = ShaderType->GetHashedName();
 			Dependency.SourceHash = ShaderType->GetSourceHash(ShaderPlatform);
@@ -1235,17 +1321,16 @@ void FMaterialShaderMapId::SetShaderDependencies(const TArray<FShaderType*>& Sha
 			}
 		}
 
-		for (int32 VFTypeIndex = 0; VFTypeIndex < VFTypes.Num(); VFTypeIndex++)
+		for (const FVertexFactoryType* VFType : VFTypes)
 		{
 			FVertexFactoryTypeDependency Dependency;
-			Dependency.VertexFactoryTypeName = VFTypes[VFTypeIndex]->GetHashedName();
-			Dependency.VFSourceHash = VFTypes[VFTypeIndex]->GetSourceHash(ShaderPlatform);
+			Dependency.VertexFactoryTypeName = VFType->GetHashedName();
+			Dependency.VFSourceHash = VFType->GetSourceHash(ShaderPlatform);
 			VertexFactoryTypeDependencies.Add(Dependency);
 		}
 
-		for (int32 TypeIndex = 0; TypeIndex < ShaderPipelineTypes.Num(); TypeIndex++)
+		for (const FShaderPipelineType* Pipeline : ShaderPipelineTypes)
 		{
-			const FShaderPipelineType* Pipeline = ShaderPipelineTypes[TypeIndex];
 			FShaderPipelineTypeDependency Dependency;
 			Dependency.ShaderPipelineTypeName = Pipeline->GetHashedName();
 			Dependency.StagesSourceHash = Pipeline->GetSourceHash(ShaderPlatform);
diff --git a/Engine/Source/Runtime/Engine/Private/Materials/MaterialShared.cpp b/Engine/Source/Runtime/Engine/Private/Materials/MaterialShared.cpp
--- a/Engine/Source/Runtime/Engine/Private/Materials/MaterialShared.cpp
+++ b/Engine/Source/Runtime/Engine/Private/Materials/MaterialShared.cpp
@@ -31,6 +31,7 @@
 #include "Engine/Texture2D.h"
 #include "Engine/Font.h"
 #include "SceneView.h"
+#include "Serialization/ShaderKeyGenerator.h"
 #include "PSOPrecacheMaterial.h"
 #include "ShaderPlatformQualitySettings.h"
 #include "MaterialShaderQualitySettings.h"
@@ -5382,6 +5383,19 @@ static bool RemapParameterLayerIndex(TArrayView<const int32> IndexRemap, const T
 	return false;
 }
 
+void FMaterialParameterInfo::AppendString(FString& Out) const
+{
+	FShaderKeyGenerator KeyGen(Out);
+	Append(KeyGen);
+}
+
+void FMaterialParameterInfo::Append(FShaderKeyGenerator& KeyGen) const
+{
+	KeyGen.Append(Name);
+	KeyGen.Append(Association);
+	KeyGen.Append(Index);
+}
+
 bool FMaterialParameterInfo::RemapLayerIndex(TArrayView<const int32> IndexRemap, FMaterialParameterInfo& OutResult) const
 {
 	return RemapParameterLayerIndex(IndexRemap, *this, OutResult);
diff --git a/Engine/Source/Runtime/Engine/Private/ShaderCompiler/ShaderCompiler.cpp b/Engine/Source/Runtime/Engine/Private/ShaderCompiler/ShaderCompiler.cpp
--- a/Engine/Source/Runtime/Engine/Private/ShaderCompiler/ShaderCompiler.cpp
+++ b/Engine/Source/Runtime/Engine/Private/ShaderCompiler/ShaderCompiler.cpp
@@ -1224,16 +1224,28 @@ namespace GlobalShaderCookStats
 
 const FString& GetGlobalShaderMapDDCKey()
 {
-	static FString GlobalShaderMapDDCKey = FDevSystemGuids::GetSystemGuid(FDevSystemGuids::Get().GLOBALSHADERMAP_DERIVEDDATA_VER).ToString();
+	static FString GlobalShaderMapDDCKey = GetGlobalShaderMapDDCGuid().ToString();
 	return GlobalShaderMapDDCKey;
 }
 
+const FGuid& GetGlobalShaderMapDDCGuid()
+{
+	static FGuid GlobalShaderMapDDCGuid = FDevSystemGuids::GetSystemGuid(FDevSystemGuids::Get().GLOBALSHADERMAP_DERIVEDDATA_VER);
+	return GlobalShaderMapDDCGuid;
+}
+
 const FString& GetMaterialShaderMapDDCKey()
 {
-	static FString MaterialShaderMapDDCKey = FDevSystemGuids::GetSystemGuid(FDevSystemGuids::Get().MATERIALSHADERMAP_DERIVEDDATA_VER).ToString();
+	static FString MaterialShaderMapDDCKey = GetMaterialShaderMapDDCGuid().ToString();
 	return MaterialShaderMapDDCKey;
 }
 
+const FGuid& GetMaterialShaderMapDDCGuid()
+{
+	static FGuid MaterialShaderMapDDCGuid = FDevSystemGuids::GetSystemGuid(FDevSystemGuids::Get().MATERIALSHADERMAP_DERIVEDDATA_VER);
+	return MaterialShaderMapDDCGuid;
+}
+
 bool ShouldDumpShaderDDCKeys()
 {
 	return CVarShaderCompilerDumpDDCKeys.GetValueOnAnyThread();
@@ -9684,7 +9696,7 @@ static FString GetGlobalShaderMapKeyString(const FGlobalShaderMapId& ShaderMapId
 	check(EnvHashString.GetAllocatedSize() == 0);
 	ShaderMapKeyString.Append(EnvHashString.ToView());
 
-	return FString::Printf(TEXT("%s_%s_%s"), TEXT("GSM"), *GetGlobalShaderMapDDCKey(), *ShaderMapKeyString);
+	return FString::Printf(TEXT("%s_%s_%s"), TEXT("GSM"), *GetGlobalShaderMapDDCGuid().ToString(), *ShaderMapKeyString);
 }
 
 /** Creates a string key for the derived data cache entry for the global shader map. */
diff --git a/Engine/Source/Runtime/Engine/Private/ShaderCompiler/ShaderGenerationUtil.cpp b/Engine/Source/Runtime/Engine/Private/ShaderCompiler/ShaderGenerationUtil.cpp
--- a/Engine/Source/Runtime/Engine/Private/ShaderCompiler/ShaderGenerationUtil.cpp
+++ b/Engine/Source/Runtime/Engine/Private/ShaderCompiler/ShaderGenerationUtil.cpp
@@ -14,6 +14,7 @@
 #include "Misc/FileHelper.h"
 #include "RenderUtils.h"
 #include "SceneManagement.h"
+#include "Serialization/ShaderKeyGenerator.h"
 #include "ShaderCompilerDefinitions.h"
 #include "ShaderMaterial.h"
 
@@ -2038,13 +2039,24 @@ void FShaderCompileUtilities::ApplyDerivedDefines(FShaderCompilerEnvironment& Ou
 }
 
 void FShaderCompileUtilities::AppendGBufferDDCKeyString(const EShaderPlatform Platform, FString& KeyString)
+{
+	FShaderKeyGenerator KeyGen(KeyString);
+	AppendGBufferDDCKey(Platform, KeyGen);
+}
+
+void FShaderCompileUtilities::AppendGBufferDDCKey(const EShaderPlatform Platform, FShaderKeyGenerator& KeyGen)
 {
 	for (uint32 Layout = 0; Layout < GBL_Num; ++Layout)
 	{
 		FShaderGlobalDefines GlobalDefines = FetchShaderGlobalDefines(Platform, (EGBufferLayout)Layout);
-		KeyString.Appendf(TEXT("_%d%d%d"),GlobalDefines.GBUFFER_HAS_VELOCITY, GlobalDefines.GBUFFER_HAS_TANGENT, GlobalDefines.ALLOW_STATIC_LIGHTING);
-	}
-	KeyString.Appendf(TEXT("_%d;\n"), GBufferGeneratorVersion);
+		KeyGen.AppendSeparator();
+		KeyGen.AppendBoolInt(GlobalDefines.GBUFFER_HAS_VELOCITY);
+		KeyGen.AppendBoolInt(GlobalDefines.GBUFFER_HAS_TANGENT);
+		KeyGen.AppendBoolInt(GlobalDefines.ALLOW_STATIC_LIGHTING);
+	}
+	KeyGen.AppendSeparator();
+	KeyGen.Append(GBufferGeneratorVersion);
+	KeyGen.AppendDebugText(TEXTVIEW(";\n"));
 }
 
 static FGBufferInfo GLastGBufferInfo[SP_NumPlatforms] = {};
diff --git a/Engine/Source/Runtime/Engine/Public/MaterialShaderType.h b/Engine/Source/Runtime/Engine/Public/MaterialShaderType.h
--- a/Engine/Source/Runtime/Engine/Public/MaterialShaderType.h
+++ b/Engine/Source/Runtime/Engine/Public/MaterialShaderType.h
@@ -29,6 +29,7 @@ class FMaterialShaderMap;
 class FMaterialShaderMapId;
 class FShaderCommonCompileJob;
 class FShaderCompileJob;
+class FShaderKeyGenerator;
 class FUniformExpressionSet;
 class FVertexFactoryType;
 struct FMaterialShaderParameters;
@@ -51,6 +52,7 @@ extern ENGINE_API FString GetShadingModelFieldString(FMaterialShadingModelField
 /** Converts an EBlendMode to a string description. */
 extern ENGINE_API FString GetBlendModeString(EBlendMode BlendMode);
 
+#if WITH_EDITOR
 /** Creates a string key for the derived data cache given a shader map id. */
 UE_DEPRECATED(5.5, "GetMaterialShaderMapKeyString now requires passing an FMaterialShaderParameters instance as input.")
 inline FString GetMaterialShaderMapKeyString(const FMaterialShaderMapId& ShaderMapId, EShaderPlatform Platform, bool bIncludeKeyStringShaderDependencies = true)
@@ -63,6 +65,13 @@ extern ENGINE_API FString GetMaterialShaderMapKeyString(
 	const FMaterialShaderParameters& ShaderParameters,
 	EShaderPlatform Platform,
 	bool bIncludeKeyStringShaderDependencies = true);
+extern ENGINE_API void GetMaterialShaderMapKey(
+	FShaderKeyGenerator& KeyGen,
+	const FMaterialShaderMapId& ShaderMapId,
+	const FMaterialShaderParameters& ShaderParameters,
+	EShaderPlatform Platform,
+	bool bIncludeKeyStringShaderDependencies = true);
+#endif
 
 /** Called for every material shader to update the appropriate stats. */
 extern void UpdateMaterialShaderCompilingStats(const FMaterial* Material);
diff --git a/Engine/Source/Runtime/Engine/Public/MaterialShared.h b/Engine/Source/Runtime/Engine/Public/MaterialShared.h
--- a/Engine/Source/Runtime/Engine/Public/MaterialShared.h
+++ b/Engine/Source/Runtime/Engine/Public/MaterialShared.h
@@ -67,6 +67,7 @@ class FUniformExpressionSet;
 class FMeshMaterialShaderType;
 class FSceneView;
 class FShaderCommonCompileJob;
+class FShaderKeyGenerator;
 enum class EShaderCompileJobPriority : uint8;
 class FVirtualTexture2DResource;
 class IAllocatedVirtualTexture;
@@ -1167,7 +1168,11 @@ public:
 	void UpdateFromParameterSet(const FStaticParameterSet& StaticParameters);
 
 	/** Appends string representations of this Id to a key string. */
-	void AppendKeyString(FString& KeyString, bool bIncludeSourceAndMaterialState = true, bool bIncludeKeyStringShaderDependencies = true) const;
+	void AppendKeyString(FString& KeyString, bool bIncludeSourceAndMaterialState = true,
+		bool bIncludeKeyStringShaderDependencies = true) const;
+	void Append(FShaderKeyGenerator& KeyGen, bool bIncludeSourceAndMaterialState = true,
+		bool bIncludeKeyStringShaderDependencies = true) const;
+	UE_DEPRECATED(5.5, "Internal function only, called through AppendKeyString.")
 	void AppendStaticParametersString(FString& ParamsString) const;
 
 	const TArray<FStaticSwitchParameter> &GetStaticSwitchParameters() const 					{ return StaticSwitchParameters; }
@@ -1184,6 +1189,11 @@ public:
 	/** Returns true if the requested vertex factory type is a dependency of this shader map Id. */
 	bool ContainsVertexFactoryType(const FVertexFactoryType* VFType) const;
 #endif // WITH_EDITOR
+
+private:
+#if WITH_EDITOR
+	void AppendStaticParameters(FShaderKeyGenerator& KeyGen) const;
+#endif
 };
 
 /**
diff --git a/Engine/Source/Runtime/Engine/Public/MaterialTypes.h b/Engine/Source/Runtime/Engine/Public/MaterialTypes.h
--- a/Engine/Source/Runtime/Engine/Public/MaterialTypes.h
+++ b/Engine/Source/Runtime/Engine/Public/MaterialTypes.h
@@ -10,6 +10,7 @@
 #include "Shader/ShaderTypes.h"
 #include "MaterialTypes.generated.h"
 
+class FShaderKeyGenerator;
 class UTexture;
 class UCurveLinearColor;
 class UCurveLinearColorAtlas;
@@ -56,12 +57,8 @@ struct FMaterialParameterInfo
 
 	ENGINE_API explicit FMaterialParameterInfo(const struct FMemoryImageMaterialParameterInfo& Rhs);
 
-	void AppendString(FString& Out) const
-	{
-		Name.AppendString(Out);
-		Out.AppendInt(Association);
-		Out.AppendInt(Index);
-	}
+	ENGINE_API void AppendString(FString& Out) const;
+	ENGINE_API void Append(FShaderKeyGenerator& KeyGen) const;
 	FString ToString() const
 	{
 		FString Out;
@@ -469,11 +466,14 @@ struct FSubstrateCompilationConfig
 	int16 BytesPerPixelOverride = -1;
 	int16 ClosuresPerPixelOverride = -1;
 
+#if WITH_EDITOR
 	FString GetShaderMapKeyString() const;
+	void Append(FShaderKeyGenerator& KeyGen) const;
 
 	void UpdateHash(FSHA1& Hasher) const;
 
 	void Serialize(FArchive& Ar);
+#endif
 
 	friend inline bool operator==(const FSubstrateCompilationConfig& Lhs, const FSubstrateCompilationConfig& Rhs)
 	{
@@ -484,4 +484,12 @@ struct FSubstrateCompilationConfig
 	{
 		return !operator==(Lhs, Rhs);
 	}
+private:
+#if WITH_EDITOR
+	// Hidden friend for the FShaderKeyGenerator Append API
+	friend inline void Append(FShaderKeyGenerator& KeyGen, const FSubstrateCompilationConfig& Value)
+	{
+		Value.Append(KeyGen);
+	}
+#endif
 };
diff --git a/Engine/Source/Runtime/Engine/Public/Materials/MaterialAttributeDefinitionMap.h b/Engine/Source/Runtime/Engine/Public/Materials/MaterialAttributeDefinitionMap.h
--- a/Engine/Source/Runtime/Engine/Public/Materials/MaterialAttributeDefinitionMap.h
+++ b/Engine/Source/Runtime/Engine/Public/Materials/MaterialAttributeDefinitionMap.h
@@ -13,6 +13,7 @@ enum EMaterialShadingModel : int;
 enum EShaderFrequency : uint8;
 
 class FMaterialCompiler;
+class FShaderKeyGenerator;
 class UMaterial;
 
 struct FMaterialShadingModelField;
@@ -130,6 +131,7 @@ public:
 
 	/** Appends a hash of the property map intended for use with the DDC key */
 	static ENGINE_API void AppendDDCKeyString(FString& String);
+	static ENGINE_API void AppendDDCKey(FShaderKeyGenerator& KeyGen);
 
 	/** Appends a new attribute definition to the custom output list */
 	static ENGINE_API void AddCustomAttribute(const FGuid& AttributeID, const FString& AttributeName, const FString& FunctionName, EMaterialValueType ValueType, const FVector4& DefaultValue, MaterialAttributeBlendFunction BlendFunction = nullptr);
diff --git a/Engine/Source/Runtime/Engine/Public/ShaderCompiler.h b/Engine/Source/Runtime/Engine/Public/ShaderCompiler.h
--- a/Engine/Source/Runtime/Engine/Public/ShaderCompiler.h
+++ b/Engine/Source/Runtime/Engine/Public/ShaderCompiler.h
@@ -36,6 +36,7 @@ class IDistributedBuildController;
 class FMaterialShaderMap;
 class FShaderCompileJob;
 class FShaderCompilerStats;
+class FShaderKeyGenerator;
 class FShaderPipelineCompileJob;
 struct FAnalyticsEventAttribute;
 
@@ -264,6 +265,7 @@ public:
 	static ENGINE_API void GenerateBrdfHeaders(const FName& ShaderFormat);
 	static void ApplyDerivedDefines(FShaderCompilerEnvironment& OutEnvironment, FShaderCompilerEnvironment* SharedEnvironment, const EShaderPlatform Platform);
 	static void AppendGBufferDDCKeyString(const EShaderPlatform Platform, FString& KeyString);
+	static void AppendGBufferDDCKey(const EShaderPlatform Platform, FShaderKeyGenerator& KeyGen);
 	static ENGINE_API void WriteGBufferInfoAutogen(EShaderPlatform TargetPlatform, ERHIFeatureLevel::Type FeatureLevel);
 
 	static void ApplyFetchEnvironment(FShaderMaterialPropertyDefines& DefineData, const FShaderCompilerEnvironment& Environment);
@@ -1277,9 +1279,13 @@ extern ENGINE_API void CompileGlobalShaderMap(EShaderPlatform Platform, bool bRe
 extern ENGINE_API void CompileGlobalShaderMap(EShaderPlatform Platform, const ITargetPlatform* TargetPlatform, bool bRefreshShaderMap);
 extern ENGINE_API void ShutdownGlobalShaderMap();
 
+UE_DEPRECATED(5.5, "Use GetGlobalShaderMapDDCGuid")
 extern ENGINE_API const FString& GetGlobalShaderMapDDCKey();
+extern ENGINE_API const FGuid& GetGlobalShaderMapDDCGuid();
 
+UE_DEPRECATED(5.5, "Use GetMaterialShaderMapDDCGuid")
 extern ENGINE_API const FString& GetMaterialShaderMapDDCKey();
+extern ENGINE_API const FGuid& GetMaterialShaderMapDDCGuid();
 
 extern ENGINE_API bool ShouldDumpShaderDDCKeys();
 UE_DEPRECATED(5.4, "DumpShaderDDCKeyToFile now takes DebugGroupName as parameter (these files now go into the ShaderDebugInfo folder alongside other debug artifacts).")
diff --git a/Engine/Source/Runtime/Engine/Public/StaticParameterSet.h b/Engine/Source/Runtime/Engine/Public/StaticParameterSet.h
--- a/Engine/Source/Runtime/Engine/Public/StaticParameterSet.h
+++ b/Engine/Source/Runtime/Engine/Public/StaticParameterSet.h
@@ -14,6 +14,7 @@
 #include "Materials/MaterialLayersFunctions.h"
 #include "StaticParameterSet.generated.h"
 
+class FShaderKeyGenerator;
 class FSHA1;
 struct FStaticParameterSet;
 
@@ -76,12 +77,8 @@ struct FStaticParameterBase
 		HashState.Update((const uint8*)&Override, sizeof(Override));
 	}
 
-	void AppendKeyString(FString& KeyString) const
-	{
-		ParameterInfo.AppendString(KeyString);
-		KeyString.AppendInt(bOverride);
-		ExpressionGUID.AppendString(KeyString);
-	}
+	ENGINE_API void AppendKeyString(FString& KeyString) const;
+	ENGINE_API void Append(FShaderKeyGenerator& KeyGen) const;
 };
 
 
@@ -133,11 +130,8 @@ struct FStaticSwitchParameter : public FStaticParameterBase
 		HashState.Update((const uint8*)&HashValue, sizeof(HashValue));
 	}
 
-	void AppendKeyString(FString& KeyString) const
-	{
-		FStaticParameterBase::AppendKeyString(KeyString);
-		KeyString.AppendInt(Value);
-	}
+	ENGINE_API void AppendKeyString(FString& KeyString) const;
+	ENGINE_API void Append(FShaderKeyGenerator& KeyGen) const;
 
 	void GetValue(FMaterialParameterMetadata& OutResult) const
 	{
@@ -217,14 +211,8 @@ struct FStaticComponentMaskParameter : public FStaticParameterBase
 		HashState.Update((const uint8*)&Values, sizeof(Values));
 	}
 
-	void AppendKeyString(FString& KeyString) const
-	{
-		FStaticParameterBase::AppendKeyString(KeyString);
-		KeyString += FString::FromInt(R);
-		KeyString += FString::FromInt(G);
-		KeyString += FString::FromInt(B);
-		KeyString += FString::FromInt(A);
-	}
+	ENGINE_API void AppendKeyString(FString& KeyString) const;
+	ENGINE_API void Append(FShaderKeyGenerator& KeyGen) const;
 
 	void GetValue(FMaterialParameterMetadata& OutResult) const
 	{
@@ -338,13 +326,8 @@ struct FStaticTerrainLayerWeightParameter
 		Values[1] = bWeightBasedBlend;
 		HashState.Update((const uint8*)&Values, sizeof(Values));
 	}
-
-	void AppendKeyString(FString& KeyString) const
-	{
-		KeyString += LayerName.ToString();
-		KeyString += FString::FromInt(WeightmapIndex);
-		KeyString += FString::FromInt(bWeightBasedBlend);
-	}
+	ENGINE_API void AppendKeyString(FString& KeyString) const;
+	ENGINE_API void Append(FShaderKeyGenerator& KeyGen) const;
 };
 
 struct UE_DEPRECATED(5.0, "Material layers are no longer material parameters, use FStaticParameterSet::MaterialLayers") FStaticMaterialLayersParameter;
diff --git a/Engine/Source/Runtime/RenderCore/Private/Shader.cpp b/Engine/Source/Runtime/RenderCore/Private/Shader.cpp
--- a/Engine/Source/Runtime/RenderCore/Private/Shader.cpp
+++ b/Engine/Source/Runtime/RenderCore/Private/Shader.cpp
@@ -12,6 +12,7 @@
 #include "Interfaces/ITargetPlatform.h"
 #include "Interfaces/ITargetPlatformManagerModule.h"
 #include "Interfaces/IShaderFormat.h"
+#include "Serialization/ShaderKeyGenerator.h"
 #include "ShaderCodeLibrary.h"
 #include "ShaderCore.h"
 #include "ShaderCompilerCore.h"
@@ -1497,26 +1498,32 @@ bool IsUsingEmulatedUniformBuffers(EShaderPlatform Platform)
 }
 
 void ShaderMapAppendKeyString(EShaderPlatform Platform, FString& KeyString)
+{
+	FShaderKeyGenerator KeyGen(KeyString);
+	ShaderMapAppendKey(Platform, KeyGen);
+}
+
+void ShaderMapAppendKey(EShaderPlatform Platform, FShaderKeyGenerator& KeyGen)
 {
 	const FName ShaderFormatName = LegacyShaderPlatformToShaderFormat(Platform);
 
 	for (const FAutoConsoleObject* ConsoleObject : FAutoConsoleObject::AccessGeneralShaderChangeCvars())
 	{
 		FString ConsoleObjectName = IConsoleManager::Get().FindConsoleObjectName(ConsoleObject->AsVariable());
-		KeyString += TEXT("_");
-		KeyString += ConsoleObjectName;
-		KeyString += TEXT("_");
-		KeyString += ConsoleObject->AsVariable()->GetString();
+		KeyGen.AppendSeparator();
+		KeyGen.Append(ConsoleObjectName);
+		KeyGen.AppendSeparator();
+		KeyGen.Append(ConsoleObject->AsVariable()->GetString());
 	}
 	if (IsMobilePlatform(Platform))
 	{
 		for (const FAutoConsoleObject* ConsoleObject : FAutoConsoleObject::AccessMobileShaderChangeCvars())
 		{
 			FString ConsoleObjectName = IConsoleManager::Get().FindConsoleObjectName(ConsoleObject->AsVariable());
-			KeyString += TEXT("_");
-			KeyString += ConsoleObjectName;
-			KeyString += TEXT("_");
-			KeyString += ConsoleObject->AsVariable()->GetString();
+			KeyGen.AppendSeparator();
+			KeyGen.Append(ConsoleObjectName);
+			KeyGen.AppendSeparator();
+			KeyGen.Append(ConsoleObject->AsVariable()->GetString());
 		}
 	}
 	else if (IsConsolePlatform(Platform))
@@ -1524,10 +1531,10 @@ void ShaderMapAppendKeyString(EShaderPlatform Platform, FString& KeyString)
 		for (const FAutoConsoleObject* ConsoleObject : FAutoConsoleObject::AccessDesktopShaderChangeCvars())
 		{
 			FString ConsoleObjectName = IConsoleManager::Get().FindConsoleObjectName(ConsoleObject->AsVariable());
-			KeyString += TEXT("_");
-			KeyString += ConsoleObjectName;
-			KeyString += TEXT("_");
-			KeyString += ConsoleObject->AsVariable()->GetString();
+			KeyGen.AppendSeparator();
+			KeyGen.Append(ConsoleObjectName);
+			KeyGen.AppendSeparator();
+			KeyGen.Append(ConsoleObject->AsVariable()->GetString());
 		}
 	}
 
@@ -1536,27 +1543,32 @@ void ShaderMapAppendKeyString(EShaderPlatform Platform, FString& KeyString)
 
 	{
 		static const auto CVar = IConsoleManager::Get().FindTConsoleVariableDataInt(TEXT("Compat.UseDXT5NormalMaps"));
-		KeyString += (CVar && CVar->GetValueOnAnyThread() != 0) ? TEXT("_DXTN") : TEXT("_BC5N");
+		KeyGen.AppendSeparator();
+		KeyGen.Append((CVar && CVar->GetValueOnAnyThread() != 0) ? TEXT("DXTN") : TEXT("BC5N"));
 	}
 
 	{
 		static const auto CVar = IConsoleManager::Get().FindTConsoleVariableDataInt(TEXT("r.ClearCoatNormal"));
-		KeyString += (CVar && CVar->GetValueOnAnyThread() != 0) ? TEXT("_CCBN") : TEXT("_NoCCBN");
+		KeyGen.AppendSeparator();
+		KeyGen.Append((CVar && CVar->GetValueOnAnyThread() != 0) ? TEXT("CCBN") : TEXT("NoCCBN"));
 	}
 
 	{
 		static const auto CVar = IConsoleManager::Get().FindTConsoleVariableDataInt(TEXT("r.IrisNormal"));
-		KeyString += (CVar && CVar->GetValueOnAnyThread() != 0) ? TEXT("_Iris") : TEXT("_NoIris");
+		KeyGen.AppendSeparator();
+		KeyGen.Append((CVar && CVar->GetValueOnAnyThread() != 0) ? TEXT("Iris") : TEXT("NoIris"));
 	}
 
 	{
 		static const auto CVar = IConsoleManager::Get().FindTConsoleVariableDataInt(TEXT("r.CompileShadersForDevelopment"));
-		KeyString += (CVar && CVar->GetValueOnAnyThread() != 0) ? TEXT("_DEV") : TEXT("_NoDEV");
+		KeyGen.AppendSeparator();
+		KeyGen.Append((CVar && CVar->GetValueOnAnyThread() != 0) ? TEXT("DEV") : TEXT("NoDEV"));
 	}
 
 	{
 		const bool bValue = IsStaticLightingAllowed();
-		KeyString += bValue ? TEXT("_SL") : TEXT("_NoSL");
+		KeyGen.AppendSeparator();
+		KeyGen.Append(bValue ? TEXT("SL") : TEXT("NoSL"));
 	}
 
 	{
@@ -1564,16 +1576,20 @@ void ShaderMapAppendKeyString(EShaderPlatform Platform, FString& KeyString)
 		const int32 LWCTruncateValue = CVar ? CVar->GetValueOnAnyThread() : 0;
 		if (LWCTruncateValue == 1)
 		{
-			KeyString += TEXT("_LWC1");
+			KeyGen.AppendSeparator();
+			KeyGen.Append(TEXT("LWC1"));
 		}
 		else if (LWCTruncateValue == 2)
 		{
-			KeyString += TEXT("_LWC2");
+			KeyGen.AppendSeparator();
+			KeyGen.Append(TEXT("LWC2"));
 		}
 	}
 
+	if (IsUsingBasePassVelocity(Platform))
 	{
-		KeyString += IsUsingBasePassVelocity(Platform) ? TEXT("_GV") : TEXT("");
+		KeyGen.AppendSeparator();
+		KeyGen.Append(TEXT("GV"));
 	}
 
 	{
@@ -1581,63 +1597,100 @@ void ShaderMapAppendKeyString(EShaderPlatform Platform, FString& KeyString)
 
 		if (Aspects.IsInstancedStereoEnabled())
 		{
-			KeyString += TEXT("_VRIS");
+			KeyGen.AppendSeparator();
+			KeyGen.Append(TEXT("VRIS"));
 
 			if (Aspects.IsInstancedMultiViewportEnabled())
 			{
-				KeyString += TEXT("_MVIEW");
+				KeyGen.AppendSeparator();
+				KeyGen.Append(TEXT("MVIEW"));
 			}
 		}
 
 		if (Aspects.IsMobileMultiViewEnabled())
 		{
-			KeyString += TEXT("_MMVIEW");
+			KeyGen.AppendSeparator();
+			KeyGen.Append(TEXT("MMVIEW"));
 		}
 	}
 
 	{
-		KeyString += IsUsingSelectiveBasePassOutputs(Platform) ? TEXT("_SO") : TEXT("");
+		if (IsUsingSelectiveBasePassOutputs(Platform))
+		{
+			KeyGen.AppendSeparator();
+			KeyGen.Append(TEXT("SO"));
+		}
 	}
 
 	{
 		// PreExposure is always used
-		KeyString += TEXT("_PreExp");
+		KeyGen.AppendSeparator();
+		KeyGen.Append(TEXT("PreExp"));
 	}
 
 	{
-		KeyString += IsUsingDBuffers(Platform) ? TEXT("_DBuf") : TEXT("_NoDBuf");
+		KeyGen.AppendSeparator();
+		KeyGen.Append(IsUsingDBuffers(Platform) ? TEXT("DBuf") : TEXT("NoDBuf"));
 	}
 
 	{
 		static IConsoleVariable* CVar = IConsoleManager::Get().FindConsoleVariable(TEXT("r.AllowGlobalClipPlane"));
-		KeyString += (CVar && CVar->GetInt() != 0) ? TEXT("_ClipP") : TEXT("");
+		if (CVar && CVar->GetInt() != 0)
+		{
+			KeyGen.AppendSeparator();
+			KeyGen.Append(TEXT("ClipP"));
+		}
 	}
 
 	{
 		// Extra data (names, etc)
-		KeyString += ShouldEnableExtraShaderData(ShaderFormatName) ? TEXT("_ExtraData") : TEXT("");
+		if (ShouldEnableExtraShaderData(ShaderFormatName))
+		{
+			KeyGen.AppendSeparator();
+			KeyGen.Append(TEXT("ExtraData"));
+		}
 		// Symbols and/or SymbolsInfo and version if symbols serialization changes
-		KeyString += ShouldGenerateShaderSymbols(ShaderFormatName) ? TEXT("_Symbols") : TEXT("");
-		KeyString += ShouldGenerateShaderSymbolsInfo(ShaderFormatName) ? TEXT("_SymbolsInfo") : TEXT("");
+		if (ShouldGenerateShaderSymbols(ShaderFormatName))
+		{
+			KeyGen.AppendSeparator();
+			KeyGen.Append(TEXT("Symbols"));
+		}
+		if (ShouldGenerateShaderSymbolsInfo(ShaderFormatName))
+		{
+			KeyGen.AppendSeparator();
+			KeyGen.Append(TEXT("SymbolsInfo"));
+		}
 		// Are symbols based on source or results
-		KeyString += ShouldAllowUniqueShaderSymbols(ShaderFormatName) ? TEXT("_FullDbg") : TEXT("");
+		if (ShouldAllowUniqueShaderSymbols(ShaderFormatName))
+		{
+			KeyGen.AppendSeparator();
+			KeyGen.Append(TEXT("FullDbg"));
+		}
 	}
 
+	if (!ShouldOptimizeShaders(ShaderFormatName))
 	{
-		KeyString += ShouldOptimizeShaders(ShaderFormatName) ? TEXT("") : TEXT("_NoOpt");
+		KeyGen.AppendSeparator();
+		KeyGen.Append(TEXT("NoOpt"));
 	}
 	
 	{
 		// Always default to fast math unless specified
 		static const auto CVar = IConsoleManager::Get().FindConsoleVariable(TEXT("r.Shaders.FastMath"));
-		KeyString += (CVar && CVar->GetInt() == 0) ? TEXT("_NoFastMath") : TEXT("");
+		if ((CVar && CVar->GetInt() == 0))
+		{
+			KeyGen.AppendSeparator();
+			KeyGen.Append(TEXT("NoFastMath"));
+		}
 	}
 
 	{
 		static FShaderPlatformCachedIniValue<int32> CVarWarningsAsErrorsPerPlatform(TEXT("r.Shaders.WarningsAsErrors"));
 		if (const int32 Level = CVarWarningsAsErrorsPerPlatform.Get(Platform); Level != 0)
 		{
-			KeyString.Appendf(TEXT("_WX%d"), Level);
+			KeyGen.AppendSeparator();
+			KeyGen.Append(TEXT("WX"));
+			KeyGen.Append(Level);
 		}
 	}
 	
@@ -1646,7 +1699,9 @@ void ShaderMapAppendKeyString(EShaderPlatform Platform, FString& KeyString)
 		// Note: Since 1 is the default, we don't modify the hash for this case, so as to not force a rebuild, and to keep the hash shorter.
 		if (CVar && (CVar->GetInt() == 0 || CVar->GetInt() == 2))
 		{
-			KeyString.Appendf(TEXT("_C%d"), CVar->GetInt());
+			KeyGen.AppendSeparator();
+			KeyGen.Append(TEXT("C"));
+			KeyGen.Append(CVar->GetInt());
 		}
 	}
 	
@@ -1657,10 +1712,12 @@ void ShaderMapAppendKeyString(EShaderPlatform Platform, FString& KeyString)
 			switch(CVar->GetInt())
 			{
 				case 2:
-					KeyString += TEXT("_AvoidFlow");
+					KeyGen.AppendSeparator();
+					KeyGen.Append(TEXT("AvoidFlow"));
 					break;
 				case 1:
-					KeyString += TEXT("_PreferFlow");
+					KeyGen.AppendSeparator();
+					KeyGen.Append(TEXT("PreferFlow"));
 					break;
 				case 0:
 				default:
@@ -1671,61 +1728,89 @@ void ShaderMapAppendKeyString(EShaderPlatform Platform, FString& KeyString)
 
 	if (!AllowPixelDepthOffset(Platform))
 	{
-		KeyString += TEXT("_NoPDO");
+		KeyGen.AppendSeparator();
+		KeyGen.Append(TEXT("NoPDO"));
 	}
 
 	if (!AllowPerPixelShadingModels(Platform))
 	{
-		KeyString += TEXT("_NoPPSM");
+		KeyGen.AppendSeparator();
+		KeyGen.Append(TEXT("NoPPSM"));
 	}
 	
 	if (UseRemoveUnsedInterpolators(Platform) && !IsOpenGLPlatform(Platform))
 	{
-		KeyString += TEXT("_UnInt");
+		KeyGen.AppendSeparator();
+		KeyGen.Append(TEXT("UnInt"));
 	}
 
 	if (ForwardShadingForcesSkyLightCubemapBlending(Platform))
 	{
-		KeyString += TEXT("_FwdSkyBlnd");
+		KeyGen.AppendSeparator();
+		KeyGen.Append(TEXT("FwdSkyBlnd"));
 	}
 
 	if (IsMobilePlatform(Platform))
 	{
 		{
 			static IConsoleVariable* CVar = IConsoleManager::Get().FindConsoleVariable(TEXT("r.Mobile.DisableVertexFog"));
-			KeyString += (CVar && CVar->GetInt() != 0) ? TEXT("_NoVFog") : TEXT("");
+			if (CVar && CVar->GetInt() != 0)
+			{
+				KeyGen.AppendSeparator();
+				KeyGen.Append(TEXT("NoVFog"));
+			}
 		}
 		
 		{
 			static IConsoleVariable* CVar = IConsoleManager::Get().FindConsoleVariable(TEXT("r.Mobile.FloatPrecisionMode"));
-			if(CVar && CVar->GetInt() > 0)
+			if (CVar && CVar->GetInt() > 0)
 			{
-				KeyString.Appendf(TEXT("_highp%d"), CVar->GetInt());
+				KeyGen.AppendSeparator();
+				KeyGen.Append(TEXT("highp"));
+				KeyGen.Append(CVar->GetInt());
 			}
 		}
 
 		{
 			static IConsoleVariable* CVar = IConsoleManager::Get().FindConsoleVariable(TEXT("r.Mobile.AllowDitheredLODTransition"));
-			KeyString += (CVar && CVar->GetInt() != 0) ? TEXT("_DLODT") : TEXT("");
+			if (CVar && CVar->GetInt() != 0)
+			{
+				KeyGen.AppendSeparator();
+				KeyGen.Append(TEXT("DLODT"));
+			}
 		}
 		
-		KeyString += IsUsingEmulatedUniformBuffers(Platform) ? TEXT("_NoUB") : TEXT("");
+		if (IsUsingEmulatedUniformBuffers(Platform))
+		{
+			KeyGen.AppendSeparator();
+			KeyGen.Append(TEXT("NoUB"));
+		}
 
 		{
 			const bool bMobileMovableSpotlightShadowsEnabled = IsMobileMovableSpotlightShadowsEnabled(Platform);
-			KeyString += bMobileMovableSpotlightShadowsEnabled ? TEXT("S") : TEXT("");
+			if (bMobileMovableSpotlightShadowsEnabled)
+			{
+				KeyGen.Append(TEXT("S"));
+			}
 		}
 		
 		{
 			static IConsoleVariable* CVar = IConsoleManager::Get().FindConsoleVariable(TEXT("r.Mobile.UseHWsRGBEncoding"));
-			KeyString += (CVar && CVar->GetInt() != 0) ? TEXT("_HWsRGB") : TEXT("");
+			if (CVar && CVar->GetInt() != 0)
+			{
+				KeyGen.AppendSeparator();
+				KeyGen.Append(TEXT("HWsRGB"));
+			}
 		}
 		
 		{
 			// make it per shader platform ?
 			static IConsoleVariable* CVar = IConsoleManager::Get().FindConsoleVariable(TEXT("r.Mobile.SupportGPUScene"));
-			bool bMobileGpuScene = (CVar && CVar->GetInt() != 0);
-			KeyString += bMobileGpuScene ? TEXT("_MobGPUSc") : TEXT("");
+			if (CVar && CVar->GetInt() != 0)
+			{
+				KeyGen.AppendSeparator();
+				KeyGen.Append(TEXT("MobGPUSc"));
+			}
 		}
 
 		{
@@ -1733,14 +1818,16 @@ void ShaderMapAppendKeyString(EShaderPlatform Platform, FString& KeyString)
 
 			if (bIsMobileDeferredShading)
 			{
-				KeyString += (MobileUsesExtenedGBuffer(Platform) ? TEXT("_MobDShEx") : TEXT("_MobDSh"));
+				KeyGen.AppendSeparator();
+				KeyGen.Append(MobileUsesExtenedGBuffer(Platform) ? TEXT("MobDShEx") : TEXT("MobDSh"));
 			}
 			else
 			{
 				static IConsoleVariable* MobileForwardEnableClusteredReflectionsCVAR = IConsoleManager::Get().FindConsoleVariable(TEXT("r.Mobile.Forward.EnableClusteredReflections"));
 				if (MobileForwardEnableClusteredReflectionsCVAR && MobileForwardEnableClusteredReflectionsCVAR->GetInt() != 0)
 				{
-					KeyString += TEXT("_MobFCR");
+					KeyGen.AppendSeparator();
+					KeyGen.Append(TEXT("MobFCR"));
 				}
 			}
 		}
@@ -1749,45 +1836,65 @@ void ShaderMapAppendKeyString(EShaderPlatform Platform, FString& KeyString)
 			static IConsoleVariable* MobileGTAOPreIntegratedTextureTypeCVar = IConsoleManager::Get().FindConsoleVariable(TEXT("r.Mobile.GTAOPreIntegratedTextureType"));
 			static IConsoleVariable* MobileAmbientOcclusionCVar = IConsoleManager::Get().FindConsoleVariable(TEXT("r.Mobile.AmbientOcclusion"));
 			int32 GTAOPreIntegratedTextureType = MobileGTAOPreIntegratedTextureTypeCVar ? MobileGTAOPreIntegratedTextureTypeCVar->GetInt() : 0;
-			KeyString += ((MobileAmbientOcclusionCVar && MobileAmbientOcclusionCVar->GetInt() != 0) && IsMobileHDR()) ? FString::Printf(TEXT("_MobileAO_%d"), GTAOPreIntegratedTextureType) : TEXT("");
+			if (MobileAmbientOcclusionCVar && MobileAmbientOcclusionCVar->GetInt() != 0 && IsMobileHDR())
+			{
+				KeyGen.AppendSeparator();
+				KeyGen.Append(TEXT("MobileAO"));
+				KeyGen.AppendSeparator();
+				KeyGen.Append(GTAOPreIntegratedTextureType);
+			}
 		}
 
+		if (IsMobileDistanceFieldEnabled(Platform))
 		{
-			KeyString += IsMobileDistanceFieldEnabled(Platform) ? TEXT("_MobSDF") : TEXT("");
+			KeyGen.AppendSeparator();
+			KeyGen.Append(TEXT("MobSDF"));
 		}
 
 		{
 			static FShaderPlatformCachedIniValue<bool> EnableCullBeforeFetchIniValue(TEXT("r.CullBeforeFetch"));
 			if (EnableCullBeforeFetchIniValue.Get(Platform) == 1)
 			{
-				KeyString += TEXT("_CBF");
+				KeyGen.AppendSeparator();
+				KeyGen.Append(TEXT("CBF"));
 			}
 			static FShaderPlatformCachedIniValue<bool> EnableWarpCullingIniValue(TEXT("r.WarpCulling"));
 			if (EnableWarpCullingIniValue.Get(Platform) == 1)
 			{
-				KeyString += TEXT("_WC");
+				KeyGen.AppendSeparator();
+				KeyGen.Append(TEXT("WC"));
 			}
 		}
 
+		if (MobileUsesFullDepthPrepass(Platform))
 		{
-			KeyString += MobileUsesFullDepthPrepass(Platform) ? TEXT("_MobFDP") : TEXT("");
+			KeyGen.AppendSeparator();
+			KeyGen.Append(TEXT("MobFDP"));
 		}
 	}
 	else
 	{
-		KeyString += IsUsingEmulatedUniformBuffers(Platform) ? TEXT("_NoUB") : TEXT("");
+		if (IsUsingEmulatedUniformBuffers(Platform))
+		{
+			KeyGen.AppendSeparator();
+			KeyGen.Append(TEXT("NoUB"));
+		}
 	}
 
 	uint32 PlatformShadingModelsMask = GetPlatformShadingModelsMask(Platform);
 	if (PlatformShadingModelsMask != 0xFFFFFFFF)
 	{
-		KeyString +=  FString::Printf(TEXT("SMM_%X"), PlatformShadingModelsMask);
+		KeyGen.Append(TEXT("SMM"));
+		KeyGen.AppendSeparator();
+		KeyGen.AppendHex(PlatformShadingModelsMask);
 	}
 
 	const IShaderFormat* ShaderFormat = GetTargetPlatformManagerRef().FindShaderFormat(ShaderFormatName);
 	if (ShaderFormat)
 	{
-		ShaderFormat->AppendToKeyString(KeyString);
+		FString ShaderFormatExtraData;
+		ShaderFormat->AppendToKeyString(ShaderFormatExtraData);
+		KeyGen.Append(ShaderFormatExtraData);
 	}
 	
 	// Encode the Metal standard into the shader compile options so that they recompile if the settings change.
@@ -1795,19 +1902,35 @@ void ShaderMapAppendKeyString(EShaderPlatform Platform, FString& KeyString)
 	{
 		{
 			static const auto CVar = IConsoleManager::Get().FindConsoleVariable(TEXT("r.Shaders.ZeroInitialise"));
-			KeyString += (CVar && CVar->GetInt() != 0) ? TEXT("_ZeroInit") : TEXT("");
+			if (CVar && CVar->GetInt() != 0)
+			{
+				KeyGen.AppendSeparator();
+				KeyGen.Append(TEXT("ZeroInit"));
+			}
 		}
 		{
 			static const auto CVar = IConsoleManager::Get().FindConsoleVariable(TEXT("r.Shaders.BoundsChecking"));
-			KeyString += (CVar && CVar->GetInt() != 0) ? TEXT("_BoundsChecking") : TEXT("");
+			if (CVar && CVar->GetInt() != 0)
+			{
+				KeyGen.AppendSeparator();
+				KeyGen.Append(TEXT("BoundsChecking"));
+			}
 		}
 		{
-			KeyString += RHISupportsManualVertexFetch(Platform) ? TEXT("_MVF_") : TEXT("");
+			if (RHISupportsManualVertexFetch(Platform))
+			{
+				KeyGen.AppendSeparator();
+				KeyGen.Append(TEXT("MVF"));
+				KeyGen.AppendSeparator();
+			}
 		}
 		
 		uint32 ShaderVersion = RHIGetMetalShaderLanguageVersion(Platform);
-		KeyString += FString::Printf(TEXT("_MTLSTD%u_"), ShaderVersion);
-		
+		KeyGen.AppendSeparator();
+		KeyGen.Append(TEXT("MTLSTD"));
+		KeyGen.Append(ShaderVersion);
+		KeyGen.AppendSeparator();
+
 		bool bAllowFastIntrinsics = false;
 		bool bEnableMathOptimisations = true;
 		bool bForceFloats = false;
@@ -1831,33 +1954,40 @@ void ShaderMapAppendKeyString(EShaderPlatform Platform, FString& KeyString)
 		
 		if (bAllowFastIntrinsics)
 		{
-			KeyString += TEXT("_MTLSL_FastIntrin");
+			KeyGen.AppendSeparator();
+			KeyGen.Append(TEXT("MTLSL_FastIntrin"));
 		}
 		
 		// Same as console-variable above, but that's global and this is per-platform, per-project
 		if (!bEnableMathOptimisations)
 		{
-			KeyString += TEXT("_NoFastMath");
+			KeyGen.AppendSeparator();
+			KeyGen.Append(TEXT("NoFastMath"));
 		}
 		
 		if (bForceFloats)
 		{
-			KeyString += TEXT("_FP32");
+			KeyGen.AppendSeparator();
+			KeyGen.Append(TEXT("FP32"));
 		}
 		
         if(bSupportAppleA8)
         {
-            KeyString += TEXT("_A8GPU");
+			KeyGen.AppendSeparator();
+			KeyGen.Append(TEXT("A8GPU"));
         }
         
-		KeyString += FString::Printf(TEXT("_IAB%d"), IndirectArgumentTier);
+		KeyGen.AppendSeparator();
+		KeyGen.Append(TEXT("IAB"));
+		KeyGen.Append(IndirectArgumentTier);
 		
 		// Shaders built for archiving - for Metal that requires compiling the code in a different way so that we can strip it later
 		bool bArchive = false;
 		GConfig->GetBool(TEXT("/Script/UnrealEd.ProjectPackagingSettings"), TEXT("bSharedMaterialNativeLibraries"), bArchive, GGameIni);
 		if (bArchive)
 		{
-			KeyString += TEXT("_ARCHIVE");
+			KeyGen.AppendSeparator();
+			KeyGen.Append(TEXT("ARCHIVE"));
 		}
 	}
 
@@ -1867,19 +1997,22 @@ void ShaderMapAppendKeyString(EShaderPlatform Platform, FString& KeyString)
 		GConfig->GetBool(TEXT("/Script/AndroidRuntimeSettings.AndroidRuntimeSettings"), TEXT("bStripShaderReflection"), bStripReflect, GEngineIni);
 		if (!bStripReflect)
 		{
-			KeyString += TEXT("_NoStripReflect");
+			KeyGen.AppendSeparator();
+			KeyGen.Append(TEXT("NoStripReflect"));
 		}
 	}
 
 	// Is DXC shader compiler enabled for this platform?
-	KeyString += (IsDxcEnabledForPlatform(Platform) ? TEXT("_DXC1") : TEXT("_DXC0"));
+	KeyGen.AppendSeparator();
+	KeyGen.Append(IsDxcEnabledForPlatform(Platform) ? TEXT("DXC1") : TEXT("DXC0"));
 
 	if (IsFeatureLevelSupported(Platform, ERHIFeatureLevel::SM5))
 	{
 		static const auto CVar = IConsoleManager::Get().FindTConsoleVariableDataInt(TEXT("r.StencilForLODDither"));
 		if (CVar && CVar->GetValueOnAnyThread() > 0)
 		{
-			KeyString += TEXT("_SD");
+			KeyGen.AppendSeparator();
+			KeyGen.Append(TEXT("SD"));
 		}
 	}
 
@@ -1901,7 +2034,8 @@ void ShaderMapAppendKeyString(EShaderPlatform Platform, FString& KeyString)
 
 		if (bForwardShading)
 		{
-			KeyString += TEXT("_FS");
+			KeyGen.AppendSeparator();
+			KeyGen.Append(TEXT("FS"));
 		}
 	}
 
@@ -1926,11 +2060,13 @@ void ShaderMapAppendKeyString(EShaderPlatform Platform, FString& KeyString)
 		{
 			if (PropagateAlphaType == 2)
 			{
-				KeyString += TEXT("_SA2");
+				KeyGen.AppendSeparator();
+				KeyGen.Append(TEXT("SA2"));
 			}
 			else
 			{
-				KeyString += TEXT("_SA");
+				KeyGen.AppendSeparator();
+				KeyGen.Append(TEXT("SA"));
 			}
 		}
 	}
@@ -1938,7 +2074,8 @@ void ShaderMapAppendKeyString(EShaderPlatform Platform, FString& KeyString)
 	if (TargetPlatform && 
 		TargetPlatform->SupportsFeature(ETargetPlatformFeatures::NormalmapLAEncodingMode))
 	{
-		KeyString += TEXT("_NLA");
+		KeyGen.AppendSeparator();
+		KeyGen.Append(TEXT("NLA"));
 	}
 
 	{
@@ -1958,7 +2095,8 @@ void ShaderMapAppendKeyString(EShaderPlatform Platform, FString& KeyString)
 		}
 		if (bVertexFoggingForOpaque)
 		{
-			KeyString += TEXT("_VFO");
+			KeyGen.AppendSeparator();
+			KeyGen.Append(TEXT("VFO"));
 		}
 	}
 
@@ -1968,14 +2106,16 @@ void ShaderMapAppendKeyString(EShaderPlatform Platform, FString& KeyString)
 		bSupportLocalFogVolumes = CVar && CVar->GetInt() > 0;
 		if (bSupportLocalFogVolumes)
 		{
-			KeyString += TEXT("_LFV");
+			KeyGen.AppendSeparator();
+			KeyGen.Append(TEXT("LFV"));
 		}
 	}
 
 	{
 		if (DoesProjectSupportLumenRayTracedTranslucentRefraction())
 		{
-			KeyString += TEXT("_LTRRT");
+			KeyGen.AppendSeparator();
+			KeyGen.Append(TEXT("LTRRT"));
 		}
 	}
 
@@ -1984,7 +2124,8 @@ void ShaderMapAppendKeyString(EShaderPlatform Platform, FString& KeyString)
 		const bool bLocalFogVolumesApplyOnTranclucent = CVar && CVar->GetInt() > 0;
 		if (bSupportLocalFogVolumes && bLocalFogVolumesApplyOnTranclucent)
 		{
-			KeyString += TEXT("_LFVTRA");
+			KeyGen.AppendSeparator();
+			KeyGen.Append(TEXT("LFVTRA"));
 		}
 	}
 
@@ -1992,12 +2133,14 @@ void ShaderMapAppendKeyString(EShaderPlatform Platform, FString& KeyString)
 		static const auto CVar = IConsoleManager::Get().FindTConsoleVariableDataInt(TEXT("r.SupportSkyAtmosphere"));
 		if (CVar && CVar->GetValueOnAnyThread() > 0)
 		{
-			KeyString += TEXT("_SKYATM");
+			KeyGen.AppendSeparator();
+			KeyGen.Append(TEXT("SKYATM"));
 
 			static const auto CVarHeightFog = IConsoleManager::Get().FindTConsoleVariableDataInt(TEXT("r.SupportSkyAtmosphereAffectsHeightFog"));
 			if (CVarHeightFog && CVarHeightFog->GetValueOnAnyThread() > 0)
 			{
-				KeyString += TEXT("_SKYHF");
+				KeyGen.AppendSeparator();
+				KeyGen.Append(TEXT("SKYHF"));
 			}
 		}
 	}
@@ -2005,14 +2148,16 @@ void ShaderMapAppendKeyString(EShaderPlatform Platform, FString& KeyString)
 	const bool bNeedsSeparateMainDirLightTexture = IsWaterDistanceFieldShadowEnabled(Platform) || IsWaterVirtualShadowMapFilteringEnabled(Platform);
 	if (bNeedsSeparateMainDirLightTexture)
 	{
-		KeyString += TEXT("_SLWSMDLT");
+		KeyGen.AppendSeparator();
+		KeyGen.Append(TEXT("SLWSMDLT"));
 	}
 
 	{
 		static const auto CVar = IConsoleManager::Get().FindTConsoleVariableDataInt(TEXT("r.SupportCloudShadowOnForwardLitTranslucent"));
 		if (CVar && CVar->GetValueOnAnyThread() > 0)
 		{
-			KeyString += TEXT("_CLDTRANS");
+			KeyGen.AppendSeparator();
+			KeyGen.Append(TEXT("CLDTRANS"));
 		}
 	}
 
@@ -2020,7 +2165,8 @@ void ShaderMapAppendKeyString(EShaderPlatform Platform, FString& KeyString)
 		const bool bTranslucentUsesLightRectLights = GetTranslucentUsesLightRectLights();
 		if (bTranslucentUsesLightRectLights)
 		{
-			KeyString += TEXT("_RECTTRANS");
+			KeyGen.AppendSeparator();
+			KeyGen.Append(TEXT("RECTTRANS"));
 		}
 	}
 
@@ -2028,7 +2174,8 @@ void ShaderMapAppendKeyString(EShaderPlatform Platform, FString& KeyString)
 		const bool bTranslucentUsesLightIESProfiles = GetTranslucentUsesLightIESProfiles();
 		if (bTranslucentUsesLightIESProfiles)
 		{
-			KeyString += TEXT("_IESTRANS");
+			KeyGen.AppendSeparator();
+			KeyGen.Append(TEXT("IESTRANS"));
 		}
 	}
 
@@ -2036,70 +2183,88 @@ void ShaderMapAppendKeyString(EShaderPlatform Platform, FString& KeyString)
 		static const auto CVar = IConsoleManager::Get().FindTConsoleVariableDataInt(TEXT("r.Shadow.Virtual.TranslucentQuality"));
 		if (CVar && CVar->GetValueOnAnyThread() > 0)
 		{
-			KeyString += TEXT("_VSMTRANSQUALITY");
+			KeyGen.AppendSeparator();
+			KeyGen.Append(TEXT("VSMTRANSQUALITY"));
 		}
 	}
 
 	if (GetHairStrandsUsesTriangleStrips())
 	{
-		KeyString += FString::Printf(TEXT("_STRDSTRIP"));
+		KeyGen.AppendSeparator();
+		KeyGen.Append(TEXT("STRDSTRIP"));
 	}
 
 	if (Substrate::IsSubstrateEnabled())
 	{
 		{
-			KeyString += TEXT("_SUBSTRATE");
+			KeyGen.AppendSeparator();
+			KeyGen.Append(TEXT("SUBSTRATE"));
 		}
 
 		{
-			KeyString += FString::Printf(TEXT("_BUDGET%u"), Substrate::GetBytePerPixel(Platform));
+			KeyGen.AppendSeparator();
+			KeyGen.Append(TEXT("BUDGET"));
+			KeyGen.Append(Substrate::GetBytePerPixel(Platform));
 		}
 
 		{
-			KeyString += FString::Printf(TEXT("_CLOSURE%u"), Substrate::GetClosurePerPixel(Platform));
+			KeyGen.AppendSeparator();
+			KeyGen.Append(TEXT("CLOSURE"));
+			KeyGen.Append(Substrate::GetClosurePerPixel(Platform));
 		}
 
 		if (Substrate::IsDBufferPassEnabled(Platform))
 		{
-			KeyString += FString::Printf(TEXT("_DBUFFERPASS"));
+			KeyGen.AppendSeparator();
+			KeyGen.Append(TEXT("DBUFFERPASS"));
 		}
 
 		if (Substrate::IsBackCompatibilityEnabled())
 		{
-			KeyString += FString::Printf(TEXT("_BACKCOMPAT"));
+			KeyGen.AppendSeparator();
+			KeyGen.Append(TEXT("BACKCOMPAT"));
 		}
 
 		if (Substrate::IsOpaqueRoughRefractionEnabled())
 		{
-			KeyString += FString::Printf(TEXT("_ROUGHDIFF"));
+			KeyGen.AppendSeparator();
+			KeyGen.Append(TEXT("ROUGHDIFF"));
 		}
 
 		if (Substrate::GetNormalQuality() > 0)
 		{
-			KeyString += FString::Printf(TEXT("_STRTNRMQ"));
+			KeyGen.AppendSeparator();
+			KeyGen.Append(TEXT("STRTNRMQ"));
 		}
 
 		if (Substrate::IsAdvancedVisualizationEnabled())
 		{
-			KeyString += FString::Printf(TEXT("_ADVDEBUG"));
+			KeyGen.AppendSeparator();
+			KeyGen.Append(TEXT("ADVDEBUG"));
 		}
 
 		{
-			KeyString += FString::Printf(TEXT("_STSHQL%u"), Substrate::GetShadingQuality(Platform));
+			KeyGen.AppendSeparator();
+			KeyGen.Append(TEXT("STSHQL"));
+			KeyGen.Append(Substrate::GetShadingQuality(Platform));
 		}
 
 		{
-			KeyString += FString::Printf(TEXT("_SSHEEN%u"), Substrate::GetSheenQuality(Platform));
+			KeyGen.AppendSeparator();
+			KeyGen.Append(TEXT("SSHEEN"));
+			KeyGen.Append(Substrate::GetSheenQuality(Platform));
 		}
 
 		if (Substrate::IsGlintEnabled(Platform))
 		{
-			KeyString += FString::Printf(TEXT("_STRTGLT"));
+			KeyGen.AppendSeparator();
+			KeyGen.Append(TEXT("STRTGLT"));
 		}
 
 		if (Substrate::IsSpecularProfileEnabled(Platform))
 		{
-			KeyString += FString::Printf(TEXT("_STRTSP"));
+			KeyGen.AppendSeparator();
+			KeyGen.Append(TEXT("STRTSP"));
 		}
 	}
 
@@ -2107,7 +2272,8 @@ void ShaderMapAppendKeyString(EShaderPlatform Platform, FString& KeyString)
 		static const auto CVar = IConsoleManager::Get().FindTConsoleVariableDataInt(TEXT("r.Material.RoughDiffuse"));
 		if (CVar && CVar->GetValueOnAnyThread() > 0)
 		{
-			KeyString += FString::Printf(TEXT("_MATRDIFF"));
+			KeyGen.AppendSeparator();
+			KeyGen.Append(TEXT("MATRDIFF"));
 		}
 	}
 
@@ -2115,19 +2281,23 @@ void ShaderMapAppendKeyString(EShaderPlatform Platform, FString& KeyString)
 		int32 LightFunctionAtlasFormat = GetLightFunctionAtlasFormat();
 		if (LightFunctionAtlasFormat > 0)
 		{
-			KeyString += FString::Printf(TEXT("_LFAC%u"), LightFunctionAtlasFormat);
+			KeyGen.AppendSeparator();
+			KeyGen.Append(TEXT("LFAC"));
+			KeyGen.Append(static_cast<uint32>(LightFunctionAtlasFormat));
 		}
 
 		bool bSingleLayerWaterUsesLightFunctionAtlas = GetSingleLayerWaterUsesLightFunctionAtlas();
 		if (bSingleLayerWaterUsesLightFunctionAtlas)
 		{
-			KeyString += FString::Printf(TEXT("_SLWLFA"));
+			KeyGen.AppendSeparator();
+			KeyGen.Append(TEXT("SLWLFA"));
 		}
 
 		bool bTranslucentUsesLightFunctionAtlas = GetTranslucentUsesLightFunctionAtlas();
 		if (bTranslucentUsesLightFunctionAtlas)
 		{
-			KeyString += FString::Printf(TEXT("_FWDLFA"));
+			KeyGen.AppendSeparator();
+			KeyGen.Append(TEXT("FWDLFA"));
 		}
 
 	}
@@ -2136,14 +2306,16 @@ void ShaderMapAppendKeyString(EShaderPlatform Platform, FString& KeyString)
 		static const auto CVar = IConsoleManager::Get().FindTConsoleVariableDataInt(TEXT("r.Material.EnergyConservation"));
 		if (CVar && CVar->GetValueOnAnyThread() > 0)
 		{
-			KeyString += FString::Printf(TEXT("_MATENERGY"));
+			KeyGen.AppendSeparator();
+			KeyGen.Append(TEXT("MATENERGY"));
 		}
 	}
 
 	{
 		if (MaskedInEarlyPass(Platform))
 		{
-			KeyString += TEXT("_EZPMM");
+			KeyGen.AppendSeparator();
+			KeyGen.Append(TEXT("EZPMM"));
 		}
 	}
 	
@@ -2151,35 +2323,41 @@ void ShaderMapAppendKeyString(EShaderPlatform Platform, FString& KeyString)
 		static const auto CVar = IConsoleManager::Get().FindTConsoleVariableDataInt(TEXT("r.GPUSkin.Limit2BoneInfluences"));
 		if (CVar && CVar->GetValueOnAnyThread() != 0)
 		{
-			KeyString += TEXT("_2bi");
+			KeyGen.AppendSeparator();
+			KeyGen.Append(TEXT("2bi"));
 		}
 	}
 	{
 		if(UseGPUScene(Platform, GetMaxSupportedFeatureLevel(Platform)))
 		{
-			KeyString += TEXT("_gs1");
+			KeyGen.AppendSeparator();
+			KeyGen.Append(TEXT("gs1"));
 		}
 		else
 		{
-			KeyString += TEXT("_gs0");
+			KeyGen.AppendSeparator();
+			KeyGen.Append(TEXT("gs0"));
 		}
 	}
 
 	if (FDataDrivenShaderPlatformInfo::GetSupportSceneDataCompressedTransforms(Platform))
 	{
-		KeyString += TEXT("_sdct");
+		KeyGen.AppendSeparator();
+		KeyGen.Append(TEXT("sdct"));
 	}
 
 	if (FDataDrivenShaderPlatformInfo::GetSupportsVariableRateShading(Platform) && GRHIAttachmentVariableRateShadingEnabled)
 	{
-		KeyString += TEXT("_VRS");
+		KeyGen.AppendSeparator();
+		KeyGen.Append(TEXT("VRS"));
 	}
 
 	{
 		static const auto CVar = IConsoleManager::Get().FindTConsoleVariableDataInt(TEXT("r.GBufferDiffuseSampleOcclusion"));
 		if (CVar && CVar->GetValueOnAnyThread() != 0)
 		{
-			KeyString += TEXT("_GDSO");
+			KeyGen.AppendSeparator();
+			KeyGen.Append(TEXT("GDSO"));
 		}
 	}
 
@@ -2207,7 +2385,16 @@ void ShaderMapAppendKeyString(EShaderPlatform Platform, FString& KeyString)
 
 		const bool VTSupported = TargetPlatform != nullptr && TargetPlatform->SupportsFeature(ETargetPlatformFeatures::VirtualTextureStreaming);
 
- 		KeyString += FString::Printf(TEXT("_VT-%d-%d-%d-%d"), VTLightmaps, VTTextures, VTSupported, VTFiltering);
+		KeyGen.AppendSeparator();
+		KeyGen.Append(TEXT("VT"));
+		KeyGen.AppendDebugText(TEXT("-"));
+		KeyGen.Append(VTLightmaps);
+		KeyGen.AppendDebugText(TEXT("-"));
+		KeyGen.Append(VTTextures);
+		KeyGen.AppendDebugText(TEXT("-"));
+		KeyGen.Append(VTSupported);
+		KeyGen.AppendDebugText(TEXT("-"));
+		KeyGen.Append(VTFiltering);
 	}
 
 	{
@@ -2220,13 +2407,17 @@ void ShaderMapAppendKeyString(EShaderPlatform Platform, FString& KeyString)
 			WCSHash ^= GetTypeHash(WCS.GetGreenChromaticity());
 			WCSHash ^= GetTypeHash(WCS.GetBlueChromaticity());
 			WCSHash ^= GetTypeHash(WCS.GetWhiteChromaticity());
-			KeyString += FString::Printf(TEXT("_WCS-%u"), WCSHash);
+			KeyGen.AppendSeparator();
+			KeyGen.Append(TEXT("WCS"));
+			KeyGen.AppendDebugText(TEXT("-"));
+			KeyGen.Append(WCSHash);
 		}
 
 		static const auto CVar = IConsoleManager::Get().FindConsoleVariable(TEXT("r.LegacyLuminanceFactors"));
 		if (CVar && CVar->GetInt() != 0)
 		{
-			KeyString += TEXT("_LLF");
+			KeyGen.AppendSeparator();
+			KeyGen.Append(TEXT("LLF"));
 		}
 	}
 
@@ -2234,7 +2425,8 @@ void ShaderMapAppendKeyString(EShaderPlatform Platform, FString& KeyString)
 		static const auto CVar = IConsoleManager::Get().FindTConsoleVariableDataInt(TEXT("r.Shaders.RemoveDeadCode"));
 		if (CVar && CVar->GetValueOnAnyThread() != 0)
 		{
-			KeyString += TEXT("_MIN");
+			KeyGen.AppendSeparator();
+			KeyGen.Append(TEXT("MIN"));
 		}
 	}
 
@@ -2242,43 +2434,51 @@ void ShaderMapAppendKeyString(EShaderPlatform Platform, FString& KeyString)
 		static const auto CVar = IConsoleManager::Get().FindTConsoleVariableDataBool(TEXT("r.ShaderCompiler.PreprocessedJobCache"));
 		if (CVar && CVar->GetValueOnAnyThread())
 		{
-			KeyString += TEXT("_PJC");
+			KeyGen.AppendSeparator();
+			KeyGen.Append(TEXT("PJC"));
 		}
 	}
 
 	if (RHISupportsShaderRootConstants(Platform))
 	{
-		KeyString += TEXT("_SHRC");
+		KeyGen.AppendSeparator();
+		KeyGen.Append(TEXT("SHRC"));
 	}
 
 	if (RHISupportsShaderBundleDispatch(Platform))
 	{
-		KeyString += TEXT("_SHBD");
+		KeyGen.AppendSeparator();
+		KeyGen.Append(TEXT("SHBD"));
 	}
 
 	if (RHISupportsRenderTargetWriteMask(Platform))
 	{
-		KeyString += TEXT("_RTWM");
+		KeyGen.AppendSeparator();
+		KeyGen.Append(TEXT("RTWM"));
 	}
 
 	if (FDataDrivenShaderPlatformInfo::GetSupportsPerPixelDBufferMask(Platform))
 	{
-		KeyString += TEXT("_PPDBM");
+		KeyGen.AppendSeparator();
+		KeyGen.Append(TEXT("PPDBM"));
 	}
 
 	if (FDataDrivenShaderPlatformInfo::GetSupportsDistanceFields(Platform))
 	{
-		KeyString += TEXT("_DF");
+		KeyGen.AppendSeparator();
+		KeyGen.Append(TEXT("DF"));
 	}
 
 	if (RHISupportsMeshShadersTier0(Platform))
 	{
-		KeyString += TEXT("_MS_T0");
+		KeyGen.AppendSeparator();
+		KeyGen.Append(TEXT("MS_T0"));
 	}
 
 	if (RHISupportsMeshShadersTier1(Platform))
 	{
-		KeyString += TEXT("_MS_T1");
+		KeyGen.AppendSeparator();
+		KeyGen.Append(TEXT("MS_T1"));
 	}
 
 	if (RHIGetBindlessSupport(Platform) != ERHIBindlessSupport::Unsupported)
@@ -2288,12 +2488,16 @@ void ShaderMapAppendKeyString(EShaderPlatform Platform, FString& KeyString)
 
 		if (ResourcesConfig != ERHIBindlessConfiguration::Disabled)
 		{
-			KeyString += ResourcesConfig == ERHIBindlessConfiguration::RayTracingShaders ? TEXT("_BNDLSRTRES") : TEXT("_BNDLSRES");
+			KeyGen.AppendSeparator();
+			KeyGen.Append(ResourcesConfig == ERHIBindlessConfiguration::RayTracingShaders ?
+				TEXT("BNDLSRTRES") : TEXT("BNDLSRES"));
 		}
 
 		if (SamplersConfig != ERHIBindlessConfiguration::Disabled)
 		{
-			KeyString += SamplersConfig == ERHIBindlessConfiguration::RayTracingShaders ? TEXT("_BNDLSRTSAM") : TEXT("_BNDLSSAM");
+			KeyGen.AppendSeparator();
+			KeyGen.Append(SamplersConfig == ERHIBindlessConfiguration::RayTracingShaders ?
+				TEXT("BNDLSRTSAM") : TEXT("BNDLSSAM"));
 		}
 	}
 
@@ -2303,10 +2507,14 @@ void ShaderMapAppendKeyString(EShaderPlatform Platform, FString& KeyString)
 		static const auto CVarCompileAHS = IConsoleManager::Get().FindConsoleVariable(TEXT("r.RayTracing.CompileMaterialAHS"));
 		static const auto CVarTextureLod = IConsoleManager::Get().FindConsoleVariable(TEXT("r.RayTracing.UseTextureLod"));
 
-		KeyString += FString::Printf(TEXT("_RAY-CHS%dAHS%dLOD%d"),
-			CVarCompileCHS && CVarCompileCHS->GetBool() ? 1 : 0,
-			CVarCompileAHS && CVarCompileAHS->GetBool() ? 1 : 0,
-			CVarTextureLod && CVarTextureLod->GetBool() ? 1 : 0);
+		KeyGen.AppendSeparator();
+		KeyGen.Append(TEXT("RAY"));
+		KeyGen.AppendDebugText(TEXT("-CHS"));
+		KeyGen.AppendBoolInt(CVarCompileCHS && CVarCompileCHS->GetBool());
+		KeyGen.AppendDebugText(TEXT("AHS"));
+		KeyGen.AppendBoolInt(CVarCompileAHS && CVarCompileAHS->GetBool());
+		KeyGen.AppendDebugText(TEXT("LOD"));
+		KeyGen.AppendBoolInt(CVarTextureLod && CVarTextureLod->GetBool());
 	}
 
 	if (DoesPlatformSupportHeterogeneousVolumes(Platform))
@@ -2314,67 +2522,88 @@ void ShaderMapAppendKeyString(EShaderPlatform Platform, FString& KeyString)
 		static const auto ShadowCVar = IConsoleManager::Get().FindTConsoleVariableDataInt(TEXT("r.HeterogeneousVolumes.Shadows"));
 		if (ShadowCVar && ShadowCVar->GetValueOnAnyThread() != 0)
 		{
-			KeyString += TEXT("_HVSHADOW");
+			KeyGen.AppendSeparator();
+			KeyGen.Append(TEXT("HVSHADOW"));
 		}
 
 		static const auto CompTranslucencyCVar = IConsoleManager::Get().FindTConsoleVariableDataInt(TEXT("r.Translucency.HeterogeneousVolumes"));
 		if (CompTranslucencyCVar && CompTranslucencyCVar->GetValueOnAnyThread() != 0)
 		{
-			KeyString += TEXT("_HVCOMPTRANSL");
+			KeyGen.AppendSeparator();
+			KeyGen.Append(TEXT("HVCOMPTRANSL"));
 		}
 	}
 
 	if (ForceSimpleSkyDiffuse(Platform))
 	{
-		KeyString += TEXT("_SSD");
+		KeyGen.AppendSeparator();
+		KeyGen.Append(TEXT("SSD"));
 	}
 
 	if (VelocityEncodeDepth(Platform))
 	{
-		KeyString += TEXT("_VED");
+		KeyGen.AppendSeparator();
+		KeyGen.Append(TEXT("VED"));
 	}
 
 	{
 		const bool bSupportsAnisotropicMaterials = FDataDrivenShaderPlatformInfo::GetSupportsAnisotropicMaterials(Platform);
-		KeyString += FString::Printf(TEXT("_Aniso-%d"), bSupportsAnisotropicMaterials ? 1 : 0);
+		KeyGen.AppendSeparator();
+		KeyGen.Append(TEXT("Aniso"));
+		KeyGen.AppendDebugText(TEXT("-"));
+		KeyGen.AppendBoolInt(bSupportsAnisotropicMaterials);
 	}
 
 	{
 		// add shader compression format
-		KeyString += TEXT("_Compr");
+		KeyGen.AppendSeparator();
+		KeyGen.Append(TEXT("Compr"));
 		FName CompressionFormat = GetShaderCompressionFormat();
-		KeyString += CompressionFormat.ToString();
+		KeyGen.Append(CompressionFormat);
 		if (CompressionFormat == NAME_Oodle)
 		{
 			FOodleDataCompression::ECompressor OodleCompressor;
 			FOodleDataCompression::ECompressionLevel OodleLevel;
 			GetShaderCompressionOodleSettings(OodleCompressor, OodleLevel);
-			KeyString += FString::Printf(TEXT("_Compr%d_Lev%d"), static_cast<int32>(OodleCompressor), static_cast<int32>(OodleLevel));
+			KeyGen.AppendSeparator();
+			KeyGen.Append(TEXT("Compr"));
+			KeyGen.Append(static_cast<int32>(OodleCompressor));
+			KeyGen.AppendSeparator();
+			KeyGen.AppendDebugText(TEXT("Lev"));
+			KeyGen.Append(static_cast<int32>(OodleLevel));
 		}
 	}
 
 	{
 		// add whether or not non-pipelined shader types are included
-		KeyString += FString::Printf(TEXT("_ExclNonPipSh-%d"), ExcludeNonPipelinedShaderTypes(Platform));
+		KeyGen.AppendSeparator();
+		KeyGen.Append(TEXT("ExclNonPipSh-"));
+		KeyGen.AppendBoolInt(ExcludeNonPipelinedShaderTypes(Platform));
 	}
 
-	KeyString += FString::Printf(TEXT("_LWC-%d"), FMath::FloorToInt(FLargeWorldRenderScalar::GetTileSize()));
+	KeyGen.AppendSeparator();
+	KeyGen.Append(TEXT("LWC-"));
+	KeyGen.Append(FMath::FloorToInt(FLargeWorldRenderScalar::GetTileSize()));
 
 	uint64 ShaderPlatformPropertiesHash = FDataDrivenShaderPlatformInfo::GetShaderPlatformPropertiesHash(Platform);
-	KeyString += FString::Printf(TEXT("_%u"), ShaderPlatformPropertiesHash);
+	KeyGen.AppendSeparator();
+	KeyGen.Append(ShaderPlatformPropertiesHash);
 
 	if (IsSingleLayerWaterDepthPrepassEnabled(Platform, GetMaxSupportedFeatureLevel(Platform)))
 	{
-		KeyString += TEXT("_SLWDP");
+		KeyGen.AppendSeparator();
+		KeyGen.Append(TEXT("SLWDP"));
 	}
 
 	if (IsGPUSkinPassThroughSupported(Platform))
 	{
-		KeyString += TEXT("_SKPassThrough1");
+		KeyGen.AppendSeparator();
+		KeyGen.Append(TEXT("SKPassThrough1"));
 	}
 	else
 	{
-		KeyString += TEXT("_SKPassThrough0");
+		KeyGen.AppendSeparator();
+		KeyGen.Append(TEXT("SKPassThrough0"));
 	}
 
 	if (UseNanite(Platform))
@@ -2383,12 +2612,13 @@ void ShaderMapAppendKeyString(EShaderPlatform Platform, FString& KeyString)
 		static const auto CVarAllowCSMat = IConsoleManager::Get().FindConsoleVariable(TEXT("r.Nanite.AllowComputeMaterials"));
 		static const auto CVarAllowPSMat = IConsoleManager::Get().FindConsoleVariable(TEXT("r.Nanite.AllowLegacyMaterials"));
 
-		KeyString.Appendf(
-			TEXT("_Nanite-Tess%dCSMat%dPSMat%d"),
-			CVarAllowTess ? CVarAllowTess->GetInt() : 0,
-			CVarAllowCSMat ? CVarAllowCSMat->GetInt() : 0,
-			CVarAllowPSMat ? CVarAllowPSMat->GetInt() : 0
-		);
+		KeyGen.AppendSeparator();
+		KeyGen.Append(TEXT("Nanite-Tess"));
+		KeyGen.Append(CVarAllowTess ? CVarAllowTess->GetInt() : 0);
+		KeyGen.AppendDebugText(TEXT("CSMat"));
+		KeyGen.Append(CVarAllowCSMat ? CVarAllowCSMat->GetInt() : 0);
+		KeyGen.AppendDebugText(TEXT("PSMat"));
+		KeyGen.Append(CVarAllowPSMat ? CVarAllowPSMat->GetInt() : 0);
 	}
 }
 
diff --git a/Engine/Source/Runtime/RenderCore/Private/ShaderCore.cpp b/Engine/Source/Runtime/RenderCore/Private/ShaderCore.cpp
--- a/Engine/Source/Runtime/RenderCore/Private/ShaderCore.cpp
+++ b/Engine/Source/Runtime/RenderCore/Private/ShaderCore.cpp
@@ -30,6 +30,7 @@
 #include "RHIShaderFormatDefinitions.inl"
 #include "Serialization/MemoryHasher.h"
 #include "Serialization/MemoryWriter.h"
+#include "Serialization/ShaderKeyGenerator.h"
 #include "Shader.h"
 #include "ShaderCompilerCore.h"
 #include "ShaderCompilerDefinitions.h"
@@ -3366,6 +3367,19 @@ void AppendKeyStringShaderDependencies(
 	FPlatformTypeLayoutParameters LayoutParams,
 	FString& OutKeyString,
 	bool bIncludeSourceHashes)
+{
+	FShaderKeyGenerator KeyGen(OutKeyString);
+	AppendShaderDependencies(KeyGen, ShaderTypeDependencies, ShaderPipelineTypeDependencies,
+		VertexFactoryTypeDependencies, LayoutParams, bIncludeSourceHashes);
+}
+
+void AppendShaderDependencies(
+	FShaderKeyGenerator& KeyGen,
+	TConstArrayView<FShaderTypeDependency> ShaderTypeDependencies,
+	TConstArrayView<FShaderPipelineTypeDependency> ShaderPipelineTypeDependencies,
+	TConstArrayView<FVertexFactoryTypeDependency> VertexFactoryTypeDependencies,
+	FPlatformTypeLayoutParameters LayoutParams,
+	bool bIncludeSourceHashes)
 {
 	FMemMark MemMark(FMemStack::Get());
 	using FMemStackSetAllocator = TSetAllocator<TSparseArrayAllocator<TMemStackAllocator<>, TMemStackAllocator<>>, TMemStackAllocator<>>;
@@ -3375,30 +3389,31 @@ void AppendKeyStringShaderDependencies(
 	for (const FShaderTypeDependency& ShaderTypeDependency : ShaderTypeDependencies)
 	{
 		const FShaderType* ShaderType = FindShaderTypeByName(ShaderTypeDependency.ShaderTypeName);
-		checkf(ShaderType != nullptr, TEXT("Failed to find FShaderType for dependency %hs (total in the NameToTypeMap: %d)"), ShaderTypeDependency.ShaderTypeName.GetDebugString().String.Get(), FShaderType::GetNameToTypeMap().Num());
-
-		OutKeyString.AppendChar('_');
-		OutKeyString.Append(ShaderType->GetName());
-		OutKeyString.AppendInt(ShaderTypeDependency.PermutationId);
-		OutKeyString.AppendChar('_');
+		checkf(ShaderType != nullptr,
+			TEXT("Failed to find FShaderType for dependency %hs (total in the NameToTypeMap: %d)"),
+			ShaderTypeDependency.ShaderTypeName.GetDebugString().String.Get(), FShaderType::GetNameToTypeMap().Num());
+
+		KeyGen.AppendSeparator();
+		KeyGen.Append(ShaderType->GetName());
+		KeyGen.Append(ShaderTypeDependency.PermutationId);
+		KeyGen.AppendSeparator();
 		ERayTracingPayloadType RayTracingPayloadType = ShaderType->GetRayTracingPayloadType(ShaderTypeDependency.PermutationId);
-		OutKeyString.AppendInt(static_cast<uint32>(RayTracingPayloadType));
-		OutKeyString.AppendChar('_');
-		OutKeyString.AppendInt(GetRayTracingPayloadTypeMaxSize(RayTracingPayloadType));
+		KeyGen.Append(static_cast<uint32>(RayTracingPayloadType));
+		KeyGen.AppendSeparator();
+		KeyGen.Append(GetRayTracingPayloadTypeMaxSize(RayTracingPayloadType));
 
 		if (bIncludeSourceHashes)
 		{
 			// Add the type's source hash so that we can invalidate cached shaders when .usf changes are made
-			ShaderTypeDependency.SourceHash.AppendString(OutKeyString);
+			KeyGen.Append(ShaderTypeDependency.SourceHash);
 		}
 
 		if (const FShaderParametersMetadata* ParameterStructMetadata = ShaderType->GetRootParametersMetadata())
 		{
-			ParameterStructMetadata->AppendKeyString(OutKeyString);
+			ParameterStructMetadata->Append(KeyGen);
 		}
 
-		const FSHAHash LayoutHash = GetShaderTypeLayoutHash(ShaderType->GetLayout(), LayoutParams);
-		LayoutHash.AppendString(OutKeyString);
+		KeyGen.Append(GetShaderTypeLayoutHash(ShaderType->GetLayout(), LayoutParams));
 
 		for (const FShaderParametersMetadata* UniformBuffer : ShaderType->GetReferencedUniformBuffers())
 		{
@@ -3410,21 +3425,23 @@ void AppendKeyStringShaderDependencies(
 	for (const FShaderPipelineTypeDependency& Dependency : ShaderPipelineTypeDependencies)
 	{
 		const FShaderPipelineType* ShaderPipelineType = FShaderPipelineType::GetShaderPipelineTypeByName(Dependency.ShaderPipelineTypeName);
-		checkf(ShaderPipelineType != nullptr, TEXT("Failed to find FShaderPipelineType for dependency %hs (total in the NameToTypeMap: %d)"), Dependency.ShaderPipelineTypeName.GetDebugString().String.Get(), FShaderType::GetNameToTypeMap().Num());
+		checkf(ShaderPipelineType != nullptr,
+			TEXT("Failed to find FShaderPipelineType for dependency %hs (total in the NameToTypeMap: %d)"),
+			Dependency.ShaderPipelineTypeName.GetDebugString().String.Get(), FShaderType::GetNameToTypeMap().Num());
 
-		OutKeyString.AppendChar('_');
-		OutKeyString.Append(ShaderPipelineType->GetName());
+		KeyGen.AppendSeparator();
+		KeyGen.Append(ShaderPipelineType->GetName());
 
 		if (bIncludeSourceHashes)
 		{
-			Dependency.StagesSourceHash.AppendString(OutKeyString);
+			KeyGen.Append(Dependency.StagesSourceHash);
 		}
 
 		for (const FShaderType* ShaderType : ShaderPipelineType->GetStages())
 		{
 			if (const FShaderParametersMetadata* ParameterStructMetadata = ShaderType->GetRootParametersMetadata())
 			{
-				ParameterStructMetadata->AppendKeyString(OutKeyString);
+				ParameterStructMetadata->Append(KeyGen);
 			}
 
 			for (const FShaderParametersMetadata* UniformBuffer : ShaderType->GetReferencedUniformBuffers())
@@ -3436,15 +3453,15 @@ void AppendKeyStringShaderDependencies(
 
 	for (const FVertexFactoryTypeDependency& VFDependency : VertexFactoryTypeDependencies)
 	{
-		OutKeyString.AppendChar('_');
+		KeyGen.AppendSeparator();
 
 		const FVertexFactoryType* VertexFactoryType = FVertexFactoryType::GetVFByName(VFDependency.VertexFactoryTypeName);
 
-		OutKeyString.Append(VertexFactoryType->GetName());
+		KeyGen.Append(VertexFactoryType->GetName());
 
 		if (bIncludeSourceHashes)
 		{
-			VFDependency.VFSourceHash.AppendString(OutKeyString);
+			KeyGen.Append(VFDependency.VFSourceHash);
 		}
 
 		for (int32 Frequency = 0; Frequency < SF_NumFrequencies; Frequency++)
@@ -3453,7 +3470,7 @@ void AppendKeyStringShaderDependencies(
 			if (ParameterLayout)
 			{
 				const FSHAHash LayoutHash = GetShaderTypeLayoutHash(*ParameterLayout, LayoutParams);
-				LayoutHash.AppendString(OutKeyString);
+				KeyGen.Append(LayoutHash);
 			}
 		}
 
@@ -3476,7 +3493,7 @@ void AppendKeyStringShaderDependencies(
 	// Save uniform buffer member info so we can detect when layout has changed
 	for (const FShaderParametersMetadata* UniformBufferMetadata : ReferencedUniformBuffers)
 	{
-		UniformBufferMetadata->AppendKeyString(OutKeyString);
+		UniformBufferMetadata->Append(KeyGen);
 	}
 }
 
diff --git a/Engine/Source/Runtime/RenderCore/Private/ShaderParameterMetadata.cpp b/Engine/Source/Runtime/RenderCore/Private/ShaderParameterMetadata.cpp
--- a/Engine/Source/Runtime/RenderCore/Private/ShaderParameterMetadata.cpp
+++ b/Engine/Source/Runtime/RenderCore/Private/ShaderParameterMetadata.cpp
@@ -12,6 +12,7 @@
 #include "RenderCore.h"
 #include "RHIUniformBufferLayoutInitializer.h"
 #include "Serialization/MemoryHasher.h"
+#include "Serialization/ShaderKeyGenerator.h"
 #include "ShaderCore.h"
 #include "ShaderCompilerCore.h"
 #include "ShaderParameters.h"
@@ -460,6 +461,18 @@ void FShaderParametersMetadata::InitializeAllUniformBufferStructs()
 
 #if WITH_EDITOR
 
+void FShaderParametersMetadata::AppendKeyString(FString& OutKeyString) const
+{
+	FShaderKeyGenerator KeyGen(OutKeyString);
+	Append(KeyGen);
+}
+
+void FShaderParametersMetadata::Append(FShaderKeyGenerator& KeyGen) const
+{
+	KeyGen.AppendDebugText(TEXT("SPM_"));
+	KeyGen.Append(LayoutSignature);
+}
+
 void FShaderParametersMetadata::FMember::HashLayout(FMemoryHasherBlake3& Hasher)
 {
 	Hasher << Offset;
diff --git a/Engine/Source/Runtime/RenderCore/Public/Shader.h b/Engine/Source/Runtime/RenderCore/Public/Shader.h
--- a/Engine/Source/Runtime/RenderCore/Public/Shader.h
+++ b/Engine/Source/Runtime/RenderCore/Public/Shader.h
@@ -72,6 +72,7 @@ class FNiagaraShaderType;
 class FOpenColorIOShaderType;
 class FRHIComputeCommandList;
 class FShader;
+class FShaderKeyGenerator;
 class FShaderMapBase;
 class FShaderMapPointerTable;
 class FShaderParametersMetadata;
@@ -2554,3 +2555,4 @@ extern RENDERCORE_API bool IsDxcEnabledForPlatform(EShaderPlatform Platform, boo
 
 /** Appends to KeyString for all shaders. */
 extern RENDERCORE_API void ShaderMapAppendKeyString(EShaderPlatform Platform, FString& KeyString);
+extern RENDERCORE_API void ShaderMapAppendKey(EShaderPlatform Platform, FShaderKeyGenerator& KeyGen);
diff --git a/Engine/Source/Runtime/RenderCore/Public/ShaderCore.h b/Engine/Source/Runtime/RenderCore/Public/ShaderCore.h
--- a/Engine/Source/Runtime/RenderCore/Public/ShaderCore.h
+++ b/Engine/Source/Runtime/RenderCore/Public/ShaderCore.h
@@ -50,6 +50,7 @@ class FMemoryUnfreezeContent;
 class FPointerTableBase;
 class FShaderCompilerDefinitions;
 class FShaderCompileUtilities;
+class FShaderKeyGenerator;
 class FShaderPreprocessorUtilities;
 class FSHA1;
 class ITargetPlatform;
@@ -1259,6 +1260,13 @@ extern RENDERCORE_API void AppendKeyStringShaderDependencies(
 	FPlatformTypeLayoutParameters LayoutParams,
 	FString& OutKeyString,
 	bool bIncludeSourceHashes = true);
+extern RENDERCORE_API void AppendShaderDependencies(
+	FShaderKeyGenerator& KeyGen,
+	TConstArrayView<FShaderTypeDependency> ShaderTypeDependencies,
+	TConstArrayView<FShaderPipelineTypeDependency> ShaderPipelineTypeDependencies,
+	TConstArrayView<FVertexFactoryTypeDependency> VertexFactoryTypeDependencies,
+	FPlatformTypeLayoutParameters LayoutParams,
+	bool bIncludeSourceHashes = true);
 #endif // WITH_EDITOR
 
 /** Create a block of source code to be injected in the preprocessed shader code. The Block will be put into a #line directive
diff --git a/Engine/Source/Runtime/RenderCore/Public/ShaderParameterMetadata.h b/Engine/Source/Runtime/RenderCore/Public/ShaderParameterMetadata.h
--- a/Engine/Source/Runtime/RenderCore/Public/ShaderParameterMetadata.h
+++ b/Engine/Source/Runtime/RenderCore/Public/ShaderParameterMetadata.h
@@ -29,6 +29,7 @@
 using FThreadSafeSharedStringPtr = TSharedPtr<FString, ESPMode::ThreadSafe>;
 using FThreadSafeSharedAnsiStringPtr = TSharedPtr<TArray<ANSICHAR>, ESPMode::ThreadSafe>;
 using FThreadSafeNameBufferPtr = TSharedPtr<TArray<TCHAR>, ESPMode::ThreadSafe>;
+class FShaderKeyGenerator;
 struct FShaderResourceTableMap;
 
 namespace EShaderPrecisionModifier
@@ -395,13 +396,8 @@ public:
 	}
 
 #if WITH_EDITOR
-	inline void AppendKeyString(FString& OutKeyString) const
-	{
-		TStringBuilder<sizeof(TCHAR) * (sizeof(FBlake3Hash::ByteArray) * 2 + 4)> StrBuilder;
-		StrBuilder << "SPM_";
-		StrBuilder << LayoutSignature;
-		OutKeyString.Append(StrBuilder.ToView());
-	}
+	RENDERCORE_API void AppendKeyString(FString& OutKeyString) const;
+	RENDERCORE_API void Append(FShaderKeyGenerator& KeyGen) const;
 #endif
 
 	inline const FBlake3Hash& GetLayoutSignature() const
