Added CVar to control whether render command fence bundling is enabled.

diff --git a/Engine/Source/Runtime/RenderCore/Private/RenderingThread.cpp b/Engine/Source/Runtime/RenderCore/Private/RenderingThread.cpp
--- a/Engine/Source/Runtime/RenderCore/Private/RenderingThread.cpp
+++ b/Engine/Source/Runtime/RenderCore/Private/RenderingThread.cpp
@@ -92,6 +92,15 @@ static FAutoConsoleVariable CVarRenderThreadPollPeriodMs(
 	FConsoleVariableDelegate::CreateStatic(&OnRenderThreadPollPeriodMsChanged)
 );
 
+bool GRenderCommandFenceBundling = true;
+FAutoConsoleVariableRef CVarRenderCommandFenceBundling(
+	TEXT("r.RenderCommandFenceBundling"),
+	GRenderCommandFenceBundling,
+	TEXT("Controls whether render command fences are allowed to be batched.\n")
+	TEXT(" 0: disabled;\n")
+	TEXT(" 1: enabled (default);\n"),
+	ECVF_Default);
+
 /**
  * Tick all rendering thread tickable objects
  */
@@ -855,7 +864,7 @@ static struct FRenderCommandFenceBundlerState
 
 void StartRenderCommandFenceBundler()
 {
-	if (!GIsThreadedRendering)
+	if (!GIsThreadedRendering || !GRenderCommandFenceBundling)
 	{
 		return;
 	}
