Added FRHIStreamSourceSlot to support updating stream source vertex buffers for mesh draw commands. Also removed unnecessary null vertex streams from being added to mesh draw commands when overridden.

diff --git a/Engine/Source/Runtime/RHI/Private/RHIStrings.cpp b/Engine/Source/Runtime/RHI/Private/RHIStrings.cpp
--- a/Engine/Source/Runtime/RHI/Private/RHIStrings.cpp
+++ b/Engine/Source/Runtime/RHI/Private/RHIStrings.cpp
@@ -436,6 +436,7 @@ static const FRHIResourceTypeName GRHIResourceTypeNames[] =
 	RHI_RESOURCE_TYPE_DEF(CustomPresent),
 	RHI_RESOURCE_TYPE_DEF(ShaderLibrary),
 	RHI_RESOURCE_TYPE_DEF(PipelineBinaryLibrary),
+	RHI_RESOURCE_TYPE_DEF(StreamSourceSlot)
 };
 
 ERHIResourceType RHIResourceTypeFromString(const FString& InString)
diff --git a/Engine/Source/Runtime/RHI/Public/RHICommandList.h b/Engine/Source/Runtime/RHI/Public/RHICommandList.h
--- a/Engine/Source/Runtime/RHI/Public/RHICommandList.h
+++ b/Engine/Source/Runtime/RHI/Public/RHICommandList.h
@@ -733,6 +733,15 @@ public:
 		GDynamicRHI->RHIUpdateUniformBuffer(*this, UniformBufferRHI, Contents);
 	}
 
+	FORCEINLINE void UpdateStreamSourceSlot(FRHIStreamSourceSlot* StreamSourceSlotRHI, FRHIBuffer* BufferRHI)
+	{
+		check(StreamSourceSlotRHI);
+		EnqueueLambda([this, StreamSourceSlotRHI, BufferRHI] (FRHICommandListBase&)
+		{
+			StreamSourceSlotRHI->Buffer = BufferRHI;
+		});
+	}
+
 	FORCEINLINE void UpdateTexture2D(FRHITexture* Texture, uint32 MipIndex, const struct FUpdateTextureRegion2D& UpdateRegion, uint32 SourcePitch, const uint8* SourceData)
 	{
 		checkf(UpdateRegion.DestX + UpdateRegion.Width <= Texture->GetSizeX(), TEXT("UpdateTexture2D out of bounds on X. Texture: %s, %i, %i, %i"), *Texture->GetName().ToString(), UpdateRegion.DestX, UpdateRegion.Width, Texture->GetSizeX());
@@ -3267,6 +3276,15 @@ public:
 		ALLOC_COMMAND(FRHICommandSetStreamSource)(StreamIndex, VertexBuffer, Offset);
 	}
 
+	FORCEINLINE_DEBUGGABLE void SetStreamSourceSlot(uint32 StreamIndex, FRHIStreamSourceSlot* StreamSourceSlot, uint32 Offset)
+	{
+		EnqueueLambda([StreamIndex, StreamSourceSlot, Offset] (FRHICommandListBase& RHICmdList)
+		{
+			FRHICommandSetStreamSource Command(StreamIndex, StreamSourceSlot ? StreamSourceSlot->Buffer : nullptr, Offset);
+			Command.Execute(RHICmdList);
+		});
+	}
+
 	FORCEINLINE_DEBUGGABLE void SetStencilRef(uint32 StencilRef)
 	{
 		//check(IsOutsideRenderPass());
diff --git a/Engine/Source/Runtime/RHI/Public/RHIDefinitions.h b/Engine/Source/Runtime/RHI/Public/RHIDefinitions.h
--- a/Engine/Source/Runtime/RHI/Public/RHIDefinitions.h
+++ b/Engine/Source/Runtime/RHI/Public/RHIDefinitions.h
@@ -953,6 +953,7 @@ enum ERHIResourceType : uint8
 	RRT_ShaderBundle,
 	RRT_WorkGraphShader,
 	RRT_WorkGraphPipelineState,
+	RRT_StreamSourceSlot,
 
 	RRT_Num
 };
diff --git a/Engine/Source/Runtime/RHI/Public/RHIFwd.h b/Engine/Source/Runtime/RHI/Public/RHIFwd.h
--- a/Engine/Source/Runtime/RHI/Public/RHIFwd.h
+++ b/Engine/Source/Runtime/RHI/Public/RHIFwd.h
@@ -88,6 +88,7 @@ class FRHIViewableResource;
 class FRHIViewport;
 class FRHIWorkGraphPipelineState;
 class FRHIWorkGraphShader;
+class FRHIStreamSourceSlot;
 
 struct FRHIUniformBufferLayout;
 
@@ -129,6 +130,7 @@ using FVertexDeclarationRHIRef         = TRefCountPtr<FRHIVertexDeclaration>;
 using FVertexShaderRHIRef              = TRefCountPtr<FRHIVertexShader>;
 using FViewportRHIRef                  = TRefCountPtr<FRHIViewport>;
 using FWorkGraphPipelineStateRHIRef    = TRefCountPtr<FRHIWorkGraphPipelineState>;
+using FStreamSourceSlotRHIRef          = TRefCountPtr<FRHIStreamSourceSlot>;
 using FWorkGraphShaderRHIRef           = TRefCountPtr<FRHIWorkGraphShader>;
 
 // Deprecated typenames
diff --git a/Engine/Source/Runtime/RHI/Public/RHIResources.h b/Engine/Source/Runtime/RHI/Public/RHIResources.h
--- a/Engine/Source/Runtime/RHI/Public/RHIResources.h
+++ b/Engine/Source/Runtime/RHI/Public/RHIResources.h
@@ -1395,6 +1395,26 @@ private:
 	FRHIBufferDesc Desc;
 };
 
+/** Represents a simple indirection to a vertex buffer to bind to a source stream. The underlying vertex buffer can be updated dynamically using an RHI command list. */
+class FRHIStreamSourceSlot : public FRHIResource
+{
+	friend FRHICommandListBase;
+	friend FRHICommandList;
+public:
+	static TRefCountPtr<FRHIStreamSourceSlot> Create(FRHIBuffer* InBuffer)
+	{
+		return new FRHIStreamSourceSlot(InBuffer);
+	}
+
+private:
+	FRHIStreamSourceSlot(FRHIBuffer* InBuffer)
+		: FRHIResource(RRT_StreamSourceSlot)
+		, Buffer(InBuffer)
+	{}
+
+	TRefCountPtr<FRHIBuffer> Buffer;
+};
+
 //
 // Textures
 //
diff --git a/Engine/Source/Runtime/RenderCore/Private/VertexFactory.cpp b/Engine/Source/Runtime/RenderCore/Private/VertexFactory.cpp
--- a/Engine/Source/Runtime/RenderCore/Private/VertexFactory.cpp
+++ b/Engine/Source/Runtime/RenderCore/Private/VertexFactory.cpp
@@ -30,6 +30,18 @@ static TMap<FHashedName, FVertexFactoryType*>& GetVFTypeMap()
 	return VTTypeMap;
 }
 
+void FVertexInputStream::SetOnRHICommandList(FRHICommandList& RHICmdList) const
+{
+	if (bStreamSourceSlot)
+	{
+		RHICmdList.SetStreamSourceSlot(StreamIndex, StreamSourceSlot, Offset);
+	}
+	else
+	{
+		RHICmdList.SetStreamSource(StreamIndex, VertexBuffer, Offset);
+	}
+}
+
 /**
  * @return The global shader factory list.
  */
@@ -188,31 +200,32 @@ void FVertexFactory::GetStreams(ERHIFeatureLevel::Type InFeatureLevel, EVertexIn
 	check(IsInitialized());
 	if (VertexStreamType == EVertexInputStreamType::Default)
 	{
+		const bool bSupportsVertexFetch = SupportsManualVertexFetch(InFeatureLevel);
 
-		bool bSupportsVertexFetch = SupportsManualVertexFetch(InFeatureLevel);
-
-		for (int32 StreamIndex = 0;StreamIndex < Streams.Num();StreamIndex++)
+		for (int32 StreamIndex = 0; StreamIndex < Streams.Num(); StreamIndex++)
 		{
 			const FVertexStream& Stream = Streams[StreamIndex];
 
-			if (!(EnumHasAnyFlags(EVertexStreamUsage::ManualFetch, Stream.VertexStreamUsage) && bSupportsVertexFetch))
+			// Skip streams that are bound using manual vertex fetch if we support that.
+			if (EnumHasAnyFlags(Stream.VertexStreamUsage, EVertexStreamUsage::ManualFetch) && bSupportsVertexFetch)
 			{
-				if (!Stream.VertexBuffer)
-				{
-					OutVertexStreams.Add(FVertexInputStream(StreamIndex, 0, nullptr));
-				}
-				else
-				{
-					if (EnumHasAnyFlags(EVertexStreamUsage::Overridden, Stream.VertexStreamUsage) && !Stream.VertexBuffer->IsInitialized())
-					{
-						OutVertexStreams.Add(FVertexInputStream(StreamIndex, 0, nullptr));
-					}
-					else
-					{
-						checkf(Stream.VertexBuffer->IsInitialized(), TEXT("Vertex buffer was not initialized! Stream %u, Stride %u, Name %s"), StreamIndex, Stream.Stride, *Stream.VertexBuffer->GetFriendlyName());
-						OutVertexStreams.Add(FVertexInputStream(StreamIndex, Stream.Offset, Stream.VertexBuffer->VertexBufferRHI));
-					}
-				}
+				continue;
+			}
+
+			// Skip streams that are overridden as they will be provided manually by the vertex factory shader bindings.
+			if (EnumHasAnyFlags(Stream.VertexStreamUsage, EVertexStreamUsage::Overridden))
+			{
+				continue;
+			}
+
+			// Issue a null binding since we don't appear to have one available.
+			if (!Stream.VertexBuffer || !Stream.VertexBuffer->IsInitialized())
+			{
+				OutVertexStreams.Add(FVertexInputStream(StreamIndex, 0, nullptr));
+			}
+			else
+			{
+				OutVertexStreams.Add(FVertexInputStream(StreamIndex, Stream.Offset, Stream.VertexBuffer->VertexBufferRHI));
 			}
 		}
 	}
diff --git a/Engine/Source/Runtime/RenderCore/Public/VertexFactory.h b/Engine/Source/Runtime/RenderCore/Public/VertexFactory.h
--- a/Engine/Source/Runtime/RenderCore/Public/VertexFactory.h
+++ b/Engine/Source/Runtime/RenderCore/Public/VertexFactory.h
@@ -46,27 +46,59 @@ struct FVertexFactoryShaderPermutationParameters;
 struct FVertexInputStream
 {
 	uint32 StreamIndex : 4;
-	uint32 Offset : 28;
-	FRHIBuffer* VertexBuffer;
+	uint32 Offset : 27;
+	uint32 bStreamSourceSlot : 1;
+
+	union
+	{
+		// Contains the direct vertex buffer pointer
+		FRHIBuffer* VertexBuffer;
+
+		// Contains the stream source slot
+		FRHIStreamSourceSlot* StreamSourceSlot;
+
+		// Represents either ptr type as an opaque int for comparisons.
+		void* Pointer;
+	};
+
+	RENDERCORE_API void SetOnRHICommandList(FRHICommandList& RHICmdList) const;
 
 	FVertexInputStream() :
 		StreamIndex(0),
 		Offset(0),
+		bStreamSourceSlot(0),
 		VertexBuffer(nullptr)
 	{}
 
 	FVertexInputStream(uint32 InStreamIndex, uint32 InOffset, FRHIBuffer* InVertexBuffer)
-		: StreamIndex(InStreamIndex), Offset(InOffset), VertexBuffer(InVertexBuffer)
+		: StreamIndex(InStreamIndex), Offset(InOffset), bStreamSourceSlot(0), VertexBuffer(InVertexBuffer)
 	{
 		// Verify no overflow
 		checkSlow(InStreamIndex == StreamIndex && InOffset == Offset);
 	}
 
+	// Creates a vertex input stream from an RHI stream source slot.
+	static FVertexInputStream CreateFromStreamSourceSlot(uint32 InStreamIndex, uint32 InOffset, FRHIStreamSourceSlot* StreamSourceSlot)
+	{
+		check(StreamSourceSlot);
+
+		FVertexInputStream Stream;
+		Stream.StreamIndex = InStreamIndex;
+		Stream.bStreamSourceSlot = true;
+		Stream.Offset = InOffset;
+		Stream.StreamSourceSlot = StreamSourceSlot;
+
+		// Verify no overflow
+		checkSlow(InStreamIndex == Stream.StreamIndex && InOffset == Stream.Offset);
+
+		return Stream;
+	}
+
 	inline bool operator==(const FVertexInputStream& rhs) const
 	{
 		if (StreamIndex != rhs.StreamIndex ||
-			Offset != rhs.Offset || 
-			VertexBuffer != rhs.VertexBuffer) 
+			Offset != rhs.Offset ||
+			Pointer != rhs.Pointer)
 		{
 			return false;
 		}
diff --git a/Engine/Source/Runtime/Renderer/Private/MeshPassProcessor.cpp b/Engine/Source/Runtime/Renderer/Private/MeshPassProcessor.cpp
--- a/Engine/Source/Runtime/Renderer/Private/MeshPassProcessor.cpp
+++ b/Engine/Source/Runtime/Renderer/Private/MeshPassProcessor.cpp
@@ -1158,21 +1158,21 @@ void FMeshDrawShaderBindings::SetOnCommandList(FRHIComputeCommandList& RHICmdLis
 bool FMeshDrawShaderBindings::MatchesForDynamicInstancing(const FMeshDrawShaderBindings& Rhs) const
 {
 	if (ShaderFrequencyBits != Rhs.ShaderFrequencyBits)
-{
+	{
 		return false;
 	}
 
 	if (ShaderLayouts.Num() != Rhs.ShaderLayouts.Num())
 	{
 		return false;
-}
+	}
 
 	for (int Index = 0; Index < ShaderLayouts.Num(); Index++)
-{
-		if (!(ShaderLayouts[Index] == Rhs.ShaderLayouts[Index]))
 	{
-		return false;
-	}
+		if (!(ShaderLayouts[Index] == Rhs.ShaderLayouts[Index]))
+		{
+			return false;
+		}
 	}
 
 	const uint8* ShaderBindingDataPtr = GetData();
@@ -1420,7 +1420,7 @@ bool FMeshDrawCommand::SubmitDrawBegin(
 		}
 		else if (StateCache.VertexStreams[Stream.StreamIndex] != Stream)
 		{
-			RHICmdList.SetStreamSource(Stream.StreamIndex, Stream.VertexBuffer, Stream.Offset);
+			Stream.SetOnRHICommandList(RHICmdList);
 			StateCache.VertexStreams[Stream.StreamIndex] = Stream;
 		}
 	}
@@ -2198,7 +2198,7 @@ void FCachedPassMeshDrawListContextDeferred::DeferredFinalizeMeshDrawCommands(co
 		{
 			FPrimitiveSceneInfo* SceneInfo = SceneInfos[SceneInfoIndex];
 			for (auto& CmdInfo : SceneInfo->StaticMeshCommandInfos)
-			{				
+			{
 				check(CmdInfo.MeshPass < EMeshPass::Num);
 				FStateBucketMap& BucketMap = Scene.CachedMeshDrawCommandStateBuckets[CmdInfo.MeshPass];
 				int32 DeferredIndex = CmdInfo.StateBucketId;
diff --git a/Engine/Source/Runtime/Renderer/Public/MeshPassProcessor.h b/Engine/Source/Runtime/Renderer/Public/MeshPassProcessor.h
--- a/Engine/Source/Runtime/Renderer/Public/MeshPassProcessor.h
+++ b/Engine/Source/Runtime/Renderer/Public/MeshPassProcessor.h
@@ -1197,7 +1197,7 @@ public:
 
 			uint32 Packed = (StreamIndex << 28) | Offset;
 			HashKey.VertexStreams = FHashKey::HashCombine(HashKey.VertexStreams, Packed);
-			HashKey.VertexBuffers = FHashKey::HashCombine(HashKey.VertexBuffers, FHashKey::PointerHash(VertexInputStream.VertexBuffer));
+			HashKey.VertexBuffers = FHashKey::HashCombine(HashKey.VertexBuffers, FHashKey::PointerHash(VertexInputStream.Pointer));
 		}
 
 		HashKey.IndexBuffer = FHashKey::PointerHash(IndexBuffer);
