Added AddPostExecuteCallback to RDG to support deferring an operation until after graph execution.

diff --git a/Engine/Source/Runtime/RenderCore/Private/RenderGraphBuilder.cpp b/Engine/Source/Runtime/RenderCore/Private/RenderGraphBuilder.cpp
--- a/Engine/Source/Runtime/RenderCore/Private/RenderGraphBuilder.cpp
+++ b/Engine/Source/Runtime/RenderCore/Private/RenderGraphBuilder.cpp
@@ -2004,8 +2004,7 @@ void FRDGBuilder::Execute()
 
 	RHICmdList.SetTrackedAccess(EpilogueResourceAccesses);
 
-	// Wait on the actual parallel execute tasks in the Execute call. When draining is okay to let them overlap with other graph setup.
-	// This also needs to be done before extraction of external resources to be consistent with non-parallel rendering.
+	// Wait on the actual parallel execute tasks in the Execute call. This needs to be done before extraction of external resources to be consistent with non-parallel rendering.
 	if (!ParallelExecute.Tasks.IsEmpty())
 	{
 		UE::Tasks::Wait(ParallelExecute.Tasks);
@@ -2024,6 +2023,12 @@ void FRDGBuilder::Execute()
 		*ExtractedBuffer.PooledBuffer = ExtractedBuffer.Buffer->PooledBuffer;
 	}
 
+	for (TUniqueFunction<void()>& Callback : PostExecuteCallbacks)
+	{
+		Callback();
+	}
+	PostExecuteCallbacks.Empty();
+
 	IF_RDG_ENABLE_TRACE(Trace.OutputGraphEnd(*this));
 
 	IF_RDG_ENABLE_DEBUG(UserValidation.ValidateExecuteEnd());
diff --git a/Engine/Source/Runtime/RenderCore/Public/RenderGraphBuilder.h b/Engine/Source/Runtime/RenderCore/Public/RenderGraphBuilder.h
--- a/Engine/Source/Runtime/RenderCore/Public/RenderGraphBuilder.h
+++ b/Engine/Source/Runtime/RenderCore/Public/RenderGraphBuilder.h
@@ -203,6 +203,13 @@ public:
 
 	//////////////////////////////////////////////////////////////////////////
 
+	/** Adds a callback that is called after pass execution is complete. */
+	void AddPostExecuteCallback(TUniqueFunction<void()>&& Callback)
+	{
+		check(Callback);
+		PostExecuteCallbacks.Emplace(Forward<TUniqueFunction<void()>&&>(Callback));
+	}
+
 	/** Adds a lambda pass to the graph with an accompanied pass parameter struct.
 	 *
 	 *  RDG resources declared in the struct (via _RDG parameter macros) are safe to access in the lambda. The pass parameter struct
@@ -1019,6 +1026,11 @@ private:
 
 	RENDERCORE_API void FlushAccessModeQueue();
 
+	/////////////////////////////////////////////////////////////////////////////
+	// Post-Execution Callbacks
+
+	TArray<TUniqueFunction<void()>, FRDGArrayAllocator> PostExecuteCallbacks;
+
 	/////////////////////////////////////////////////////////////////////////////
 	// Resource Deletion Flushing
 
diff --git a/Engine/Source/Runtime/RenderCore/Public/RenderGraphBuilder.inl b/Engine/Source/Runtime/RenderCore/Public/RenderGraphBuilder.inl
--- a/Engine/Source/Runtime/RenderCore/Public/RenderGraphBuilder.inl
+++ b/Engine/Source/Runtime/RenderCore/Public/RenderGraphBuilder.inl
@@ -273,11 +273,13 @@ FRDGPassRef FRDGBuilder::AddPassInternal(
 
 	FlushAccessModeQueue();
 
+	const TCHAR* NameString = Name.GetTCHAR();
+
 	FRDGPass* Pass = Allocators.Root.AllocNoDestruct<LambdaPassType>(
 		MoveTemp(Name),
 		ParametersMetadata,
 		ParameterStruct,
-		OverridePassFlags(Name.GetTCHAR(), Flags),
+		OverridePassFlags(NameString, Flags),
 		MoveTemp(ExecuteLambda));
 
 	IF_RDG_ENABLE_DEBUG(ClobberPassOutputs(Pass));
