Removed FPipe from render command pipe implementation due to race in WaitUntilEmpty.

diff --git a/Engine/Source/Runtime/RenderCore/Private/RenderingThread.cpp b/Engine/Source/Runtime/RenderCore/Private/RenderingThread.cpp
--- a/Engine/Source/Runtime/RenderCore/Private/RenderingThread.cpp
+++ b/Engine/Source/Runtime/RenderCore/Private/RenderingThread.cpp
@@ -1661,7 +1661,7 @@ public:
 			{
 				Pipe->bRecording = true;
 
-				FRenderCommandPipe::FFrame* NextFrame = new FRenderCommandPipe::FFrame(Pipe->Name, TaskEvent);
+				FRenderCommandPipe::FFrame* NextFrame = new FRenderCommandPipe::FFrame(TaskEvent);
 				PipesToStartRecording.Emplace(Pipe, NextFrame);
 
 				UE::TScopeLock PipeLock(Pipe->Mutex);
@@ -1792,7 +1792,7 @@ private:
 				FRenderCommandPipe* Pipe = AllPipes[BitIt.GetIndex()];
 				FRenderCommandPipe::FFrame*& Frame_RenderThread = Pipe->Frame_RenderThread;
 				check(Frame_RenderThread);
-				Frame_RenderThread->Pipe.WaitUntilEmpty();
+				Frame_RenderThread->LastTask.Wait();
 
 				if (Frame_RenderThread->RHICmdList)
 				{
@@ -1980,7 +1980,7 @@ void FRenderCommandPipe::EnqueueAndLaunch(FFunctionVariant&& FunctionVariant, co
 	{
 		TRACE_CPUPROFILER_EVENT_SCOPE_ON_CHANNEL_STR("RenderCommandPipe LaunchTask", RenderCommandsChannel)
 	
-		Frame_GameThread->Pipe.Launch(Name, [this]
+		Frame_GameThread->LastTask = UE::Tasks::Launch(Name, [this]
 		{
 			check(Frame_RenderThread);
 			TRACE_CPUPROFILER_EVENT_SCOPE_ON_CHANNEL_STR("RenderCommandPipe ReplayCommands", RenderCommandsChannel)
@@ -2005,6 +2005,6 @@ void FRenderCommandPipe::EnqueueAndLaunch(FFunctionVariant&& FunctionVariant, co
 			UE::RenderCommandPipe::ReplayingPipe = PreviousReplayingPipe;
 			NumInFlightCommands.fetch_sub(PoppedQueue.Num(), std::memory_order_release);
 	
-		}, Frame_GameThread->TaskEvent);
+		}, Frame_GameThread->LastTask);
 	}
 }
\ No newline at end of file
diff --git a/Engine/Source/Runtime/RenderCore/Public/RenderingThread.h b/Engine/Source/Runtime/RenderCore/Public/RenderingThread.h
--- a/Engine/Source/Runtime/RenderCore/Public/RenderingThread.h
+++ b/Engine/Source/Runtime/RenderCore/Public/RenderingThread.h
@@ -526,13 +526,11 @@ private:
 
 	struct FFrame : public TConcurrentLinearObject<FFrame>
 	{
-		FFrame(const TCHAR* Name, const UE::Tasks::FTaskEvent& InTaskEvent)
-			: Pipe(Name)
-			, TaskEvent(InTaskEvent)
+		FFrame(const UE::Tasks::FTaskEvent& InTaskEvent)
+			: LastTask(InTaskEvent)
 		{}
 
-		UE::Tasks::FPipe Pipe;
-		UE::Tasks::FTaskEvent TaskEvent;
+		UE::Tasks::FTask LastTask;
 		TArray<FCommand> Queue;
 		FRHICommandList* RHICmdList = nullptr;
 	};
