[shaders] bugfix - widesource wasn't being set when requested in the case where stripping did not occur, resulting in empty debug usf files (i.e. niagara shaders which don't use preprocessed job cache)

diff --git a/Engine/Source/Runtime/RenderCore/Public/ShaderPreprocessTypes.h b/Engine/Source/Runtime/RenderCore/Public/ShaderPreprocessTypes.h
--- a/Engine/Source/Runtime/RenderCore/Public/ShaderPreprocessTypes.h
+++ b/Engine/Source/Runtime/RenderCore/Public/ShaderPreprocessTypes.h
@@ -115,17 +115,21 @@ public:
 		// if the unstripped source is requested, check if the "original source" field has been populated
 		// if not then stripping hasn't occurred so there's only one preprocessed source; return it
 #if SHADER_SOURCE_ANSI
+		// convert and store wide versions of requested source if view is requested. 
+		// this is only used in debug paths (debug output and viewing source in-editor)
+		// and the job should be freed shortly after, so the memory overhead is not a concern
 		if (OriginalPreprocessedSource.IsEmpty())
 		{
+			if (WideSource.IsEmpty())
+			{
+				WideSource = FString(PreprocessedSource.GetView());
+			}
 			return FStringView(WideSource);
 		}
 		else
 		{
 			if (WideSourceUnstripped.IsEmpty())
 			{
-				// convert and store wide version of unstripped source if view is requested. 
-				// this is only called in debug paths (a particular case of shader debug dumps,
-				// and visualizing shader source in the editor) so the overhead doesn't matter.
 				WideSourceUnstripped = FString(OriginalPreprocessedSource.GetView());
 			}
 			return FStringView(WideSourceUnstripped);
