[Nanite-Skinning] Added missing shader map DDC key values for r.Nanite.AllowSplineMeshes and r.Nanite.AllowSkinnedMeshes

diff --git a/Engine/Source/Runtime/RenderCore/Private/Shader.cpp b/Engine/Source/Runtime/RenderCore/Private/Shader.cpp
--- a/Engine/Source/Runtime/RenderCore/Private/Shader.cpp
+++ b/Engine/Source/Runtime/RenderCore/Private/Shader.cpp
@@ -2623,14 +2623,18 @@ void ShaderMapAppendKey(EShaderPlatform Platform, FShaderKeyGenerator& KeyGen)
 		static const auto CVarAllowTess = IConsoleManager::Get().FindConsoleVariable(TEXT("r.Nanite.AllowTessellation"));
 		static const auto CVarAllowCSMat = IConsoleManager::Get().FindConsoleVariable(TEXT("r.Nanite.AllowComputeMaterials"));
 		static const auto CVarAllowPSMat = IConsoleManager::Get().FindConsoleVariable(TEXT("r.Nanite.AllowLegacyMaterials"));
+		static const auto CVarAllowSpline = IConsoleManager::Get().FindConsoleVariable(TEXT("r.Nanite.AllowSplineMeshes"));
+		static const auto CVarAllowSkinned = IConsoleManager::Get().FindConsoleVariable(TEXT("r.Nanite.AllowSkinnedMeshes"));
 
 		KeyGen.AppendSeparator();
 		KeyGen.Append(TEXT("Nanite-Tess"));
 		KeyGen.Append(CVarAllowTess ? CVarAllowTess->GetInt() : 0);
 		KeyGen.AppendDebugText(TEXT("CSMat"));
 		KeyGen.Append(CVarAllowCSMat ? CVarAllowCSMat->GetInt() : 0);
-		KeyGen.AppendDebugText(TEXT("PSMat"));
-		KeyGen.Append(CVarAllowPSMat ? CVarAllowPSMat->GetInt() : 0);
+		KeyGen.AppendDebugText(TEXT("Spline"));
+		KeyGen.Append(CVarAllowSpline ? CVarAllowSpline->GetInt() : 0);
+		KeyGen.AppendDebugText(TEXT("Skinned"));
+		KeyGen.Append(CVarAllowSkinned ? CVarAllowSkinned->GetInt() : 0);
 	}
 }
 
