Vectorized TMatrix<float>::To3x4MatrixTranspose, FRenderTransform::To3x4MatrixTranspose, and FRenderTransform::ToMatrix44f

diff --git a/Engine/Source/Runtime/Core/Public/Math/Matrix.h b/Engine/Source/Runtime/Core/Public/Math/Matrix.h
--- a/Engine/Source/Runtime/Core/Public/Math/Matrix.h
+++ b/Engine/Source/Runtime/Core/Public/Math/Matrix.h
@@ -301,7 +301,6 @@ public:
 	 */
 	inline void SetAxes(const TVector<T>* Axis0 = NULL, const TVector<T>* Axis1 = NULL, const TVector<T>* Axis2 = NULL, const TVector<T>* Origin = NULL);
 
-
 	/**
 	 * get a column of this matrix
 	 *
@@ -328,6 +327,11 @@ public:
 	 */
 	CORE_API UE::Math::TQuat<T> ToQuat() const;
 
+	/**
+	 * Convert this Atom to the 3x4 transpose of the transformation matrix.
+	 */
+	FORCEINLINE void To3x4MatrixTranspose(T* Out) const;
+
 	// Frustum plane extraction.
 	/** @param OuTPln the near plane of the Frustum of this matrix */
 	FORCEINLINE bool GetFrustumNearPlane(TPlane<T>& OuTPln) const;
@@ -402,30 +406,6 @@ public:
 	
 	bool SerializeFromMismatchedTag(FName StructTag, FArchive& Ar);
 
-	/**
-	 * Convert this Atom to the 3x4 transpose of the transformation matrix.
-	 */
-	void To3x4MatrixTranspose(T* Out) const
-	{
-		const T* RESTRICT Src = &(M[0][0]);
-		T* RESTRICT Dest = Out;
-
-		Dest[0] = Src[0];   // [0][0]
-		Dest[1] = Src[4];   // [1][0]
-		Dest[2] = Src[8];   // [2][0]
-		Dest[3] = Src[12];  // [3][0]
-
-		Dest[4] = Src[1];   // [0][1]
-		Dest[5] = Src[5];   // [1][1]
-		Dest[6] = Src[9];   // [2][1]
-		Dest[7] = Src[13];  // [3][1]
-
-		Dest[8] = Src[2];   // [0][2]
-		Dest[9] = Src[6];   // [1][2]
-		Dest[10] = Src[10]; // [2][2]
-		Dest[11] = Src[14]; // [3][2]
-	}
-
 	// Conversion to other type.
 	template<typename FArg, TEMPLATE_REQUIRES(!std::is_same_v<T, FArg>)>
 	explicit TMatrix(const TMatrix<FArg>& From)
diff --git a/Engine/Source/Runtime/Core/Public/Math/Matrix.inl b/Engine/Source/Runtime/Core/Public/Math/Matrix.inl
--- a/Engine/Source/Runtime/Core/Public/Math/Matrix.inl
+++ b/Engine/Source/Runtime/Core/Public/Math/Matrix.inl
@@ -250,6 +250,73 @@ FORCEINLINE TMatrix<T> TMatrix<T>::GetTransposed() const
 	return Result;
 }
 
+/**
+* Convert this Atom to the 3x4 transpose of the transformation matrix.
+*/
+template<typename T>
+FORCEINLINE void TMatrix<T>::To3x4MatrixTranspose(T* Out) const
+{
+	const T* RESTRICT Src = &(M[0][0]);
+	T* RESTRICT Dest = Out;
+
+	Dest[0] = Src[0];   // [0][0]
+	Dest[1] = Src[4];   // [1][0]
+	Dest[2] = Src[8];   // [2][0]
+	Dest[3] = Src[12];  // [3][0]
+
+	Dest[4] = Src[1];   // [0][1]
+	Dest[5] = Src[5];   // [1][1]
+	Dest[6] = Src[9];   // [2][1]
+	Dest[7] = Src[13];  // [3][1]
+
+	Dest[8] = Src[2];   // [0][2]
+	Dest[9] = Src[6];   // [1][2]
+	Dest[10] = Src[10];  // [2][2]
+	Dest[11] = Src[14];  // [3][2]
+}
+
+template<>
+FORCEINLINE void TMatrix<float>::To3x4MatrixTranspose(float* Out) const
+{
+	const float* RESTRICT Src = &(M[0][0]);
+	float* RESTRICT Dest = Out;
+
+#if PLATFORM_ENABLE_VECTORINTRINSICS
+	VectorRegister4Float InRow0 = VectorLoad(&Src[0]);
+	VectorRegister4Float InRow1 = VectorLoad(&Src[4]);
+	VectorRegister4Float InRow2 = VectorLoad(&Src[8]);
+	VectorRegister4Float InRow3 = VectorLoad(&Src[12]);
+
+	VectorRegister4Float Temp0 = VectorCombineLow( InRow0, InRow1);
+	VectorRegister4Float Temp1 = VectorCombineLow( InRow2, InRow3);
+	VectorRegister4Float Temp2 = VectorCombineHigh(InRow0, InRow1);
+	VectorRegister4Float Temp3 = VectorCombineHigh(InRow2, InRow3);
+
+	VectorRegister4Float Row0, Row1, Row2, Row3;
+	VectorDeinterleave(Row0, Row1, Temp0, Temp1);
+	VectorDeinterleave(Row2, Row3, Temp2, Temp3);
+
+	VectorStore(Row0, &Dest[0]);
+	VectorStore(Row1, &Dest[4]);
+	VectorStore(Row2, &Dest[8]);
+#else
+	Dest[0]  = Src[0];   // [0][0]
+	Dest[1]  = Src[4];   // [1][0]
+	Dest[2]  = Src[8];   // [2][0]
+	Dest[3]  = Src[12];  // [3][0]
+
+	Dest[4]  = Src[1];   // [0][1]
+	Dest[5]  = Src[5];   // [1][1]
+	Dest[6]  = Src[9];   // [2][1]
+	Dest[7]  = Src[13];  // [3][1]
+
+	Dest[8]  = Src[2];   // [0][2]
+	Dest[9]  = Src[6];   // [1][2]
+	Dest[10] = Src[10];  // [2][2]
+	Dest[11] = Src[14];  // [3][2]
+#endif
+}
+
 // Determinant.
 
 template<typename T>
diff --git a/Engine/Source/Runtime/RenderCore/Public/RenderTransform.h b/Engine/Source/Runtime/RenderCore/Public/RenderTransform.h
--- a/Engine/Source/Runtime/RenderCore/Public/RenderTransform.h
+++ b/Engine/Source/Runtime/RenderCore/Public/RenderTransform.h
@@ -87,6 +87,18 @@ public:
 	FORCEINLINE FMatrix44f ToMatrix44f() const
 	{
 		FMatrix44f Matrix;
+	#if PLATFORM_ENABLE_VECTORINTRINSICS
+		VectorRegister4Float Row0 = VectorSet_W0(VectorLoad(&TransformRows[0].X)); // Intentionally loads W from the next row (then set to 0)
+		VectorRegister4Float Row1 = VectorSet_W0(VectorLoad(&TransformRows[1].X)); // Intentionally loads W from the next row (then set to 0)
+		VectorRegister4Float Row2 = VectorSet_W0(VectorLoad(&TransformRows[2].X)); // Intentionally loads W from the next row (then set to 0)
+		VectorRegister4Float Row3 = MakeVectorRegisterFloat(Origin.X, Origin.Y, Origin.Z, 1.0f); // Do not read past Origin
+
+		float* RESTRICT Dest = &Matrix.M[0][0];
+		VectorStore(Row0, &Dest[0]);
+		VectorStore(Row1, &Dest[4]);
+		VectorStore(Row2, &Dest[8]);
+		VectorStore(Row3, &Dest[12]);
+	#else
 		Matrix.M[0][0] = TransformRows[0].X;
 		Matrix.M[0][1] = TransformRows[0].Y;
 		Matrix.M[0][2] = TransformRows[0].Z;
@@ -103,6 +115,7 @@ public:
 		Matrix.M[3][1] = Origin.Y;
 		Matrix.M[3][2] = Origin.Z;
 		Matrix.M[3][3] = 1.0f;
+	#endif
 		return Matrix;
 	}
 
@@ -115,6 +128,25 @@ public:
 	{
 		float* RESTRICT Dest = Result;
 
+	#if PLATFORM_ENABLE_VECTORINTRINSICS
+		VectorRegister4Float InRow0 = VectorLoad(&TransformRows[0].X); // Intentionally loads W from the next row (discarded)
+		VectorRegister4Float InRow1 = VectorLoad(&TransformRows[1].X); // Intentionally loads W from the next row (discarded)
+		VectorRegister4Float InRow2 = VectorLoad(&TransformRows[2].X); // Intentionally loads W from the next row (discarded)
+		VectorRegister4Float InRow3 = MakeVectorRegisterFloat(Origin.X, Origin.Y, Origin.Z, 0.0f); // Do not read past Origin
+
+		VectorRegister4Float Temp0 = VectorCombineLow( InRow0, InRow1);
+		VectorRegister4Float Temp1 = VectorCombineLow( InRow2, InRow3);
+		VectorRegister4Float Temp2 = VectorCombineHigh(InRow0, InRow1);
+		VectorRegister4Float Temp3 = VectorCombineHigh(InRow2, InRow3);
+
+		VectorRegister4Float Row0, Row1, Row2, Row3;
+		VectorDeinterleave(Row0, Row1, Temp0, Temp1);
+		VectorDeinterleave(Row2, Row3, Temp2, Temp3);
+
+		VectorStore(Row0, &Dest[0]);
+		VectorStore(Row1, &Dest[4]);
+		VectorStore(Row2, &Dest[8]);
+	#else
 		Dest[ 0] = TransformRows[0].X;	// [0][0]
 		Dest[ 1] = TransformRows[1].X;	// [1][0]
 		Dest[ 2] = TransformRows[2].X;	// [2][0]
@@ -129,6 +161,7 @@ public:
 		Dest[ 9] = TransformRows[1].Z;	// [1][2]
 		Dest[10] = TransformRows[2].Z;	// [2][2]
 		Dest[11] = Origin.Z;			// [3][2]
+	#endif
 	}
 
 	FORCEINLINE FRenderTransform operator* (const FRenderTransform& Other) const
