[shaders]
- set bCachePreprocessed member to true by default instead of false; the only impact of this is to enable this feature for niagara VM shaders which were previously still preprocessing in SCW unintentionally
- add a CFLAG to disable the source stripping process that runs at the end of preprocessing if preprocessed job cache is enabled (and further a cvar to disable this process across the board for debug workflow purposes)
- always disable the stripping for vectorVM shaders; this backend blows up if the shader source contains comments, and some minimal comments are added as part of the stripping process
- change condition which determines whether to run preprocessing in the compile step to check the preprocess output bSucceeded flag instead of bCachePreprocessed; this works in both SCW directcompile mode and normal execution without needing the former to explicitly set that flag to false (since it's now true by default)

diff --git a/Engine/Source/Developer/ShaderFormatVectorVM/Private/ShaderFormatVectorVM.cpp b/Engine/Source/Developer/ShaderFormatVectorVM/Private/ShaderFormatVectorVM.cpp
--- a/Engine/Source/Developer/ShaderFormatVectorVM/Private/ShaderFormatVectorVM.cpp
+++ b/Engine/Source/Developer/ShaderFormatVectorVM/Private/ShaderFormatVectorVM.cpp
@@ -69,6 +69,9 @@ public:
 
 		// minifier has not been tested on vector VM; it's possible this could be removed to improve deduplication rate
 		Input.Environment.CompilerFlags.Remove(CFLAG_RemoveDeadCode);
+
+		// source stripping process adds comments which are not handled by this backend
+		Input.Environment.CompilerFlags.Add(CFLAG_DisableSourceStripping);
 	}
 
 	virtual void CompilePreprocessedShader(const FShaderCompilerInput& Input, const FShaderPreprocessOutput& PreprocessOutput, FShaderCompilerOutput& Output, const FString& WorkingDirectory) const override
diff --git a/Engine/Source/Runtime/Engine/Private/ShaderCompiler/ShaderCompiler.cpp b/Engine/Source/Runtime/Engine/Private/ShaderCompiler/ShaderCompiler.cpp
--- a/Engine/Source/Runtime/Engine/Private/ShaderCompiler/ShaderCompiler.cpp
+++ b/Engine/Source/Runtime/Engine/Private/ShaderCompiler/ShaderCompiler.cpp
@@ -221,6 +221,12 @@ static TAutoConsoleVariable<bool> CVarShaderCompilerPerShaderDDCCook(
 	ECVF_Default
 );
 
+static TAutoConsoleVariable<bool> CVarDisableSourceStripping(
+	TEXT("r.ShaderCompiler.DisableSourceStripping"),
+	false,
+	TEXT("If true, the process which strips comments, line directives and whitespace from final preprocessed source is disabled. This results in file associations being maintained and visible in RenderDoc etc., at the cost of less effective deduplication."),
+	ECVF_Default);
+
 static TAutoConsoleVariable<bool> CVarAreShaderErrorsFatal(
 	TEXT("r.AreShaderErrorsFatal"),
 	true,
@@ -8741,6 +8747,11 @@ void GlobalBeginCompileShader(
 		Input.Environment.CompilerFlags.Add(CFLAG_RemoveDeadCode);
 	}
 
+	if (CVarDisableSourceStripping.GetValueOnAnyThread())
+	{
+		Input.Environment.CompilerFlags.Add(CFLAG_DisableSourceStripping);
+	}
+
 	{
 		static const auto CVar = IConsoleManager::Get().FindConsoleVariable(TEXT("r.VT.AnisotropicFiltering"));
 		SET_SHADER_DEFINE(Input.Environment, VIRTUAL_TEXTURE_ANISOTROPIC_FILTERING, CVar ? (CVar->GetInt() != 0) : 0);
diff --git a/Engine/Source/Runtime/RenderCore/Private/ShaderCore.cpp b/Engine/Source/Runtime/RenderCore/Private/ShaderCore.cpp
--- a/Engine/Source/Runtime/RenderCore/Private/ShaderCore.cpp
+++ b/Engine/Source/Runtime/RenderCore/Private/ShaderCore.cpp
@@ -1499,20 +1499,24 @@ public:
 
 		if (Job.PreprocessOutput.bSucceeded && Job.Input.bCachePreprocessed)
 		{
-			// if the preprocessed job cache is enabled we need to strip the preprocessed code, this removes comments, line directives
-			// and blank lines to improve deduplication (and populates data required to remap diagnostic messages to correct line numbers)
-			Job.PreprocessOutput.StripCode(Job.Input.NeedsOriginalShaderSource());
-
-			FShaderCompilerInputHash Hash = Job.GetInputHash();
-			// Replace the placeholder debug hash value appended in StripCode with the real job input hash
-			FShaderSource::FViewType DebugHashStr = GetShaderSourceDebugHashPrefix();
-			FShaderSource::FViewType SourceView = Job.PreprocessOutput.GetSourceView();
-			int32 DebugHashLoc = SourceView.Find(DebugHashStr) + DebugHashStr.Len();
-			int32 NewlineLoc = SourceView.Find(SHADER_SOURCE_VIEWLITERAL("\n"), DebugHashLoc);
-			FShaderSource::TStringBuilder<2 * sizeof(FShaderCompilerInputHash::ByteArray) + 1> HashStr;
-			HashStr << Hash;
-			check(NewlineLoc - DebugHashLoc == HashStr.Len());
-			FMemory::Memcpy(Job.PreprocessOutput.EditSource().GetData() + DebugHashLoc, HashStr.GetData(), sizeof(FShaderSource::CharType) * HashStr.Len());
+			if (!Job.Input.Environment.CompilerFlags.Contains(CFLAG_DisableSourceStripping))
+			{
+				// if the preprocessed job cache is enabled we strip the preprocessed code if not explicitly disabled; this removes comments, 
+				// line directives and blank lines to improve deduplication (and populates data required to remap diagnostic messages to correct
+				// filenames and line numbers)
+				Job.PreprocessOutput.StripCode(Job.Input.NeedsOriginalShaderSource());
+
+				FShaderCompilerInputHash Hash = Job.GetInputHash();
+				// Replace the placeholder debug hash value appended in StripCode with the real job input hash
+				FShaderSource::FViewType DebugHashStr = GetShaderSourceDebugHashPrefix();
+				FShaderSource::FViewType SourceView = Job.PreprocessOutput.GetSourceView();
+				int32 DebugHashLoc = SourceView.Find(DebugHashStr) + DebugHashStr.Len();
+				int32 NewlineLoc = SourceView.Find(SHADER_SOURCE_VIEWLITERAL("\n"), DebugHashLoc);
+				FShaderSource::TStringBuilder<2 * sizeof(FShaderCompilerInputHash::ByteArray) + 1> HashStr;
+				HashStr << Hash;
+				check(NewlineLoc - DebugHashLoc == HashStr.Len());
+				FMemory::Memcpy(Job.PreprocessOutput.EditSource().GetData() + DebugHashLoc, HashStr.GetData(), sizeof(FShaderSource::CharType) * HashStr.Len());
+			}
 
 			// always compress the code after stripping to minimize memory footprint
 			Job.PreprocessOutput.CompressCode();
@@ -1577,7 +1581,7 @@ public:
 	static void CompileShaderInternal(const IShaderFormat* Compiler, FShaderCompileJob& Job, const FString& WorkingDirectory, FString& OutExceptionMsg, FString& OutExceptionCallstack, int32* CompileCount)
 	{
 		double TimeStart = FPlatformTime::Seconds();
-		if (!Job.Input.bCachePreprocessed)
+		if (!Job.PreprocessOutput.GetSucceeded())
 		{
 			PreprocessShaderInternal(Compiler, Job);
 		}
diff --git a/Engine/Source/Runtime/RenderCore/Public/ShaderCompilerCore.h b/Engine/Source/Runtime/RenderCore/Public/ShaderCompilerCore.h
--- a/Engine/Source/Runtime/RenderCore/Public/ShaderCompilerCore.h
+++ b/Engine/Source/Runtime/RenderCore/Public/ShaderCompilerCore.h
@@ -164,6 +164,8 @@ enum ECompilerFlags
 	CFLAG_ShaderBundle,
 	// Shader contains workgraph nodes that are only intended for linking against other workgraph shaders containing the program entry node.
 	CFLAG_WorkgraphLocalNodes,
+	// Shader code should not be stripped of comments/whitespace/line directives at the end of preprocessing
+	CFLAG_DisableSourceStripping,
 	CFLAG_Max,
 };
 static_assert(CFLAG_Max < 64, "Out of bitfield space! Modify FShaderCompilerFlags");
@@ -234,7 +236,7 @@ using FShaderCompilerInputHash = FBlake3Hash;
 /** Struct that gathers all readonly inputs needed for the compilation of a single shader. */
 struct FShaderCompilerInput
 {
-	FShaderTarget Target;
+	FShaderTarget Target{ SF_NumFrequencies, SP_NumPlatforms };
 	
 	FName ShaderFormat;
 	FName CompressionFormat;
@@ -246,28 +248,28 @@ struct FShaderCompilerInput
 
 	// Skips the preprocessor and instead loads the usf file directly
 	UE_DEPRECATED(5.4, "bSkipPreprocessedCache member is deprecated; set EShaderDebugInfoFlags::CompileFromDebugUSF on DebugInfoFlags instead.")
-	bool bSkipPreprocessedCache;
+	bool bSkipPreprocessedCache = false;
 
 	// Indicates which additional debug outputs should be written for this compile job.
-	EShaderDebugInfoFlags DebugInfoFlags;
+	EShaderDebugInfoFlags DebugInfoFlags = EShaderDebugInfoFlags::Default;
 
 	UE_DEPRECATED(5.4, "bIndependentPreprocessed member no longer used now that all backends have been migrated to the new IShaderFormat API")
 	// True if the backend for this job implements the independent preprocessing API.
-	bool bIndependentPreprocessed;
+	bool bIndependentPreprocessed = true;
 
 	// True if the cache key for this job should be based on preprocessed source. If so,
 	// preprocessing will be executed in the cook process independent of compilation (and
 	// as such this will only ever be set for jobs whose shader format supports independent
 	// preprocessing)
-	bool bCachePreprocessed;
+	bool bCachePreprocessed = true;
 
 	// Array of symbols that should be maintained when deadstripping. If this is empty, entry
 	// point name alone will be used.
 	TArray<FString> RequiredSymbols;
 	
 	// Shader pipeline information
-	bool bCompilingForShaderPipeline;
-	bool bIncludeUsedOutputs;
+	bool bCompilingForShaderPipeline = false;
+	bool bIncludeUsedOutputs = false;
 	TArray<FString> UsedOutputs;
 
 	// Dump debug path (up to platform) e.g. "D:/Project/Saved/ShaderDebugInfo/PCD3D_SM5"
@@ -304,20 +306,11 @@ struct FShaderCompilerInput
 	/** Oodle-specific compression level - used if CompressionFormat is set to NAME_Oodle. */
 	FOodleDataCompression::ECompressionLevel OodleLevel;
 
-	FShaderCompilerInput() :
-		Target(SF_NumFrequencies, SP_NumPlatforms),
-		bSkipPreprocessedCache(false),
-		DebugInfoFlags(EShaderDebugInfoFlags::Default),
-		bCachePreprocessed(false),
-		bCompilingForShaderPipeline(false),
-		bIncludeUsedOutputs(false)
-	{
-	}
-
 	PRAGMA_DISABLE_DEPRECATION_WARNINGS
-	// Explicitly-defaulted copy/move ctors & assignment operators are needed temporarily due to 
+	// Explicitly-defaulted ctors & assignment operators are needed temporarily due to 
 	// deprecation of bSkipPreprocessedCache/bIndependentPreprocessed fields. These can be removed once 
 	// the deprecation window for said fields ends.
+	FShaderCompilerInput() = default;
 	FShaderCompilerInput(FShaderCompilerInput&&) = default;
 	FShaderCompilerInput(const FShaderCompilerInput&) = default;
 	FShaderCompilerInput& operator=(FShaderCompilerInput&&) = default;
