Correctly change  between Mobile Forward and Mobile Deferred, taking into account r.Mobile.AllowDeferredShadingOpenGL cvar

diff --git a/Engine/Source/Runtime/RenderCore/Private/ReadOnlyCVARCache.cpp b/Engine/Source/Runtime/RenderCore/Private/ReadOnlyCVARCache.cpp
--- a/Engine/Source/Runtime/RenderCore/Private/ReadOnlyCVARCache.cpp
+++ b/Engine/Source/Runtime/RenderCore/Private/ReadOnlyCVARCache.cpp
@@ -48,7 +48,16 @@ bool FReadOnlyCVARCache::MobileDeferredShadingIniValue(EShaderPlatform Platform)
 	static FShaderPlatformCachedIniValue<bool> MobileShadingPathIniValue(TEXT("r.Mobile.ShadingPath"));
 	static TConsoleVariableData<int32>* MobileAllowDeferredShadingOpenGL = IConsoleManager::Get().FindTConsoleVariableDataInt(TEXT("r.Mobile.AllowDeferredShadingOpenGL"));
 	// a separate cvar so we can exclude deferred from OpenGL specificaly
-	const bool bSupportedPlatform = !IsOpenGLPlatform(Platform) || (MobileAllowDeferredShadingOpenGL && MobileAllowDeferredShadingOpenGL->GetValueOnAnyThread() != 0);
+
+	bool bIsOpenGLPlatform = IsOpenGLPlatform(Platform);
+#if WITH_EDITOR
+	if (FDataDrivenShaderPlatformInfo::GetIsPreviewPlatform(Platform))
+	{
+		EShaderPlatform ParentShaderPlatform = FDataDrivenShaderPlatformInfo::GetPreviewShaderPlatformParent(Platform);
+		bIsOpenGLPlatform = IsOpenGLPlatform(ParentShaderPlatform);
+	}
+#endif
+	const bool bSupportedPlatform = !bIsOpenGLPlatform || (MobileAllowDeferredShadingOpenGL && MobileAllowDeferredShadingOpenGL->GetValueOnAnyThread() != 0);
 	return MobileShadingPathIniValue.Get(Platform) && bSupportedPlatform;
 }
 
