Fixed missing transient discard for transient extracted textures. Also properly disabled transient async compute aliasing on platforms that use the page allocator.

diff --git a/Engine/Source/Runtime/D3D12RHI/Private/D3D12Adapter.cpp b/Engine/Source/Runtime/D3D12RHI/Private/D3D12Adapter.cpp
--- a/Engine/Source/Runtime/D3D12RHI/Private/D3D12Adapter.cpp
+++ b/Engine/Source/Runtime/D3D12RHI/Private/D3D12Adapter.cpp
@@ -1139,7 +1139,10 @@ void FD3D12Adapter::InitializeDevices()
 #endif
 		}
 
-		GRHIGlobals.SupportsAsyncComputeTransientAliasing = GSupportsEfficientAsyncCompute;
+		if (!FPlatformMemory::SupportsFastVRAMMemory() && GSupportsEfficientAsyncCompute)
+		{
+			GRHIGlobals.SupportsAsyncComputeTransientAliasing = true;
+		}
 
 #if PLATFORM_WINDOWS
 		D3D12_FEATURE_DATA_D3D12_OPTIONS2 D3D12Caps2 = {};
diff --git a/Engine/Source/Runtime/RenderCore/Private/RenderGraphBuilder.cpp b/Engine/Source/Runtime/RenderCore/Private/RenderGraphBuilder.cpp
--- a/Engine/Source/Runtime/RenderCore/Private/RenderGraphBuilder.cpp
+++ b/Engine/Source/Runtime/RenderCore/Private/RenderGraphBuilder.cpp
@@ -2876,6 +2876,7 @@ void FRDGBuilder::AllocateTransientResources(TConstArrayView<FCollectResourceOp>
 					{
 						// This releases the reference without invoking a virtual function call.
 						GRDGTransientResourceAllocator.Release(TRefCountPtr<FRDGTransientRenderTarget>(MoveTemp(Texture->Allocation)), GetDeallocateFences(Texture));
+						SetDiscardPass(Texture, TransientTexture);
 					}
 				}
 				// Texture is using an internal transient texture.
@@ -4180,6 +4181,14 @@ void FRDGBuilder::SetPooledTextureRHI(FRDGTexture* Texture, FRDGPooledTexture* P
 	Owner = Texture;
 }
 
+void FRDGBuilder::SetDiscardPass(FRDGTexture* Texture, FRHITransientTexture* TransientTexture)
+{
+	if (TransientTexture->IsDiscarded())
+	{
+		Texture->DiscardPass = FRDGPassHandle(FMath::Min<uint32>(TransientTexture->GetDiscardPass(), GetEpiloguePassHandle().GetIndex()));
+	}
+}
+
 void FRDGBuilder::SetTransientTextureRHI(FRDGTexture* Texture, FRHITransientTexture* TransientTexture)
 {
 	Texture->SetRHI(TransientTexture->GetRHI());
@@ -4187,10 +4196,7 @@ void FRDGBuilder::SetTransientTextureRHI(FRDGTexture* Texture, FRHITransientText
 	Texture->ViewCache = &TransientTexture->ViewCache;
 	Texture->AliasingOverlaps = TransientTexture->GetAliasingOverlaps();
 
-	if (TransientTexture->IsDiscarded())
-	{
-		Texture->DiscardPass = FRDGPassHandle(FMath::Min<uint32>(TransientTexture->GetDiscardPass(), GetEpiloguePassHandle().GetIndex()));
-	}
+	SetDiscardPass(Texture, TransientTexture);
 }
 
 void FRDGBuilder::SetExternalPooledBufferRHI(FRDGBuffer* Buffer, FRDGPooledBuffer* PooledBuffer)
diff --git a/Engine/Source/Runtime/RenderCore/Public/RenderGraphBuilder.h b/Engine/Source/Runtime/RenderCore/Public/RenderGraphBuilder.h
--- a/Engine/Source/Runtime/RenderCore/Public/RenderGraphBuilder.h
+++ b/Engine/Source/Runtime/RenderCore/Public/RenderGraphBuilder.h
@@ -688,6 +688,7 @@ private:
 	void SetExternalPooledRenderTargetRHI(FRDGTexture* Texture, IPooledRenderTarget* RenderTarget);
 	void SetPooledTextureRHI(FRDGTexture* Texture, FRDGPooledTexture* PooledTexture);
 	void SetTransientTextureRHI(FRDGTexture* Texture, FRHITransientTexture* TransientTexture);
+	void SetDiscardPass(FRDGTexture* Texture, FRHITransientTexture* TransientTexture);
 	void SetExternalPooledBufferRHI(FRDGBuffer* Buffer, FRDGPooledBuffer* PooledBuffer);
 	void SetPooledBufferRHI(FRDGBuffer* Buffer, FRDGPooledBuffer* PooledBuffer);
 	void SetTransientBufferRHI(FRDGBuffer* Buffer, FRHITransientBuffer* TransientBuffer);
