Reverting some logic changes that lumped UAVs with read-only resources.

diff --git a/Engine/Shaders/Public/ShaderVersion.ush b/Engine/Shaders/Public/ShaderVersion.ush
--- a/Engine/Shaders/Public/ShaderVersion.ush
+++ b/Engine/Shaders/Public/ShaderVersion.ush
@@ -3,4 +3,4 @@
 // Each time the console command is executed it generates a new GUID. As this file is included
 // in Platform.ush (which should be included in any shader) it allows to invalidate the shader DDC.
 // 
-#pragma message("UESHADERMETADATA_VERSION AB1B3135-817D-44A6-B0AF-A1CDE90D0BA2")
+#pragma message("UESHADERMETADATA_VERSION 948D2352-F48E-46B5-A4A7-0F6EAA8652BB")
diff --git a/Engine/Source/Runtime/RHI/Public/RHIDefinitions.h b/Engine/Source/Runtime/RHI/Public/RHIDefinitions.h
--- a/Engine/Source/Runtime/RHI/Public/RHIDefinitions.h
+++ b/Engine/Source/Runtime/RHI/Public/RHIDefinitions.h
@@ -1346,16 +1346,22 @@ inline bool IsRDGResourceReferenceShaderParameterType(EUniformBufferBaseType Bas
 	return IsRDGTextureReferenceShaderParameterType(BaseType) || IsRDGBufferReferenceShaderParameterType(BaseType) || BaseType == UBMT_RDG_UNIFORM_BUFFER;
 }
 
-inline bool IsShaderParameterTypeRHIResource(EUniformBufferBaseType BaseType)
+inline bool IsShaderParameterTypeReadOnlyRHIResource(EUniformBufferBaseType BaseType)
 {
 	return
 		BaseType == UBMT_TEXTURE ||
 		BaseType == UBMT_SRV ||
 		BaseType == UBMT_SAMPLER ||
-		BaseType == UBMT_UAV ||
 		BaseType == UBMT_RESOURCE_COLLECTION;
 }
 
+inline bool IsShaderParameterTypeRHIResource(EUniformBufferBaseType BaseType)
+{
+	return
+		IsShaderParameterTypeReadOnlyRHIResource(BaseType) ||
+		BaseType == UBMT_UAV;
+}
+
 /** Returns whether the shader parameter type needs to be passdown to RHI through FRHIUniformBufferLayout when creating an uniform buffer. */
 inline bool IsShaderParameterTypeForUniformBufferLayout(EUniformBufferBaseType BaseType)
 {
diff --git a/Engine/Source/Runtime/RenderCore/Private/ShaderParameterMetadata.cpp b/Engine/Source/Runtime/RenderCore/Private/ShaderParameterMetadata.cpp
--- a/Engine/Source/Runtime/RenderCore/Private/ShaderParameterMetadata.cpp
+++ b/Engine/Source/Runtime/RenderCore/Private/ShaderParameterMetadata.cpp
@@ -483,7 +483,7 @@ void FShaderParametersMetadata::FMember::HashLayout(FMemoryHasherBlake3& Hasher)
 	Hasher.Serialize(const_cast<TCHAR*>(Name), FCString::Strlen(Name));
 	Hasher << NumElements;
 
-	const bool bIsRHIResource = IsShaderParameterTypeRHIResource(BaseType);
+	const bool bIsRHIResource = IsShaderParameterTypeReadOnlyRHIResource(BaseType);
 	const bool bIsRDGResource = IsRDGResourceReferenceShaderParameterType(BaseType);
 
 	if (BaseType == UBMT_INT32 ||
@@ -576,7 +576,7 @@ void FShaderParametersMetadata::InitializeLayout(FRHIUniformBufferLayoutInitiali
 		const TCHAR* ShaderType = CurrentMember.GetShaderType();
 
 		const bool bIsArray = ArraySize > 0;
-		const bool bIsRHIResource = IsShaderParameterTypeRHIResource(BaseType);
+		const bool bIsRHIResource = IsShaderParameterTypeReadOnlyRHIResource(BaseType);
 		const bool bIsRDGResource = IsRDGResourceReferenceShaderParameterType(BaseType);
 		const bool bIsVariableNativeType = CurrentMember.IsVariableNativeType();
 
@@ -823,7 +823,7 @@ void FShaderParametersMetadata::InitializeLayout(FRHIUniformBufferLayoutInitiali
 			MemberHash = HashCombine(MemberHash, FCrc::Strihash_DEPRECATED(CurrentMember.GetName()));
 			MemberHash = HashCombine(MemberHash, GetTypeHash(int32(CurrentMember.GetNumElements())));
 
-			const bool bIsRHIResource = IsShaderParameterTypeRHIResource(BaseType);
+			const bool bIsRHIResource = IsShaderParameterTypeReadOnlyRHIResource(BaseType);
 			const bool bIsRDGResource = IsRDGResourceReferenceShaderParameterType(BaseType);
 
 			if (BaseType == UBMT_INT32 ||
@@ -1083,7 +1083,7 @@ void FShaderParametersMetadata::FindMemberFromOffset(uint16 MemberOffset, const
 			}
 		}
 		else if (NumElements > 0 && (
-			IsShaderParameterTypeRHIResource(BaseType) ||
+			IsShaderParameterTypeReadOnlyRHIResource(BaseType) ||
 			IsRDGResourceReferenceShaderParameterType(BaseType)))
 		{
 			uint16 ArrayStartOffset = Member.GetOffset();
