[shaders]
- strip old job cache path (constructing an input hash based on all inputs); the cache key based on preprocessed source is now the One True Job Cache
- strip parallelfor-based job submission path (gamethread-blocking); the task path has been enabled for a long while now without issues
- remove the "compile job inputs" debug dump; this is no longer relevant since it's based on inputs to the now-stripped compile job path (similar functionality will be provided by the new form of debug usf once completed)

diff --git a/Engine/Source/Developer/ShaderCompilerCommon/Private/ShaderCompilerCommon.cpp b/Engine/Source/Developer/ShaderCompilerCommon/Private/ShaderCompilerCommon.cpp
--- a/Engine/Source/Developer/ShaderCompilerCommon/Private/ShaderCompilerCommon.cpp
+++ b/Engine/Source/Developer/ShaderCompilerCommon/Private/ShaderCompilerCommon.cpp
@@ -1819,7 +1819,7 @@ namespace UE::ShaderCompilerCommon
 		const FShaderCompilerOutput& Output,
 		const FDebugShaderDataOptions& Options)
 	{
-		if (Input.bCachePreprocessed && EnumHasAnyFlags(Input.DebugInfoFlags, EShaderDebugInfoFlags::DetailedSource))
+		if (!Input.Environment.CompilerFlags.Contains(CFLAG_DisableSourceStripping) && EnumHasAnyFlags(Input.DebugInfoFlags, EShaderDebugInfoFlags::DetailedSource))
 		{
 			const TCHAR* StrippedSuffix = TEXT("_Stripped");
 			FFileHelper::SaveStringToFile(GetDebugShaderContents(Input, PreprocessOutput.GetSourceViewWide(), Options, StrippedSuffix), *Options.GetDebugShaderPath(Input, StrippedSuffix));
@@ -1935,10 +1935,9 @@ namespace UE::ShaderCompilerCommon
 
 	FString GetDebugShaderContents(const FShaderCompilerInput& Input, FStringView PreprocessedSource, const FDebugShaderDataOptions& Options, const TCHAR* Suffix)
 	{
-		// If preprocessed cache is enabled, debug dump occurs in the cook process rather than the workers, and
-		// in that case the env in Input.Environment has not been merged with the shared env. Do so here.
+		// Debug dump occurs in the cook process, so we need to merge the env in Input.Environment with the shared env (this is done in the compile step as well)
 		FShaderCompilerEnvironment MergedEnvironment(Input.Environment);
-		if (Input.bCachePreprocessed && IsValidRef(Input.SharedEnvironment))
+		if (IsValidRef(Input.SharedEnvironment))
 		{
 			MergedEnvironment.Merge(*Input.SharedEnvironment);
 		}
diff --git a/Engine/Source/Programs/ShaderCompileWorker/Private/ShaderCompileWorker.cpp b/Engine/Source/Programs/ShaderCompileWorker/Private/ShaderCompileWorker.cpp
--- a/Engine/Source/Programs/ShaderCompileWorker/Private/ShaderCompileWorker.cpp
+++ b/Engine/Source/Programs/ShaderCompileWorker/Private/ShaderCompileWorker.cpp
@@ -456,19 +456,6 @@ private:
 
 		VerifyFormatVersions(ReceivedFormatVersionMap);
 		
-		// Apply shader source directory mappings.
-		{
-			TMap<FString, FString> DirectoryMappings;
-			InputFile << DirectoryMappings;
-
-			ResetAllShaderSourceDirectoryMappings();
-			for (TPair<FString, FString>& MappingEntry : DirectoryMappings)
-			{
-				FPaths::NormalizeDirectoryName(MappingEntry.Value);
-				AddShaderSourceDirectoryMapping(MappingEntry.Key, MappingEntry.Value);
-			}
-		}
-
 		// Initialize shader hash cache before reading any includes.
 		InitializeShaderHashCache();
 
@@ -519,23 +506,6 @@ private:
 			return CharName;
 		};
 
-		// Shared inputs
-		TMap<FString, FThreadSafeSharedAnsiStringPtr> ExternalIncludes;
-		{
-			int32 NumExternalIncludes = 0;
-			InputFile << NumExternalIncludes;
-			ExternalIncludes.Reserve(NumExternalIncludes);
-
-			for (int32 IncludeIndex = 0; IncludeIndex < NumExternalIncludes; IncludeIndex++)
-			{
-				FString NewIncludeName;
-				InputFile << NewIncludeName;
-				TArray<ANSICHAR>* NewIncludeContents = new TArray<ANSICHAR>;
-				InputFile << (*NewIncludeContents);
-				ExternalIncludes.Add(NewIncludeName, MakeShareable(NewIncludeContents));
-			}
-		}
-
 		// Shared environments
 		TArray<FShaderCompilerEnvironment> SharedEnvironments;
 		{
@@ -671,15 +641,7 @@ private:
 				FShaderCompileJob& Job = OutSingleJobs.AddDefaulted_GetRef();
 				// Deserialize the job's inputs.
 				Job.SerializeWorkerInput(InputFile);
-				Job.Input.DeserializeSharedInputs(InputFile, ExternalIncludes, SharedEnvironments, ParameterStructures);
-
-				// SCW doesn't run DDPI, GShaderHasCache Initialize is run  at start with no knowledge of the CustomPlatforms
-				// CustomPlatforms are known when we parse the WorkerInput so we populate the Directory here
-				if (IsCustomPlatform((EShaderPlatform)Job.Input.Target.Platform))
-				{
-					const EShaderPlatform ShaderPlatform = ShaderFormatNameToShaderPlatform(Job.Input.ShaderFormat);
-					UpdateIncludeDirectoryForPreviewPlatform((EShaderPlatform)Job.Input.Target.Platform, ShaderPlatform);
-				}
+				Job.Input.DeserializeSharedInputs(InputFile, SharedEnvironments, ParameterStructures);
 
 				// Process the job.
 				CompileShader(GetShaderFormats(), Job, WorkingDirectory, &GNumProcessedJobs); 
@@ -715,7 +677,7 @@ private:
 					// Deserialize the job's inputs.
 					FShaderCompileJob* Job = PipelineJob.StageJobs[StageIndex]->GetSingleShaderJob();
 					Job->SerializeWorkerInput(InputFile);
-					Job->Input.DeserializeSharedInputs(InputFile, ExternalIncludes, SharedEnvironments, ParameterStructures);
+					Job->Input.DeserializeSharedInputs(InputFile, SharedEnvironments, ParameterStructures);
 
 					// SCW doesn't run DDPI, GShaderHasCache Initialize is run  at start with no knowledge of the CustomPlatforms
 					// CustomPlatforms are known when we parse the WorkerInput so we populate the Directory here
diff --git a/Engine/Source/Runtime/Engine/Private/ShaderCompiler/ShaderCompiler.cpp b/Engine/Source/Runtime/Engine/Private/ShaderCompiler/ShaderCompiler.cpp
--- a/Engine/Source/Runtime/Engine/Private/ShaderCompiler/ShaderCompiler.cpp
+++ b/Engine/Source/Runtime/Engine/Private/ShaderCompiler/ShaderCompiler.cpp
@@ -134,13 +134,6 @@ static FAutoConsoleVariableRef CVarShaderCompilerJobCacheOverflowReducePercent(
 	ECVF_Default
 );
 
-static TAutoConsoleVariable<bool> CVarPreprocessedJobCache(
-	TEXT("r.ShaderCompiler.PreprocessedJobCache"),
-	true,
-	TEXT("If enabled will shader compile jobs will be preprocessed at submission time in the cook process (when the job is queued) and generate job input hashes based on preprocessed source."),
-	ECVF_Default
-);
-
 static TAutoConsoleVariable<bool> CVarJobCacheDDC(
 	TEXT("r.ShaderCompiler.JobCacheDDC"),
 	true,
@@ -326,14 +319,6 @@ static FAutoConsoleVariableRef CVarShaderCompilerDebugDiscardCacheOutputs(
 	ECVF_Default
 );
 
-int32 GShaderCompilerParallelSubmitJobs = 1;
-static FAutoConsoleVariableRef CVarShaderCompilerParallelSubmitJobs(
-	TEXT("r.ShaderCompiler.ParallelSubmitJobs"),
-	GShaderCompilerParallelSubmitJobs,
-	TEXT("if != 0, FShaderJobCache::SubmitJobs will run in multiple parallel tasks, instead of the game thread."),
-	ECVF_Default
-);
-
 int32 GShaderCompilerDebugStallSubmitJob = 0;
 static FAutoConsoleVariableRef CVarShaderCompilerDebugStallSubmitJob(
 	TEXT("r.ShaderCompiler.DebugStallSubmitJob"),
@@ -1785,12 +1770,10 @@ void FShaderJobCache::SubmitJobs(const TArray<FShaderCommonCompileJobPtr>& InJob
 			}
 		}
 
-		if (GShaderCompilerParallelSubmitJobs)
+		for (const FShaderCommonCompileJobPtr& Job : InJobs)
 		{
-			for (FShaderCommonCompileJobPtr Job : InJobs)
-			{
-				UE::Tasks::ETaskPriority Prio = IsRunningCookCommandlet() ? UE::Tasks::ETaskPriority::Normal : UE::Tasks::ETaskPriority::BackgroundNormal;
-				UE::Tasks::Launch(UE_SOURCE_LOCATION, [Job, this]()
+			UE::Tasks::ETaskPriority Prio = IsRunningCookCommandlet() ? UE::Tasks::ETaskPriority::Normal : UE::Tasks::ETaskPriority::BackgroundNormal;
+			UE::Tasks::Launch(UE_SOURCE_LOCATION, [Job, this]()
 				{
 					TRACE_CPUPROFILER_EVENT_SCOPE(ShaderJobTask);
 					double TimeStart = FPlatformTime::Seconds();
@@ -1803,10 +1786,10 @@ void FShaderJobCache::SubmitJobs(const TArray<FShaderCommonCompileJobPtr>& InJob
 					bool bSubmitJob = true;
 					if (ShaderCompiler::IsJobCacheEnabled())
 					{
-						bSubmitJob = ConditionalPreprocessShader(Job);
+						bSubmitJob = PreprocessShader(Job);
 						Job->GetInputHash();
 					}
-					
+
 					if (bSubmitJob)
 					{
 						SubmitJob(Job);
@@ -1815,28 +1798,9 @@ void FShaderJobCache::SubmitJobs(const TArray<FShaderCommonCompileJobPtr>& InJob
 					{
 						ProcessFinishedJob(Job, /* bCompilationSkipped = */true);
 					}
-					
+
 					Job->TimeTaskSubmitJobs = FPlatformTime::Seconds() - TimeStart;
 				}, Prio);
-			}
-		}
-		else
-		{
-			// Precompute the InputHash for each job in multiple-thread.
-			if (ShaderCompiler::IsJobCacheEnabled())
-			{
-				TRACE_CPUPROFILER_EVENT_SCOPE(ShaderCompiler.GetInputHash);
-				ParallelFor(TEXT("ShaderCompiler.GetInputHash.PF"), InJobs.Num(), 1, [&InJobs](int32 Index)
-				{
-					ConditionalPreprocessShader(InJobs[Index]);
-					InJobs[Index]->GetInputHash();
-				}, EParallelForFlags::Unbalanced);
-			}
-
-			for (FShaderCommonCompileJob* Job : InJobs)
-			{
-				SubmitJob(Job);
-			}
 		}
 	}
 }
@@ -2560,32 +2524,25 @@ static int32 GetNumTotalJobs(const TArray<FShaderCommonCompileJobPtr>& Jobs)
 	return NumJobs;
 }
 
-static bool SplitJobsByType(const TArray<FShaderCommonCompileJobPtr>& QueuedJobs, TArray<FShaderCompileJob*>& OutQueuedSingleJobs, TArray<FShaderPipelineCompileJob*>& OutQueuedPipelineJobs)
+static void SplitJobsByType(const TArray<FShaderCommonCompileJobPtr>& QueuedJobs, TArray<FShaderCompileJob*>& OutQueuedSingleJobs, TArray<FShaderPipelineCompileJob*>& OutQueuedPipelineJobs)
 {
-	bool bAnyPreprocessingNeeded = false;
 	for (int32 Index = 0; Index < QueuedJobs.Num(); ++Index)
 	{
 		FShaderCommonCompileJobPtr CommonJob = QueuedJobs[Index];
 		if (FShaderCompileJob* SingleJob = CommonJob->GetSingleShaderJob())
 		{
 			OutQueuedSingleJobs.Add(SingleJob);
-			bAnyPreprocessingNeeded |= !SingleJob->Input.bCachePreprocessed;
 
 		}
 		else if (FShaderPipelineCompileJob* PipelineJob = CommonJob->GetShaderPipelineJob())
 		{
 			OutQueuedPipelineJobs.Add(PipelineJob);
-			PipelineJob->ForEachSingleShaderJob([&bAnyPreprocessingNeeded](const FShaderCompileJob& SingleJob)
-				{
-					bAnyPreprocessingNeeded |= !SingleJob.Input.bCachePreprocessed;
-				});
 		}
 		else
 		{
 			checkf(0, TEXT("FShaderCommonCompileJob::Type=%d is not a valid type for a shader compile job"), (int32)CommonJob->Type);
 		}
 	}
-	return bAnyPreprocessingNeeded;
 }
 
 bool DoWriteTasksInner(const TArray<FShaderCommonCompileJobPtr>& QueuedJobs, FArchive& InTransferFile, IDistributedBuildController* BuildDistributionController, bool bUseRelativePaths, bool bCompressTaskFile)
@@ -2609,51 +2566,17 @@ bool DoWriteTasksInner(const TArray<FShaderCommonCompileJobPtr>& QueuedJobs, FAr
 
 	TArray<FShaderCompileJob*> QueuedSingleJobs;
 	TArray<FShaderPipelineCompileJob*> QueuedPipelineJobs;
-	bool bAnyPreprocessingNeeded = SplitJobsByType(QueuedJobs, QueuedSingleJobs, QueuedPipelineJobs);
-
-	TMap<FString, FString> ShaderSourceDirectoryMappings;
-	// Only serialize source directory mappings if any jobs need to preprocess; these are only used for include handling
-	if (bAnyPreprocessingNeeded)
-	{
-		ShaderSourceDirectoryMappings = AllShaderSourceDirectoryMappings();
-	}
-	
-	// Convert all the source directory paths to absolute, since SCW might be in a different directory to the editor executable
-	for(TPair<FString, FString>& Pair : ShaderSourceDirectoryMappings)
-	{
-		// Remap/enforce relative paths when bUseRelativePaths=true
-		if (bUseRelativePaths && BuildDistributionController != nullptr)
-		{
-			FString SourcePath = FPaths::ConvertRelativePathToFull(Pair.Value);
-			if (!FPaths::IsUnderDirectory(SourcePath, FPaths::RootDir()))
-			{
-				FString DestinationPath = BuildDistributionController->RemapPath(SourcePath);
-				DestinationPath = FPaths::CreateStandardFilename(DestinationPath);
-				Pair.Value = DestinationPath;
-			}
-			else
-			{
-				Pair.Value = FPaths::CreateStandardFilename(Pair.Value);
-			}
-		}
-		else
-		{
-			Pair.Value = FPaths::ConvertRelativePathToFull(Pair.Value);
-		}
-	}
-	TransferFile << ShaderSourceDirectoryMappings;
+	SplitJobsByType(QueuedJobs, QueuedSingleJobs, QueuedPipelineJobs);
 
 	TArray<TRefCountPtr<FSharedShaderCompilerEnvironment>> SharedEnvironments;
 	TArray<const FShaderParametersMetadata*> RequestShaderParameterStructures;
 
-	// Gather External Includes and serialize separately, these are largely shared between jobs
+	// gather shared environments and parameter structures, these tend to be shared between jobs
 	{
-		TMap<FString, TArray<ANSICHAR>> ExternalIncludes;
-		ExternalIncludes.Reserve(32);
 
 		for (int32 JobIndex = 0; JobIndex < QueuedSingleJobs.Num(); JobIndex++)
 		{
-			QueuedSingleJobs[JobIndex]->Input.GatherSharedInputsAnsi(ExternalIncludes, SharedEnvironments, RequestShaderParameterStructures);
+			QueuedSingleJobs[JobIndex]->Input.GatherSharedInputs(SharedEnvironments, RequestShaderParameterStructures);
 		}
 
 		for (int32 JobIndex = 0; JobIndex < QueuedPipelineJobs.Num(); JobIndex++)
@@ -2663,19 +2586,10 @@ bool DoWriteTasksInner(const TArray<FShaderCommonCompileJobPtr>& QueuedJobs, FAr
 
 			for (int32 Index = 0; Index < NumStageJobs; Index++)
 			{
-				PipelineJob->StageJobs[Index]->Input.GatherSharedInputsAnsi(ExternalIncludes, SharedEnvironments, RequestShaderParameterStructures);
+				PipelineJob->StageJobs[Index]->Input.GatherSharedInputs(SharedEnvironments, RequestShaderParameterStructures);
 			}
 		}
 
-		int32 NumExternalIncludes = ExternalIncludes.Num();
-		TransferFile << NumExternalIncludes;
-
-		for (TMap<FString, TArray<ANSICHAR>>::TIterator It(ExternalIncludes); It; ++It)
-		{
-			TransferFile << It.Key();
-			TransferFile << It.Value();
-		}
-
 		int32 NumSharedEnvironments = SharedEnvironments.Num();
 		TransferFile << NumSharedEnvironments;
 
@@ -5327,9 +5241,9 @@ void FShaderCompilerStats::RegisterFinishedJob(FShaderCommonCompileJob& Job, boo
 
 		// Sanity check; compile time should be 0 for cache hits
 		check(!bCompilationSkipped || SingleJob.Output.CompileTime == 0.0f);
-		// Preprocess time should always be non-zero if preprocessed job cache is enabled and preprocessing succeeded;
-		// preprocessing for pipeline stage jobs may be skipped in the case preprocessing a preceding stage of the pipeline failed
-		check(!SingleJob.Input.bCachePreprocessed || !SingleJob.PreprocessOutput.GetSucceeded() || SingleJob.Output.PreprocessTime > 0.0f);
+		// Preprocess time should always be non-zero if preprocessing succeeded; note that preprocessing for pipeline stage jobs may be skipped 
+		// in the case preprocessing a preceding stage of the pipeline failed
+		check(!SingleJob.PreprocessOutput.GetSucceeded() || SingleJob.Output.PreprocessTime > 0.0f);
 
 		const FString ShaderName(SingleJob.Key.ShaderType->GetName());
 		if (FShaderTimings* Existing = ShaderTimings.Find(ShaderName))
@@ -8824,11 +8738,6 @@ void GlobalBeginCompileShader(
 	checkf(Format, TEXT("Shader format %s cannot be found"), *ShaderFormatName.ToString());
 	Format->ModifyShaderCompilerInput(Input);
 
-	if (ShaderCompiler::IsJobCacheEnabled() && CVarPreprocessedJobCache.GetValueOnAnyThread())
-	{
-		Input.bCachePreprocessed = true;
-	}
-
 	// Allow the GBuffer and other shader defines to cause dependend environment changes, but minimizing the #ifdef magic in the shaders, which
 	// is nearly impossible to debug when it goes wrong.
 	FShaderCompileUtilities::ApplyDerivedDefines(Input.Environment, Input.SharedEnvironment, (EShaderPlatform)Target.Platform);
diff --git a/Engine/Source/Runtime/RenderCore/Private/ShaderCore.cpp b/Engine/Source/Runtime/RenderCore/Private/ShaderCore.cpp
--- a/Engine/Source/Runtime/RenderCore/Private/ShaderCore.cpp
+++ b/Engine/Source/Runtime/RenderCore/Private/ShaderCore.cpp
@@ -371,14 +371,6 @@ static TAutoConsoleVariable<int32> CVarForceDebugViewModes(
 	TEXT("2: Forces debug view modes to be unavailable, even on editor builds.  Removes many shader permutations for faster shader iteration."),
 	ECVF_RenderThreadSafe | ECVF_ReadOnly);
 
-int32 GShaderCompilerDumpCompileJobInputs = 0;
-static FAutoConsoleVariableRef CVarShaderCompilerDumpCompileJobInputs(
-	TEXT("r.ShaderCompiler.DumpCompileJobInputs"),
-	GShaderCompilerDumpCompileJobInputs,
-	TEXT("if != 0, unpreprocessed input of the shader compiler jobs will be dumped into the debug directory for closer inspection. This is a debugging feature which is disabled by default."),
-	ECVF_Default
-);
-
 /** Returns true if debug viewmodes are allowed for the current platform. */
 bool AllowDebugViewmodes()
 {	
@@ -1481,9 +1473,7 @@ public:
 	{
 		const double StartPreprocessTime = FPlatformTime::Seconds();
 
-		// if preprocessed cache is disabled, this function is executed as part of the compile process
-		// which will have merged the environment inline prior to calling this, so we can skip merging here 
-		if (Job.Input.bCachePreprocessed && IsValidRef(Job.Input.SharedEnvironment))
+		if (IsValidRef(Job.Input.SharedEnvironment))
 		{
 			// only create new environment & merge if necessary, save some allocs
 			// (need a copy here as we don't want to merge the environment in place like we do in the compile path
@@ -1497,7 +1487,7 @@ public:
 			PreprocessShaderInternal(Backend, Job, Job.Input.Environment);
 		}
 
-		if (Job.PreprocessOutput.bSucceeded && Job.Input.bCachePreprocessed)
+		if (Job.PreprocessOutput.bSucceeded)
 		{
 			if (!Job.Input.Environment.CompilerFlags.Contains(CFLAG_DisableSourceStripping))
 			{
@@ -1632,39 +1622,36 @@ public:
 };
 
 bool ConditionalPreprocessShader(FShaderCommonCompileJob* Job)
+{
+	return PreprocessShader(Job);
+}
+
+bool PreprocessShader(FShaderCommonCompileJob* Job)
 {
 	static ITargetPlatformManagerModule& TargetPlatformManager = GetTargetPlatformManagerRef();
 	if (FShaderCompileJob* SingleJob = Job->GetSingleShaderJob())
 	{
-		if (SingleJob->Input.bCachePreprocessed)
-		{
-			const IShaderFormat* ShaderFormat = TargetPlatformManager.FindShaderFormat(SingleJob->Input.ShaderFormat);
-			return FInternalShaderCompilerFunctions::PreprocessShaderInternal(ShaderFormat, *SingleJob);
-		}
-		return true;
-
+		const IShaderFormat* ShaderFormat = TargetPlatformManager.FindShaderFormat(SingleJob->Input.ShaderFormat);
+		return FInternalShaderCompilerFunctions::PreprocessShaderInternal(ShaderFormat, *SingleJob);
 	}
 	else if (FShaderPipelineCompileJob* PipelineJob = Job->GetShaderPipelineJob())
 	{
 		bool bAnyFailed = false;
 		for (FShaderCompileJob* StageJob : PipelineJob->StageJobs)
 		{
-			if (StageJob->Input.bCachePreprocessed)
-			{
-				const IShaderFormat* ShaderFormat = TargetPlatformManager.FindShaderFormat(StageJob->Input.ShaderFormat);
+			const IShaderFormat* ShaderFormat = TargetPlatformManager.FindShaderFormat(StageJob->Input.ShaderFormat);
 
-				if (!bAnyFailed)
-				{
-					bAnyFailed |= !FInternalShaderCompilerFunctions::PreprocessShaderInternal(ShaderFormat, *StageJob);
-				}
-				else
-				{
-					// skip subsequent stage preprocessing if a prior stage failed to avoid unnecessary work, but log an error to indicate this
-					FString Error = FString::Printf(
-						TEXT("Preprocessing %s stage skipped due to earlier stage preprocessing failure."),
-						GetShaderFrequencyString(StageJob->Input.Target.GetFrequency()));
-					StageJob->Output.Errors.Add(FShaderCompilerError(*Error));
-				}
+			if (!bAnyFailed)
+			{
+				bAnyFailed |= !FInternalShaderCompilerFunctions::PreprocessShaderInternal(ShaderFormat, *StageJob);
+			}
+			else
+			{
+				// skip subsequent stage preprocessing if a prior stage failed to avoid unnecessary work, but log an error to indicate this
+				FString Error = FString::Printf(
+					TEXT("Preprocessing %s stage skipped due to earlier stage preprocessing failure."),
+					GetShaderFrequencyString(StageJob->Input.Target.GetFrequency()));
+				StageJob->Output.Errors.Add(FShaderCompilerError(*Error));
 			}
 		}
 		return !bAnyFailed;
@@ -1711,22 +1698,6 @@ void CompileShader(const TArray<const IShaderFormat*>& ShaderFormats, FShaderCom
 	}
 
 	Job.bSucceeded = Job.Output.bSucceeded;
-	if (Job.Input.DumpDebugInfoEnabled())
-	{
-		// if the preprocessed cache is disabled, dump debug output here, since we don't serialize preprocess output back to the cooker from SCW
-		// (if enabled this will occur in the job OnComplete callback)
-		if (!Job.Input.bCachePreprocessed)
-		{
-			if (Job.SecondaryPreprocessOutput.IsValid() && Job.SecondaryOutput.IsValid())
-			{
-				Compiler->OutputDebugData(Job.Input, Job.PreprocessOutput, *Job.SecondaryPreprocessOutput, Job.Output, *Job.SecondaryOutput);
-			}
-			else
-			{
-				Compiler->OutputDebugData(Job.Input, Job.PreprocessOutput, Job.Output);
-			}
-		}
-	}
 }
 
 void CompileShaderPipeline(const TArray<const IShaderFormat*>& ShaderFormats, FShaderPipelineCompileJob* PipelineJob, const FString& Dir, int32* CompileCount)
@@ -2837,53 +2808,7 @@ void GetShaderIncludes(const TCHAR* EntryPointVirtualFilePath, const TCHAR* Virt
 
 void HashShaderFileWithIncludes(FArchive& HashingArchive, const TCHAR* VirtualFilePath, const FString& FileContents, EShaderPlatform ShaderPlatform, bool bOnlyHashIncludedFiles)
 {
-	auto HashSingleFile = [](FArchive& HashingArchive, const TCHAR* VirtualFilePath, EShaderPlatform ShaderPlatform, const FString& FileContents)
-	{
-		// first, a "soft" check
-		bool bFoundInCache = false;
-		{
-			FRWScopeLock ShaderHashAccessLock(GShaderHashAccessRWLock, SLT_ReadOnly);
-			const FSHAHash* CachedHash = GShaderHashCache.FindHash(ShaderPlatform, VirtualFilePath);
-			// If a hash for this filename has been cached, use that
-			if (CachedHash)
-			{
-				bFoundInCache = true;
-				HashingArchive << const_cast<FSHAHash&>(*CachedHash);
-			}
-		}
-
-		// outside of the lock scope because we don't need the lock and hashing can take time
-		if (!bFoundInCache)
-		{			
-			// if the file isn't generated, add it to the cache now
-			bool bGenerated = FCString::Strstr(VirtualFilePath, TEXT("Generated")) != nullptr;
-			if (!bGenerated)
-			{
-				// this function fails hard if it cannot load
-				const FSHAHash& FileHash = GetShaderFileHash(VirtualFilePath, ShaderPlatform);
-				HashingArchive << const_cast<FSHAHash&>(FileHash);
-			}
-			else
-			{
-				// note, it is legal for some generated files to have empty contents, so hash both the name and their contents
-				HashingArchive.Serialize(reinterpret_cast<void*>(const_cast<TCHAR*>(VirtualFilePath)), FCString::Strlen(VirtualFilePath));
-				HashingArchive << const_cast<FString&>(FileContents);
-			}
-		}
-	};
-
-	// First, always hash the file itself
-	HashSingleFile(HashingArchive, VirtualFilePath, ShaderPlatform, FileContents);
-
-	// Get the list of includes this file contains
-	TArray<FString> IncludeVirtualFilePaths;
-	GetShaderIncludes(VirtualFilePath, VirtualFilePath, FileContents, IncludeVirtualFilePaths, ShaderPlatform);
-
-	for (int32 IncludeIndex = 0; IncludeIndex < IncludeVirtualFilePaths.Num(); IncludeIndex++)
-	{
-		// Here, we assume that all includes can be found in cache. This also means that generated files won't include other generated files.
-		HashSingleFile(HashingArchive, *IncludeVirtualFilePaths[IncludeIndex], ShaderPlatform, FString());
-	}
+	// deprecated
 }
 
 static bool TryUpdateSingleShaderFilehash(FSHA1& InOutHashState, const TCHAR* VirtualFilePath,
@@ -3766,7 +3691,6 @@ FArchive& operator<<(FArchive& Ar, FShaderCompilerInput& Input)
 	Ar << Input.ShaderName;
 	Ar << Input.bCompilingForShaderPipeline;
 	Ar << Input.bIncludeUsedOutputs;
-	Ar << Input.bCachePreprocessed;
 	Ar << Input.UsedOutputs;
 	Ar << Input.DumpDebugInfoRootPath;
 	Ar << Input.DumpDebugInfoPath;
@@ -3775,14 +3699,7 @@ FArchive& operator<<(FArchive& Ar, FShaderCompilerInput& Input)
 	Ar << Input.DebugGroupName;
 	Ar << Input.DebugDescription;
 	Ar << Input.Hash;
-	if (Input.bCachePreprocessed)
-	{
-		Input.Environment.SerializeCompilationDependencies(Ar);
-	}
-	else
-	{
-		Ar << Input.Environment;
-	}
+	Input.Environment.SerializeCompilationDependencies(Ar);
 	Ar << Input.ExtraSettings;
 	Ar << reinterpret_cast<uint8&>(Input.OodleCompressor);
 	Ar << reinterpret_cast<uint8&>(Input.OodleLevel);
@@ -3859,203 +3776,51 @@ FShaderCompilerInputHash FShaderCompileJob::GetInputHash()
 		return InputHash;
 	}
 
-	if (Input.bCachePreprocessed)
-	{
-		FMemoryHasherBlake3 Hasher;
+	FMemoryHasherBlake3 Hasher;
 
-		int32 FShaderCompilerOutputStructVersionLocal = FShaderCompilerOutputStructVersion;
-		Hasher << FShaderCompilerOutputStructVersionLocal;
+	int32 FShaderCompilerOutputStructVersionLocal = FShaderCompilerOutputStructVersion;
+	Hasher << FShaderCompilerOutputStructVersionLocal;
 
-		uint32 FormatVersion = GetTargetPlatformManagerRef().ShaderFormatVersion(Input.ShaderFormat);
-		Hasher << FormatVersion;
-		
-		FShaderTarget Target = Input.Target;
-		Hasher << Target;
-		Hasher << Input.EntryPointName;
-
-		// Include this flag, so shader pipeline jobs get a different hash from single shader jobs, even if the preprocessed shader is otherwise the same.
-		Hasher << Input.bCompilingForShaderPipeline;
-
-		FShaderCompilerEnvironment MergedEnvironment = Input.Environment;
-		if (Input.SharedEnvironment)
-		{
-			MergedEnvironment.Merge(*Input.SharedEnvironment);
-		}
-		MergedEnvironment.SerializeCompilationDependencies(Hasher);
-		
-		auto HashDirectives = [&Hasher](const FString* Directive)
-		{
-			check(Directive && !Directive->IsEmpty());
-			// const_cast due to serialization API requiring non-const. better than not having const correctness in the API.
-			Hasher << const_cast<FString&>(*Directive);
-		};
-		// Hash all UESHADERMETADATA_ directives encountered during preprocessing (assume these may be used to modify compilation behaviour)
-		PreprocessOutput.VisitDirectives(HashDirectives);
-
-		// const_cast due to serialization API requiring non-const. better than not having const correctness in the API.
-		Hasher << PreprocessOutput.EditSource();
-		if (SecondaryPreprocessOutput.IsValid())
-		{
-			Hasher << SecondaryPreprocessOutput->EditSource();
-		}
+	uint32 FormatVersion = GetTargetPlatformManagerRef().ShaderFormatVersion(Input.ShaderFormat);
+	Hasher << FormatVersion;
+	
+	FShaderTarget Target = Input.Target;
+	Hasher << Target;
+	Hasher << Input.EntryPointName;
 
-		if (Input.RootParametersStructure)
-		{
-			FBlake3Hash LayoutSignature = Input.RootParametersStructure->GetLayoutSignature();
-			Hasher << LayoutSignature;
-		}
+	// Include this flag, so shader pipeline jobs get a different hash from single shader jobs, even if the preprocessed shader is otherwise the same.
+	Hasher << Input.bCompilingForShaderPipeline;
 
-		InputHash = Hasher.Finalize();
+	FShaderCompilerEnvironment MergedEnvironment = Input.Environment;
+	if (Input.SharedEnvironment)
+	{
+		MergedEnvironment.Merge(*Input.SharedEnvironment);
 	}
-	else
+	MergedEnvironment.SerializeCompilationDependencies(Hasher);
+	
+	auto HashDirectives = [&Hasher](const FString* Directive)
 	{
-		auto SerializeInputs = [this](FArchive& Archive)
-		{
-			checkf(Archive.IsSaving() && !Archive.IsLoading(), TEXT("A loading archive is passed to FShaderCompileJob::GetInputHash(), this is not supported as it may corrupt its data"));
-
-			int32 FShaderCompilerOutputStructVersionLocal = FShaderCompilerOutputStructVersion;
-			Archive << FShaderCompilerOutputStructVersionLocal;
-
-			uint32 FormatVersion = GetTargetPlatformManagerRef().ShaderFormatVersion(Input.ShaderFormat);
-			Archive << FormatVersion;
-
-			// Don't include debug group name in the hashing; this drastically worsens our cache hit rate
-			FString DebugGroupNameTmp(MoveTemp(Input.DebugGroupName));
-			Archive << Input;
-			Input.DebugGroupName = MoveTemp(DebugGroupNameTmp);
-			Input.Environment.SerializeEverythingButFiles(Archive);
-
-			if (Input.RootParametersStructure)
-			{
-				FBlake3Hash LayoutSignature = Input.RootParametersStructure->GetLayoutSignature();
-				Archive << LayoutSignature;
-			}
-
-			// hash the source file so changes to files during the development are picked up
-			const FSHAHash& SourceHash = GetShaderFileHash(*Input.VirtualSourceFilePath, Input.Target.GetPlatform());
-			Archive << const_cast<FSHAHash&>(SourceHash);
-
-			// unroll the included files for the parallel processing.
-			// These are temporary arrays that only exist for the ParallelFor
-			TArray<const TCHAR*> IncludeVirtualPaths;
-			TArray<FShaderVirtualFileContents> Contents;
-			TArray<bool> OnlyHashIncludes;
-			TArray<FBlake3Hash> Hashes;
-
-			// while the contents of this is already hashed (included in Environment's operator<<()), we still need to account for includes in the generated files and hash them, too
-			for (TMap<FString, FString>::TConstIterator It(Input.Environment.IncludeVirtualPathToContentsMap); It; ++It)
-			{
-				const FString& VirtualPath = It.Key();
-				IncludeVirtualPaths.Add(*VirtualPath);
-				Contents.Add(&It.Value());
-				OnlyHashIncludes.Add(true);	// not hashing contents of the file itself, as it was included in Environment's operator<<()
-				Hashes.AddDefaulted();
-			}
-
-			for (TMap<FString, FThreadSafeSharedAnsiStringPtr>::TConstIterator It(Input.Environment.IncludeVirtualPathToSharedContentsMap); It; ++It)
-			{
-				const FString& VirtualPath = It.Key();
-				IncludeVirtualPaths.Add(*VirtualPath);
-				check(It.Value());
-				Contents.Add(&(*It.Value()));
-				OnlyHashIncludes.Add(false);
-				Hashes.AddDefaulted();
-			}
-
-			if (Input.SharedEnvironment)
-			{
-				Input.SharedEnvironment->SerializeEverythingButFiles(Archive);
-
-				for (TMap<FString, FString>::TConstIterator It(Input.SharedEnvironment->IncludeVirtualPathToContentsMap); It; ++It)
-				{
-					const FString& VirtualPath = It.Key();
-					IncludeVirtualPaths.Add(*VirtualPath);
-					Contents.Add(&It.Value());
-					OnlyHashIncludes.Add(true);	// not hashing contents of the file itself, as it was included in Environment's operator<<()
-					Hashes.AddDefaulted();
-				}
-
-				for (TMap<FString, FThreadSafeSharedAnsiStringPtr>::TConstIterator It(Input.SharedEnvironment->IncludeVirtualPathToSharedContentsMap); It; ++It)
-				{
-					const FString& VirtualPath = It.Key();
-					IncludeVirtualPaths.Add(*VirtualPath);
-					check(It.Value());
-					Contents.Add(&(*It.Value()));
-					OnlyHashIncludes.Add(false);
-					Hashes.AddDefaulted();
-				}
-			}
-
-			check(IncludeVirtualPaths.Num() == Contents.Num());
-			check(Contents.Num() == OnlyHashIncludes.Num());
-			check(OnlyHashIncludes.Num() == Hashes.Num());
-
-			EShaderPlatform Platform = Input.Target.GetPlatform();
-			ParallelFor(Contents.Num(), [&IncludeVirtualPaths, &Contents, &OnlyHashIncludes, &Hashes, &Platform](int32 FileIndex)
-				{
-					FMemoryHasherBlake3 MemHasher;
-					if (Contents[FileIndex].Wide)
-					{
-						HashShaderFileWithIncludes(MemHasher, IncludeVirtualPaths[FileIndex], *Contents[FileIndex].Wide, Platform, OnlyHashIncludes[FileIndex]);
-					}
-					else
-					{
-						// ANSI files are shared uniform buffer struct declarations (or generated stereo code), and never have includes, so we just need to hash the
-						// single file contents.  Make sure that assumption hasn't been violated (this test costs less than 0.1% of GetInputHash, so might as well).
-						check(FCStringAnsi::Strstr(Contents[FileIndex].Ansi->GetData(), "#include") == nullptr);
-
-						MemHasher.Serialize(reinterpret_cast<void*>(const_cast<TCHAR*>(IncludeVirtualPaths[FileIndex])), FCString::Strlen(IncludeVirtualPaths[FileIndex]));
-						MemHasher << const_cast<TArray<ANSICHAR>&>(*Contents[FileIndex].Ansi);
-					}
-					Hashes[FileIndex] = MemHasher.Finalize();
-				},
-				EParallelForFlags::Unbalanced
-					);
-
-			// include the hashes in the main hash (consider sorting them if includes are found to have a random order)
-			for (int32 HashIndex = 0, NumHashes = Hashes.Num(); HashIndex < NumHashes; ++HashIndex)
-			{
-				Archive << Hashes[HashIndex];
-			}
-		};
-
-		// use faster hasher that doesn't allocate memory
-		FMemoryHasherBlake3 MemHasher;
-		SerializeInputs(MemHasher);
-		InputHash = MemHasher.Finalize();
-
-		if (GShaderCompilerDumpCompileJobInputs)
-		{
-			TArray<uint8> MemoryBlob;
-			FMemoryWriter MemWriter(MemoryBlob);
-
-			SerializeInputs(MemWriter);
-
-			FString IntermediateFormatPath = FPaths::ProjectSavedDir() / TEXT("ShaderJobInputs");
-#if UE_BUILD_DEBUG
-			FString TempPath = IntermediateFormatPath / TEXT("DebugEditor");
-#else
-			FString TempPath = IntermediateFormatPath / TEXT("DevelopmentEditor");
-#endif
-			IFileManager::Get().MakeDirectory(*TempPath, true);
-
-			static int32 InputHashID = 0;
-			FString FileName = Input.DebugGroupName.Replace(TEXT("/"), TEXT("_")).Replace(TEXT("<"), TEXT("_")).Replace(TEXT(">"), TEXT("_")).Replace(TEXT(":"), TEXT("_")).Replace(TEXT("|"), TEXT("_"))
-				+ TEXT("-") + Input.EntryPointName;
-			FString TempFile = TempPath / FString::Printf(TEXT("%s-%d.bin"), *FileName, InputHashID++);
+		check(Directive && !Directive->IsEmpty());
+		// const_cast due to serialization API requiring non-const. better than not having const correctness in the API.
+		Hasher << const_cast<FString&>(*Directive);
+	};
+	// Hash all UESHADERMETADATA_ directives encountered during preprocessing (assume these may be used to modify compilation behaviour)
+	PreprocessOutput.VisitDirectives(HashDirectives);
 
-			TUniquePtr<FArchive> DumpAr(IFileManager::Get().CreateFileWriter(*TempFile));
-			DumpAr->Serialize(MemoryBlob.GetData(), MemoryBlob.Num());
+	Hasher << PreprocessOutput.EditSource();
+	if (SecondaryPreprocessOutput.IsValid())
+	{
+		Hasher << SecondaryPreprocessOutput->EditSource();
+	}
 
-			// as an additional debugging feature, make sure that the hash is the same as calculated by the memhasher
-			FBlake3Hash Check = FBlake3::HashBuffer(MemoryBlob.GetData(), MemoryBlob.Num());
-			if (Check != InputHash)
-			{
-				UE_LOG(LogShaders, Error, TEXT("Job input hash disagrees between FMemoryHasherSHA1 (%s) and FMemoryWriter + FSHA1 (%s, which was dumped to disk)"), *LexToString(InputHash), *LexToString(Check));
-			}
-		}
+	if (Input.RootParametersStructure)
+	{
+		FBlake3Hash LayoutSignature = Input.RootParametersStructure->GetLayoutSignature();
+		Hasher << LayoutSignature;
 	}
 
+	InputHash = Hasher.Finalize();
+
 	bInputHashSet = true;
 	Input.Hash = InputHash;
 	return InputHash;
@@ -4104,19 +3869,23 @@ void FShaderCompileJob::OnComplete()
 		PreprocessOutput.DecompressCode();
 	}
 
-	// For jobs using the preprocessed cache, we need to remap error messages whether or not the job was actually the one that ran
-	// the compilation step. In addition since we always run preprocessing we set the total preprocess time accordingly.
-	if (Input.bCachePreprocessed)
+	// For jobs which applied source stripping, we need to remap error messages whether or not the job was actually the one that ran
+	// the compilation step.
+	if (!Input.Environment.CompilerFlags.Contains(CFLAG_DisableSourceStripping))
 	{
 		PreprocessOutput.RemapErrors(Output);
-		Output.PreprocessTime = PreprocessOutput.ElapsedTime;
 	}
 
-	// dump debug info for the job at this point if the preprocessed cache is enabled
-	// this ensures we get debug output for all jobs, including those that were found in the job cache,
-	// or matched another in-flight job's hash and so could share its results
-	if (Input.bCachePreprocessed 
-		&& Input.DumpDebugInfoEnabled()
+	// always set the preprocess time on the output struct, regardless of whether this job hit the cache or actually compiled, since
+	// we always run preprocessing as a first step when the job was submitted.
+	Output.PreprocessTime = PreprocessOutput.ElapsedTime;
+
+	// output debug info for the job if either the job was the one which actually executed the compile, or debug info is requested 
+	// for all jobs including ones which hit the job cache (or matched and was deduplicated with another in-flight job)
+	// note that depending on shaderformat implementation this may not necessarily be _all_ debug artifacts, since some require
+	// running compilation to generate (these will always be output by the compile step, and always only be generated for the single
+	// job which executed compilation)
+	if (Input.DumpDebugInfoEnabled()
 		// if we only want debug info for jobs which actually compiled, check the CompileTime
 		// (jobs deserialized from the cache/wait list/ddc will have a compiletime of 0.0)
 		&& (CVarDumpDebugInfoForCacheHits.GetValueOnAnyThread() || Output.CompileTime > 0.0f))
@@ -4141,10 +3910,9 @@ void FShaderCompileJob::AppendDebugName(FStringBuilderBase& OutName) const
 void FShaderCompileJob::SerializeWorkerOutput(FArchive& Ar)
 {
 	// Only serialize the modified source/entry point if either:
-	// (a) this job is using the preprocessed cache and requires debug output (since this is done in the cooker process the modified source is needed outside of SCW)
+	// (a) this job requires debug output (since this is done in the cooker process the modified source is needed outside of SCW)
 	// or (b) the "extract shader source" setting is enabled (i.e. something upstream explicitly wants the final source passed to the compiler).
-	Output.bSerializeModifiedSource = (Input.bCachePreprocessed && Input.DumpDebugInfoEnabled()) || Input.ExtraSettings.bExtractShaderSource;
-
+	Output.bSerializeModifiedSource = Input.DumpDebugInfoEnabled() || Input.ExtraSettings.bExtractShaderSource;
 
 	Ar << Output;
 
@@ -4160,14 +3928,6 @@ void FShaderCompileJob::SerializeWorkerOutput(FArchive& Ar)
 		Ar << *SecondaryOutput;
 	}
 
-	// edge case for backends which have implemented independent preprocessing API when the preprocessed cache is not enabled.
-	// if no modifications have occurred as part of the compile step, we still need a copy of the source back in the cooker
-	// if bExtractShaderSource is set, so explicitly serialize just that portion of the preprocess output struct here.
-	if (Input.ExtraSettings.bExtractShaderSource && !Input.bCachePreprocessed && Output.ModifiedShaderSource.IsEmpty())
-	{
-		Ar << PreprocessOutput.EditSource();
-	}
-
 	bool bSucceededTemp = (bool)bSucceeded;
 	Ar << bSucceededTemp;
 	bSucceeded = bSucceededTemp;
@@ -4177,21 +3937,18 @@ void FShaderCompileJob::SerializeWorkerInput(FArchive& Ar)
 {
 	Ar << Input;
 
-	if (Input.bCachePreprocessed)
-	{
-		Ar << PreprocessOutput;
+	Ar << PreprocessOutput;
 
-		bool bSecondaryPreprocessOutput = SecondaryPreprocessOutput.IsValid();
-		Ar << bSecondaryPreprocessOutput;
+	bool bSecondaryPreprocessOutput = SecondaryPreprocessOutput.IsValid();
+	Ar << bSecondaryPreprocessOutput;
 
-		if (bSecondaryPreprocessOutput)
+	if (bSecondaryPreprocessOutput)
+	{
+		if (Ar.IsLoading())
 		{
-			if (Ar.IsLoading())
-			{
-				SecondaryPreprocessOutput = MakeUnique<FShaderPreprocessOutput>();
-			}
-			Ar << *SecondaryPreprocessOutput;
+			SecondaryPreprocessOutput = MakeUnique<FShaderPreprocessOutput>();
 		}
+		Ar << *SecondaryPreprocessOutput;
 	}
 }
 
diff --git a/Engine/Source/Runtime/RenderCore/Public/ShaderCompilerCore.h b/Engine/Source/Runtime/RenderCore/Public/ShaderCompilerCore.h
--- a/Engine/Source/Runtime/RenderCore/Public/ShaderCompilerCore.h
+++ b/Engine/Source/Runtime/RenderCore/Public/ShaderCompilerCore.h
@@ -31,7 +31,7 @@ class FShaderPipelineCompileJob;
 typedef TSharedPtr<TArray<ANSICHAR>, ESPMode::ThreadSafe> FShaderSharedAnsiStringPtr;
 
 // this is for the protocol, not the data, bump if FShaderCompilerInput/FShaderPreprocessOutput serialization, SerializeWorkerInput or ProcessInputFromArchive changes.
-inline const int32 ShaderCompileWorkerInputVersion = 27;
+inline const int32 ShaderCompileWorkerInputVersion = 28;
 // this is for the protocol, not the data, bump if FShaderCompilerOutput or WriteToOutputArchive changes.
 inline const int32 ShaderCompileWorkerOutputVersion = 20;
 // this is for the protocol, not the data.
@@ -260,6 +260,7 @@ struct FShaderCompilerInput
 	// True if the backend for this job implements the independent preprocessing API.
 	bool bIndependentPreprocessed = true;
 
+	UE_DEPRECATED(5.5, "bCachePreprocessed member no longer used; preprocessed job cache is now the only path for in-memory caching of shader jobs")
 	// True if the cache key for this job should be based on preprocessed source. If so,
 	// preprocessing will be executed in the cook process independent of compilation (and
 	// as such this will only ever be set for jobs whose shader format supports independent
@@ -364,59 +365,24 @@ struct FShaderCompilerInput
 		return DebugGroupName;
 	}
 
+	UE_DEPRECATED(5.5, "GatherSharedInputs no longer processes external includes, as these are not needed by SCW anymore")
 	void GatherSharedInputs(
 		TMap<FString, FString>& ExternalIncludes,
 		TArray<TRefCountPtr<FSharedShaderCompilerEnvironment>>& SharedEnvironments,
-		TArray<const FShaderParametersMetadata*>& ParametersStructures)
-	{
-		check(!SharedEnvironment || SharedEnvironment->IncludeVirtualPathToSharedContentsMap.Num() == 0);
-
-		// If the input is already preprocessed we don't need to serialize includes when writing worker input files
-		if (!bCachePreprocessed)
-		{
-			for (const auto& It : Environment.IncludeVirtualPathToSharedContentsMap)
-			{
-				FString* FoundEntry = ExternalIncludes.Find(It.Key);
-
-				if (!FoundEntry)
-				{
-					ExternalIncludes.Add(It.Key, FString(*It.Value));
-				}
-			}
-		}
-
-		if (SharedEnvironment)
-		{
-			SharedEnvironments.AddUnique(SharedEnvironment);
-		}
-
-		if (RootParametersStructure)
-		{
-			ParametersStructures.AddUnique(RootParametersStructure);
-		}
-	}
+		TArray<const FShaderParametersMetadata*>& ParametersStructures) {}
 
+	UE_DEPRECATED(5.5, "GatherSharedInputsAnsi is no longer needed since external includes are no longer serialized; use GetSharedInputs")
 	void GatherSharedInputsAnsi(
 		TMap<FString, TArray<ANSICHAR>>& ExternalIncludes,
+		TArray<TRefCountPtr<FSharedShaderCompilerEnvironment>>& SharedEnvironments,
+		TArray<const FShaderParametersMetadata*>& ParametersStructures) {}
+
+	void GatherSharedInputs(
 		TArray<TRefCountPtr<FSharedShaderCompilerEnvironment>>& SharedEnvironments,
 		TArray<const FShaderParametersMetadata*>& ParametersStructures)
 	{
 		check(!SharedEnvironment || SharedEnvironment->IncludeVirtualPathToSharedContentsMap.Num() == 0);
 
-		// If the input is already preprocessed we don't need to serialize includes when writing worker input files
-		if (!bCachePreprocessed)
-		{
-			for (const auto& It : Environment.IncludeVirtualPathToSharedContentsMap)
-			{
-				TArray<ANSICHAR>* FoundEntry = ExternalIncludes.Find(It.Key);
-
-				if (!FoundEntry)
-				{
-					ExternalIncludes.Add(It.Key, *It.Value);
-				}
-			}
-		}
-
 		if (SharedEnvironment)
 		{
 			SharedEnvironments.AddUnique(SharedEnvironment);
@@ -432,19 +398,6 @@ struct FShaderCompilerInput
 	{
 		check(Ar.IsSaving());
 
-		if (!bCachePreprocessed)
-		{
-			TArray<FString> ReferencedExternalIncludes;
-			ReferencedExternalIncludes.Empty(Environment.IncludeVirtualPathToSharedContentsMap.Num());
-
-			for (const auto& It : Environment.IncludeVirtualPathToSharedContentsMap)
-			{
-				ReferencedExternalIncludes.Add(It.Key);
-			}
-
-			Ar << ReferencedExternalIncludes;
-		}
-
 		int32 SharedEnvironmentIndex = SharedEnvironments.Find(SharedEnvironment);
 		Ar << SharedEnvironmentIndex;
 
@@ -459,25 +412,11 @@ struct FShaderCompilerInput
 
 	void DeserializeSharedInputs(
 		FArchive& Ar,
-		const TMap<FString, FThreadSafeSharedAnsiStringPtr>& ExternalIncludes,
 		const TArray<FShaderCompilerEnvironment>& SharedEnvironments,
 		const TArray<TUniquePtr<FShaderParametersMetadata>>& ShaderParameterStructures)
 	{
 		check(Ar.IsLoading());
 
-		if (!bCachePreprocessed)
-		{
-			TArray<FString> ReferencedExternalIncludes;
-			Ar << ReferencedExternalIncludes;
-
-			Environment.IncludeVirtualPathToSharedContentsMap.Reserve(ReferencedExternalIncludes.Num());
-
-			for (int32 i = 0; i < ReferencedExternalIncludes.Num(); i++)
-			{
-				Environment.IncludeVirtualPathToSharedContentsMap.Add(ReferencedExternalIncludes[i], ExternalIncludes.FindChecked(ReferencedExternalIncludes[i]));
-			}
-		}
-
 		int32 SharedEnvironmentIndex = 0;
 		Ar << SharedEnvironmentIndex;
 
@@ -740,7 +679,9 @@ extern RENDERCORE_API int HandleShaderCompileException(Windows::LPEXCEPTION_POIN
 extern RENDERCORE_API const IShaderFormat* FindShaderFormat(FName Format, const TArray<const IShaderFormat*>& ShaderFormats);
 
 // Executes preprocessing for the given job, if the job is marked to be preprocessed independently prior to compilation.
+UE_DEPRECATED(5.5, "ConditionalPreprocessShader is now just PreprocessShader (no longer conditional, always executes in the main process at job submission time)")
 extern RENDERCORE_API bool ConditionalPreprocessShader(FShaderCommonCompileJob* Job);
+extern RENDERCORE_API bool PreprocessShader(FShaderCommonCompileJob* Job);
 extern RENDERCORE_API void CompileShader(const TArray<const IShaderFormat*>& ShaderFormats, FShaderCompileJob& Job, const FString& WorkingDirectory, int32* CompileCount = nullptr);
 extern RENDERCORE_API void CompileShaderPipeline(const TArray<const IShaderFormat*>& ShaderFormats, FShaderPipelineCompileJob* PipelineJob, const FString& WorkingDirectory, int32* CompileCount = nullptr);
 
diff --git a/Engine/Source/Runtime/RenderCore/Public/ShaderCore.h b/Engine/Source/Runtime/RenderCore/Public/ShaderCore.h
--- a/Engine/Source/Runtime/RenderCore/Public/ShaderCore.h
+++ b/Engine/Source/Runtime/RenderCore/Public/ShaderCore.h
@@ -1183,14 +1183,7 @@ extern RENDERCORE_API const class FSHAHash& GetShaderFileHash(const TCHAR* Virtu
 extern RENDERCORE_API const FSHAHash* TryGetShaderFileHash(const TCHAR* VirtualFilePath,
 	EShaderPlatform ShaderPlatform, FString* OutErrorMessage = nullptr);
 
-/**
- * Calculates a hash for the given source file and all files included from it.
- * @param HashingArchive - hash to update
- * @param VirtualFilePath - name of this source code path (won't be loaded, as it is expected to be generated)
- * @param FileContents - shader source code to Hash (included files will be hashed, too)
- * @param ShaderPlatform - shader platform to Hash
- * @param bOnlyHashIncludedFiles - skip hashing contents of the file itself (useful if it was already hashed outside of this function)
- */
+UE_DEPRECATED(5.5, "HashShaderFileWithIncludes is no longer used by the shader compilation system and so no longer supported.")
 extern RENDERCORE_API void HashShaderFileWithIncludes(FArchive& HashingArchive, const TCHAR* VirtualFilePath, const FString& FileContents, EShaderPlatform ShaderPlatform, bool bOnlyHashIncludedFiles);
 
 /**
