Fix CIS compilation issue

diff --git a/Engine/Source/Runtime/RenderCore/Private/Shader.cpp b/Engine/Source/Runtime/RenderCore/Private/Shader.cpp
--- a/Engine/Source/Runtime/RenderCore/Private/Shader.cpp
+++ b/Engine/Source/Runtime/RenderCore/Private/Shader.cpp
@@ -2009,7 +2009,10 @@ void ShaderMapAppendKey(EShaderPlatform Platform, FShaderKeyGenerator& KeyGen)
 
 	if (SupportedHardwareMask != 0)
 	{
-		KeyString += FString::Printf(TEXT("_SHM_%X"), SupportedHardwareMask);
+		KeyGen.AppendSeparator();
+		KeyGen.Append(TEXT("SHM"));
+		KeyGen.AppendSeparator();
+		KeyGen.AppendHex(SupportedHardwareMask);
 	}
 
 	// Encode the Metal standard into the shader compile options so that they recompile if the settings change.
