Avoid infinite loop when Current==0 in FShaderPreprocessOutput::StripCode (happens when using directcompile)

#rnx

diff --git a/Engine/Source/Runtime/RenderCore/Private/ShaderPreprocessTypes.cpp b/Engine/Source/Runtime/RenderCore/Private/ShaderPreprocessTypes.cpp
--- a/Engine/Source/Runtime/RenderCore/Private/ShaderPreprocessTypes.cpp
+++ b/Engine/Source/Runtime/RenderCore/Private/ShaderPreprocessTypes.cpp
@@ -326,6 +326,10 @@ void FShaderPreprocessOutput::StripCode(bool bCopyOriginalPreprocessdSource)
 				*OutStripped++ = *Current++;
 			}
 		}
+		else if (*Current == 0)
+		{
+			*OutStripped++ = *Current++;
+		}
 	}
 	check(OutStripped <= OutStrippedData + PreprocessedSourceStripped.Len());
 	// ShrinkToLen null terminates for us by virtue of adding zero'd SIMD padding
