Allow wrapping of group counts for replacement buffer clear shader.
Fixes RHI validation error when clearing buffers that spill the 65535 per dimension limit for group counts.
The validation error could be seen in the nanite engine tests.

diff --git a/Engine/Shaders/Private/ClearReplacementShaders.usf b/Engine/Shaders/Private/ClearReplacementShaders.usf
--- a/Engine/Shaders/Private/ClearReplacementShaders.usf
+++ b/Engine/Shaders/Private/ClearReplacementShaders.usf
@@ -5,6 +5,7 @@
 =============================================================================*/
 
 #include "Common.ush"
+#include "ComputeShaderUtils.ush"
 
 #ifndef THREADGROUPSIZE_X
 	#define THREADGROUPSIZE_X 1
@@ -12,6 +13,10 @@
 	#define THREADGROUPSIZE_Z 1
 #endif
 
+#ifndef USE_WRAPPED_GROUP_COUNT
+	#define USE_WRAPPED_GROUP_COUNT 0
+#endif
+
 #if ENABLE_CLEAR_VALUE
 	VALUE_TYPE ClearValue;
 #else
@@ -64,8 +69,12 @@
 	#endif
 
 	[numthreads(THREADGROUPSIZE_X, THREADGROUPSIZE_Y, THREADGROUPSIZE_Z)]
-	void ClearCS(uint3 Position : SV_DispatchThreadID)
+	void ClearCS(uint3 GroupId : SV_GroupID, uint GroupThreadIndex : SV_GroupIndex, uint3 Position : SV_DispatchThreadID)
 	{
+	#if USE_WRAPPED_GROUP_COUNT
+		Position = uint3(GetUnWrappedDispatchThreadId(GroupId, GroupThreadIndex, THREADGROUPSIZE_X), 0, 0);
+	#endif
+
 	#if ENABLE_BOUNDS
 		ADDR(Position) += ADDR(MinBounds);
 		if (COMP(Position, MaxBounds))
diff --git a/Engine/Source/Runtime/RenderCore/Public/ClearReplacementShaders.h b/Engine/Source/Runtime/RenderCore/Public/ClearReplacementShaders.h
--- a/Engine/Source/Runtime/RenderCore/Public/ClearReplacementShaders.h
+++ b/Engine/Source/Runtime/RenderCore/Public/ClearReplacementShaders.h
@@ -10,6 +10,7 @@
 #include "Math/UnrealMathSSE.h"
 #include "PipelineStateCache.h"
 #include "PixelFormat.h"
+#include "RenderGraphUtils.h"
 #include "RHI.h"
 #include "RHICommandList.h"
 #include "RHIDefinitions.h"
@@ -52,6 +53,9 @@ struct TClearReplacementBase : public FGlobalShader
 	static_assert(NumChannels >= 1 && NumChannels <= 4, "Only 1 to 4 channels are supported.");
 	DECLARE_INLINE_TYPE_LAYOUT(TClearReplacementBase, NonVirtual);
 
+	// Wrapping group counts will require bounds checking to account for additional groups from the wrapping stride.
+	static constexpr bool bSupportsWrappedGroupCount = bEnableBounds;
+
 protected:
 	TClearReplacementBase() {}
 	TClearReplacementBase(const ShaderMetaType::CompiledShaderInitializerType& Initializer)
@@ -143,6 +147,9 @@ public:
 	static constexpr uint32 ThreadGroupSizeX = ClearReplacementCS::TThreadGroupSize<ResourceType>::X;
 	static constexpr uint32 ThreadGroupSizeY = ClearReplacementCS::TThreadGroupSize<ResourceType>::Y;
 	static constexpr uint32 ThreadGroupSizeZ = ClearReplacementCS::TThreadGroupSize<ResourceType>::Z;
+
+	// Use wrapping for 1D layouts if supported.
+	static constexpr bool bUseWrappedGroupCount = (ThreadGroupSizeY == 1) && (ThreadGroupSizeZ == 1) && BaseType::bSupportsWrappedGroupCount;
 	
 	TClearReplacementCS() {}
 	TClearReplacementCS(const ShaderMetaType::CompiledShaderInitializerType& Initializer)
@@ -168,6 +175,7 @@ public:
 		OutEnvironment.SetDefine(TEXT("THREADGROUPSIZE_X"), ThreadGroupSizeX);
 		OutEnvironment.SetDefine(TEXT("THREADGROUPSIZE_Y"), ThreadGroupSizeY);
 		OutEnvironment.SetDefine(TEXT("THREADGROUPSIZE_Z"), ThreadGroupSizeZ);
+		OutEnvironment.SetDefine(TEXT("USE_WRAPPED_GROUP_COUNT"), bUseWrappedGroupCount ? 1 : 0);
 		OutEnvironment.SetDefine(TEXT("RESOURCE_TYPE"), uint32(ResourceType));
 	}
 	
@@ -176,6 +184,19 @@ public:
 	inline const FShaderResourceParameter& GetClearResourceParam() const { return ClearResourceParam; }
 	inline uint32 GetResourceParamIndex() const { return ClearResourceParam.GetBaseIndex(); }
 
+	static inline FIntVector GetGroupCount(uint32 SizeX, uint32 SizeY, uint32 SizeZ)
+	{
+		if (bUseWrappedGroupCount)
+		{
+			return FComputeShaderUtils::GetGroupCountWrapped(SizeX, ThreadGroupSizeX);
+		}
+
+		return FIntVector(
+			FMath::DivideAndRoundUp(SizeX, ThreadGroupSizeX), 
+			FMath::DivideAndRoundUp(SizeY, ThreadGroupSizeY), 
+			FMath::DivideAndRoundUp(SizeZ, ThreadGroupSizeZ));
+	}
+
 private:
 	LAYOUT_FIELD(FShaderResourceParameter, ClearResourceParam);
 };
@@ -353,11 +374,9 @@ inline void ClearUAVShader_T(FRHIComputeCommandList& RHICmdList, FRHIUnorderedAc
 
 	ResourceBindCallback(ShaderRHI, ComputeShader->GetClearResourceParam(), true);
 
-	RHICmdList.DispatchComputeShader(
-		FMath::DivideAndRoundUp(SizeX, ComputeShader->ThreadGroupSizeX),
-		FMath::DivideAndRoundUp(SizeY, ComputeShader->ThreadGroupSizeY),
-		FMath::DivideAndRoundUp(SizeZ, ComputeShader->ThreadGroupSizeZ)
-	);
+	FIntVector GroupCount = FClearShader::GetGroupCount(SizeX, SizeY, SizeZ);
+
+	RHICmdList.DispatchComputeShader(GroupCount.X, GroupCount.Y, GroupCount.Z);
 
 	ResourceBindCallback(ShaderRHI, ComputeShader->GetClearResourceParam(), false);
 
@@ -392,11 +411,9 @@ inline void ClearUAVShader_T(FRHIComputeCommandList& RHICmdList, FRHIUnorderedAc
 
 	RHICmdList.SetBatchedShaderParameters(ShaderRHI, BatchedParameters);
 
-	RHICmdList.DispatchComputeShader(
-		FMath::DivideAndRoundUp(SizeX, ComputeShader->ThreadGroupSizeX),
-		FMath::DivideAndRoundUp(SizeY, ComputeShader->ThreadGroupSizeY),
-		FMath::DivideAndRoundUp(SizeZ, ComputeShader->ThreadGroupSizeZ)
-	);
+	const FIntVector GroupCount = FClearShader::GetGroupCount(SizeX, SizeY, SizeZ);
+
+	RHICmdList.DispatchComputeShader(GroupCount.X, GroupCount.Y, GroupCount.Z);
 
 	if (RHICmdList.NeedsShaderUnbinds())
 	{
