[shaders] rework fix from CL#32864562 to not add redundant padding characters

diff --git a/Engine/Source/Runtime/RenderCore/Private/ShaderPreprocessTypes.cpp b/Engine/Source/Runtime/RenderCore/Private/ShaderPreprocessTypes.cpp
--- a/Engine/Source/Runtime/RenderCore/Private/ShaderPreprocessTypes.cpp
+++ b/Engine/Source/Runtime/RenderCore/Private/ShaderPreprocessTypes.cpp
@@ -191,7 +191,7 @@ inline bool IsEndOfTheLine(FShaderSource::CharType C)
 
 inline bool StripNeedsHandling(FShaderSource::CharType C)
 {
-	return IsEndOfTheLine(C) || C == '/' || C == 0 || C == '#';
+	return IsEndOfTheLine(C) || C == '/' || C == '#';
 }
 
 inline void SkipNewLine(const FShaderSource::CharType*& Current, const FShaderSource::CharType* End)
@@ -212,7 +212,7 @@ void FShaderPreprocessOutput::StripCode(bool bCopyOriginalPreprocessdSource)
 	const FShaderSource::CharType* End = Current + PreprocessedSource.Len();
 	int32 CurrentBlockUnstrippedLineOffset = 0;
 	int32 CurrentStrippedLineNum = 1;
-	while (Current < End)
+	while (Current < End && *Current)
 	{
 		while (!StripNeedsHandling(*Current))
 		{
