[shaders] better fix for problem attempted to fix in CL#32895197. don't copy null chars through to the output, and ensure we abort scanning if a null character is hit.

diff --git a/Engine/Source/Runtime/RenderCore/Private/ShaderPreprocessTypes.cpp b/Engine/Source/Runtime/RenderCore/Private/ShaderPreprocessTypes.cpp
--- a/Engine/Source/Runtime/RenderCore/Private/ShaderPreprocessTypes.cpp
+++ b/Engine/Source/Runtime/RenderCore/Private/ShaderPreprocessTypes.cpp
@@ -191,7 +191,7 @@ inline bool IsEndOfTheLine(FShaderSource::CharType C)
 
 inline bool StripNeedsHandling(FShaderSource::CharType C)
 {
-	return IsEndOfTheLine(C) || C == '/' || C == '#';
+	return IsEndOfTheLine(C) || C == '/' || C == '#' || C == 0;
 }
 
 inline void SkipNewLine(const FShaderSource::CharType*& Current, const FShaderSource::CharType* End)
@@ -326,10 +326,6 @@ void FShaderPreprocessOutput::StripCode(bool bCopyOriginalPreprocessdSource)
 				*OutStripped++ = *Current++;
 			}
 		}
-		else if (*Current == 0)
-		{
-			*OutStripped++ = *Current++;
-		}
 	}
 	check(OutStripped <= OutStrippedData + PreprocessedSourceStripped.Len());
 	// ShrinkToLen null terminates for us by virtue of adding zero'd SIMD padding
