[shaders] add copyright boilerplate to newly added file

#rnx

diff --git a/Engine/Source/Runtime/RenderCore/Private/ShaderDiagnostics.cpp b/Engine/Source/Runtime/RenderCore/Private/ShaderDiagnostics.cpp
--- a/Engine/Source/Runtime/RenderCore/Private/ShaderDiagnostics.cpp
+++ b/Engine/Source/Runtime/RenderCore/Private/ShaderDiagnostics.cpp
@@ -1,3 +1,5 @@
+// Copyright Epic Games, Inc. All Rights Reserved.
+
 #include "ShaderDiagnostics.h"
 
 #include "DataDrivenShaderPlatformInfo.h"
