Remove r.RayTracing.EnableInGame.

diff --git a/Engine/Source/Runtime/RenderCore/Private/RenderUtils.cpp b/Engine/Source/Runtime/RenderCore/Private/RenderUtils.cpp
--- a/Engine/Source/Runtime/RenderCore/Private/RenderUtils.cpp
+++ b/Engine/Source/Runtime/RenderCore/Private/RenderUtils.cpp
@@ -41,15 +41,6 @@ FAutoConsoleVariableRef CVarAllowTranslucencyShadowsInProject(
 	ECVF_ReadOnly | ECVF_RenderThreadSafe
 );
 
-int32 GRayTracingEnableInGame = 1;
-FAutoConsoleVariableRef CVarRayTracingEnableInGame(
-	TEXT("r.RayTracing.EnableInGame"),
-	GRayTracingEnableInGame,
-	TEXT("Controls the default state of ray tracing effects when running the game. This setting is overridden by its counterpart in GameUserSettings.ini (if it exists) to allow control through in-game UI. ")
-	TEXT("(default = 1)"),
-	ECVF_ReadOnly
-);
-
 int32 GRayTracingEnableInEditor = 1;
 FAutoConsoleVariableRef CVarRayTracingEnableInEditor(
 	TEXT("r.RayTracing.EnableInEditor"),
@@ -931,34 +922,12 @@ RENDERCORE_API void RenderUtilsInit()
 			}
 			else
 			{
-				// If user preference exists in game settings file, the bRayTracingEnabled will be set based on its value.
-				// Otherwise the current value is preserved.
-				bool bUseRayTracing = false;
-				if (GRayTracingEnableOnDemand == 1)
-				{
-					GRayTracingMode = DesiredRayTracingMode;
-
-					UE_LOG(LogRendererCore, Log, TEXT("Ray tracing is %s for the game. Reason: r.RayTracing=%d and r.RayTracing.EnableOnDemand=1."),
-						GetRayTracingModeName(GRayTracingMode),
-						RayTracingInt);
-				}
-				else if (GConfig->GetBool(TEXT("RayTracing"), TEXT("r.RayTracing.EnableInGame"), bUseRayTracing, GGameUserSettingsIni))
-				{
-					GRayTracingMode = bUseRayTracing ? DesiredRayTracingMode : ERayTracingMode::Disabled;
-
-					UE_LOG(LogRendererCore, Log, TEXT("Ray tracing is %s for the game. Reason: game user setting r.RayTracing.EnableInGame=%d."),
-						GetRayTracingModeName(GRayTracingMode),
-						(int)bUseRayTracing);
-				}
-				else
-				{
-					GRayTracingMode = GRayTracingEnableInGame != 0 ? DesiredRayTracingMode : ERayTracingMode::Disabled;
+				GRayTracingMode = DesiredRayTracingMode;
 
-					UE_LOG(LogRendererCore, Log, TEXT("Ray tracing is %s for the game. Reason: CVar r.RayTracing=%d, and r.RayTracing.EnableInGame game user setting does not exist (using default from CVar: %d)."),
-						GetRayTracingModeName(GRayTracingMode),
-						RayTracingInt,
-						GRayTracingEnableInGame);
-				}
+				UE_LOG(LogRendererCore, Log, TEXT("Ray tracing is %s for the game. Reason: r.RayTracing=%d and r.RayTracing.EnableOnDemand=%d."),
+					GetRayTracingModeName(GRayTracingMode),
+					RayTracingInt,
+					GRayTracingEnableOnDemand);
 			}
 
 			// Sanity check: skin cache is *required* for ray tracing.
