Removed deferral of RDG command list setup task submission as it causes RHI work to execute out of order of the serial mode, causing validation errors.

diff --git a/Engine/Source/Runtime/RenderCore/Private/RenderGraphBuilder.cpp b/Engine/Source/Runtime/RenderCore/Private/RenderGraphBuilder.cpp
--- a/Engine/Source/Runtime/RenderCore/Private/RenderGraphBuilder.cpp
+++ b/Engine/Source/Runtime/RenderCore/Private/RenderGraphBuilder.cpp
@@ -1618,15 +1618,6 @@ void FRDGBuilder::WaitForParallelSetupTasks()
 	}
 }
 
-void FRDGBuilder::SubmitParallelSetupTasks()
-{
-	if (!ParallelSetup.CommandLists.IsEmpty())
-	{
-		RHICmdList.QueueAsyncCommandListSubmit(ParallelSetup.CommandLists);
-		ParallelSetup.CommandLists.Empty();
-	}
-}
-
 ///////////////////////////////////////////////////////////////////////////////////////////////////
 
 void FRDGBuilder::Execute()
@@ -1666,7 +1657,6 @@ void FRDGBuilder::Execute()
 
 	if (!IsImmediateMode())
 	{
-		SubmitParallelSetupTasks();
 		BeginFlushResourcesRHI();
 		WaitForParallelSetupTasks();
 
@@ -1895,7 +1885,6 @@ void FRDGBuilder::Execute()
 		FinalizeResources();
 	}
 
-	SubmitParallelSetupTasks();
 	EndFlushResourcesRHI();
 	WaitForParallelSetupTasks();
 
diff --git a/Engine/Source/Runtime/RenderCore/Public/RenderGraphBuilder.h b/Engine/Source/Runtime/RenderCore/Public/RenderGraphBuilder.h
--- a/Engine/Source/Runtime/RenderCore/Public/RenderGraphBuilder.h
+++ b/Engine/Source/Runtime/RenderCore/Public/RenderGraphBuilder.h
@@ -933,15 +933,11 @@ private:
 		/** Array of all tasks for variants of AddSetupTask. */
 		TArray<UE::Tasks::FTask, FRDGArrayAllocator> Tasks;
 
-		/** Array of all command lists to submit for AddCommandListSetupTask. */
-		TArray<FRHICommandListImmediate::FQueuedCommandList, FConcurrentLinearArrayAllocator> CommandLists;
-
 		bool bEnabled = false;
 
 	} ParallelSetup;
 
 	void WaitForParallelSetupTasks();
-	void SubmitParallelSetupTasks();
 
 	/////////////////////////////////////////////////////////////////////////////
 	// Parallel Execution
diff --git a/Engine/Source/Runtime/RenderCore/Public/RenderGraphBuilder.inl b/Engine/Source/Runtime/RenderCore/Public/RenderGraphBuilder.inl
--- a/Engine/Source/Runtime/RenderCore/Public/RenderGraphBuilder.inl
+++ b/Engine/Source/Runtime/RenderCore/Public/RenderGraphBuilder.inl
@@ -623,7 +623,7 @@ UE::Tasks::FTask FRDGBuilder::AddCommandListSetupTask(
 	{
 		SCOPED_NAMED_EVENT(CreateCommandList, FColor::Emerald);
 		RHICmdListTask = new FRHICommandList(RHICmdList.GetGPUMask());
-		ParallelSetup.CommandLists.Emplace(RHICmdListTask);
+		RHICmdList.QueueAsyncCommandListSubmit(RHICmdListTask);
 	}
 
 	auto OuterLambda = [this, TaskLambda = MoveTemp(TaskLambda), RHICmdListTask, bAllocateCommandListForTask]() mutable
