Fixed shadowed local variable in RDG builder.

diff --git a/Engine/Source/Runtime/RenderCore/Private/RenderGraphBuilder.cpp b/Engine/Source/Runtime/RenderCore/Private/RenderGraphBuilder.cpp
--- a/Engine/Source/Runtime/RenderCore/Private/RenderGraphBuilder.cpp
+++ b/Engine/Source/Runtime/RenderCore/Private/RenderGraphBuilder.cpp
@@ -1870,9 +1870,9 @@ void FRDGBuilder::Execute()
 			}
 		}
 
-		AddSetupTask([this, UniformBuffers = MoveTemp(CollectResourceContext.UniformBuffers)]
+		AddSetupTask([this, InUniformBuffers = MoveTemp(CollectResourceContext.UniformBuffers)]
 		{
-			CreateUniformBuffers(UniformBuffers);
+			CreateUniformBuffers(InUniformBuffers);
 
 		}, CreateViewsTask, TaskPriority); // Uniform buffer creation require views to be valid.
 
