Fixed incorrect RHI transition logging.

diff --git a/Engine/Source/Runtime/RenderCore/Private/RenderGraphPass.cpp b/Engine/Source/Runtime/RenderCore/Private/RenderGraphPass.cpp
--- a/Engine/Source/Runtime/RenderCore/Private/RenderGraphPass.cpp
+++ b/Engine/Source/Runtime/RenderCore/Private/RenderGraphPass.cpp
@@ -103,8 +103,6 @@ FRDGBarrierBatchBegin::FRDGBarrierBatchBegin(ERHIPipeline InPipelineToBegin, ERH
 #if RDG_ENABLE_DEBUG
 	, DebugPasses(InPlace, nullptr)
 	, DebugName(InDebugName)
-	, DebugPipelinesToBegin(InPipelineToBegin)
-	, DebugPipelinesToEnd(InPipelinesToEnd)
 #endif
 {
 #if RDG_ENABLE_DEBUG
@@ -118,8 +116,6 @@ FRDGBarrierBatchBegin::FRDGBarrierBatchBegin(ERHIPipeline InPipelinesToBegin, ER
 #if RDG_ENABLE_DEBUG
 	, DebugPasses(InDebugPasses)
 	, DebugName(InDebugName)
-	, DebugPipelinesToBegin(InPipelinesToBegin)
-	, DebugPipelinesToEnd(InPipelinesToEnd)
 #endif
 {}
 
@@ -281,7 +277,7 @@ FRDGBarrierBatchBegin& FRDGPass::GetEpilogueBarriersToBeginForAll(FRDGAllocator&
 {
 	if (!EpilogueBarriersToBeginForAll)
 	{
-		EpilogueBarriersToBeginForAll = Allocator.AllocNoDestruct<FRDGBarrierBatchBegin>(Pipeline, ERHIPipeline::All, GetEpilogueBarriersToBeginDebugName(ERHIPipeline::AsyncCompute), this);
+		EpilogueBarriersToBeginForAll = Allocator.AllocNoDestruct<FRDGBarrierBatchBegin>(Pipeline, ERHIPipeline::All, GetEpilogueBarriersToBeginDebugName(ERHIPipeline::All), this);
 		CreateQueue.Emplace(EpilogueBarriersToBeginForAll);
 	}
 	return *EpilogueBarriersToBeginForAll;
diff --git a/Engine/Source/Runtime/RenderCore/Private/RenderGraphValidation.cpp b/Engine/Source/Runtime/RenderCore/Private/RenderGraphValidation.cpp
--- a/Engine/Source/Runtime/RenderCore/Private/RenderGraphValidation.cpp
+++ b/Engine/Source/Runtime/RenderCore/Private/RenderGraphValidation.cpp
@@ -1418,9 +1418,9 @@ void FRDGBarrierValidation::ValidateBarrierBatchBegin(const FRDGPass* Pass, cons
 			UE_LOG(LogRDG, Display, TEXT("\t\tMip(%d), Array(%d), Slice(%d): [%s, %s] -> [%s, %s]"),
 				Transition.MipIndex, Transition.ArraySlice, Transition.PlaneSlice,
 				*GetRHIAccessName(Transition.AccessBefore),
-				*GetRHIPipelineName(Batch.DebugPipelinesToBegin),
+				*GetRHIPipelineName(Batch.PipelinesToBegin),
 				*GetRHIAccessName(Transition.AccessAfter),
-				*GetRHIPipelineName(Batch.DebugPipelinesToEnd));
+				*GetRHIPipelineName(Batch.PipelinesToEnd));
 		}
 	}
 
@@ -1441,9 +1441,9 @@ void FRDGBarrierValidation::ValidateBarrierBatchBegin(const FRDGPass* Pass, cons
 			Buffer->GetRHIUnchecked(),
 			Buffer->Name,
 			*GetRHIAccessName(Transition.AccessBefore),
-			*GetRHIPipelineName(Batch.DebugPipelinesToBegin),
+			*GetRHIPipelineName(Batch.PipelinesToBegin),
 			*GetRHIAccessName(Transition.AccessAfter),
-			*GetRHIPipelineName(Batch.DebugPipelinesToEnd));
+			*GetRHIPipelineName(Batch.PipelinesToEnd));
 	}
 }
 
diff --git a/Engine/Source/Runtime/RenderCore/Public/RenderGraphPass.h b/Engine/Source/Runtime/RenderCore/Public/RenderGraphPass.h
--- a/Engine/Source/Runtime/RenderCore/Public/RenderGraphPass.h
+++ b/Engine/Source/Runtime/RenderCore/Public/RenderGraphPass.h
@@ -151,8 +151,6 @@ private:
 	TArray<FRDGViewableResource*, FRDGArrayAllocator> DebugTransitionResources;
 	TArray<FRDGViewableResource*, FRDGArrayAllocator> DebugAliasingResources;
 	const TCHAR* DebugName;
-	ERHIPipeline DebugPipelinesToBegin;
-	ERHIPipeline DebugPipelinesToEnd;
 #endif
 
 	friend class FRDGBarrierBatchEnd;
