Remove r.RayTracing.EnableInEditor
- projects can use r.RayTracing.Enable to disable ray tracing in editor device profiles if necessary.

diff --git a/Engine/Source/Runtime/RenderCore/Private/RenderUtils.cpp b/Engine/Source/Runtime/RenderCore/Private/RenderUtils.cpp
--- a/Engine/Source/Runtime/RenderCore/Private/RenderUtils.cpp
+++ b/Engine/Source/Runtime/RenderCore/Private/RenderUtils.cpp
@@ -41,22 +41,13 @@ FAutoConsoleVariableRef CVarAllowTranslucencyShadowsInProject(
 	ECVF_ReadOnly | ECVF_RenderThreadSafe
 );
 
-int32 GRayTracingEnableInEditor = 1;
-FAutoConsoleVariableRef CVarRayTracingEnableInEditor(
-	TEXT("r.RayTracing.EnableInEditor"),
-	GRayTracingEnableInEditor,
-	TEXT("Controls whether ray tracing effects are available by default when running the editor. This can be useful to improve editor performance when only some people require ray tracing features. ")
-	TEXT("(default = 1)"),
-	ECVF_ReadOnly
-);
-
 static int32 GRayTracingEnableOnDemand = 0;
 static FAutoConsoleVariableRef CVarRayTracingEnableOnDemand(
 	TEXT("r.RayTracing.EnableOnDemand"),
 	GRayTracingEnableOnDemand,
 	TEXT("Controls whether ray tracing features can be toggled on demand at runtime without restarting the game (experimental).\n")
-	TEXT("Requires r.RayTracing=1 and will override GameUserSettings in game. Requires r.RayTracing.EnableInEditor=1 in editor. Has a small performance and memory overhead.\n")
-	TEXT(" 0: off (default)\n")
+	TEXT("Requires r.RayTracing=1. Has a small performance and memory overhead.\n")
+	TEXT(" 0: off\n")
 	TEXT(" 1: on"),
 	ECVF_RenderThreadSafe | ECVF_ReadOnly);
 
@@ -864,8 +855,7 @@ RENDERCORE_API void RenderUtilsInit()
 	// - Ray tracing must be enabled for the project
 	// - Skin cache must be enabled for the project
 	// - Current GPU, OS and driver must support ray tracing
-	// - User is running the Editor and r.RayTracing.EnableInEditor=1 
-	//   *OR* running the game with ray tracing enabled in graphics options
+	// - r.RayTracing.Enable = 1
 
 	// When ray tracing is enabled, we must load additional shaders and build acceleration structures for meshes.
 	// For this reason it is only possible to enable RT at startup and changing the state requires restart.
@@ -914,27 +904,12 @@ RENDERCORE_API void RenderUtilsInit()
 		const bool bRayTracingAllowedOnCurrentPlatform = (GRayTracingPlatformMask[(int)GMaxRHIShaderPlatform]);
 		if (GRHISupportsRayTracing && bRayTracingAllowedOnCurrentPlatform)
 		{
-			if (GIsEditor)
-			{
-				// Ray tracing is enabled for the project and we are running on RT-capable machine,
-				// therefore the core ray tracing features are also enabled, so that required shaders
-				// are loaded, acceleration structures are built, etc.
-				GRayTracingMode = (GRayTracingEnableInEditor != 0) ? DesiredRayTracingMode : ERayTracingMode::Disabled;
-
-				UE_LOG(LogRendererCore, Log, TEXT("Ray tracing is %s for the editor. Reason: r.RayTracing=%d and r.RayTracing.EnableInEditor=%d."),
-					GetRayTracingModeName(GRayTracingMode),
-					RayTracingInt,
-					GRayTracingEnableInEditor);
-			}
-			else
-			{
-				GRayTracingMode = DesiredRayTracingMode;
+			GRayTracingMode = DesiredRayTracingMode;
 
-				UE_LOG(LogRendererCore, Log, TEXT("Ray tracing is %s for the game. Reason: r.RayTracing=%d and r.RayTracing.EnableOnDemand=%d."),
-					GetRayTracingModeName(GRayTracingMode),
-					RayTracingInt,
-					GRayTracingEnableOnDemand);
-			}
+			UE_LOG(LogRendererCore, Log, TEXT("Ray tracing is %s. Reason: r.RayTracing=%d and r.RayTracing.EnableOnDemand=%d."),
+				GetRayTracingModeName(GRayTracingMode),
+				RayTracingInt,
+				GRayTracingEnableOnDemand);
 
 			// Sanity check: skin cache is *required* for ray tracing.
 			// It can be dynamically enabled only when its shaders have been compiled.
