Fixed RHI validation error due to SRV being in a UAV state.

diff --git a/Engine/Source/Runtime/RenderCore/Private/GlobalRenderResources.cpp b/Engine/Source/Runtime/RenderCore/Private/GlobalRenderResources.cpp
--- a/Engine/Source/Runtime/RenderCore/Private/GlobalRenderResources.cpp
+++ b/Engine/Source/Runtime/RenderCore/Private/GlobalRenderResources.cpp
@@ -112,10 +112,10 @@ public:
 	virtual void InitRHI(FRHICommandListBase& RHICmdList) override
 	{
 		// Create the buffer RHI.  		
-		FRHIResourceCreateInfo CreateInfo(TEXT("EmptyStructuredBuffer"));
+		FRHIResourceCreateInfo CreateInfo(TEXT("BlackStructuredBuffer"));
 
 		const uint32 BufferSize = sizeof(float) * 4u;
-		VertexBufferRHI = RHICmdList.CreateStructuredBuffer(sizeof(float), BufferSize, BUF_Static | BUF_ShaderResource | BUF_UnorderedAccess, CreateInfo);
+		VertexBufferRHI = RHICmdList.CreateStructuredBuffer(sizeof(float), BufferSize, BUF_Static | BUF_ShaderResource | BUF_UnorderedAccess, ERHIAccess::SRVMask, CreateInfo);
 
 		FVector4f* BufferData = (FVector4f*)RHICmdList.LockBuffer(VertexBufferRHI, 0, sizeof(FVector4f), RLM_WriteOnly);
 		*BufferData = FVector4f(0.0f, 0.0f, 0.0f, 0.0f);
