Fix for RHI validation errors on DistantSkyLightLutBuffer and MobileDistantSkyLightLutBuffer when falling back to global dummy resources which don't have a float4 stride. Fixed by introducing new float4 type dummy buffers. GBlackStructuredBufferWithSRV was only recently introduced (never shipped in a released UE version) and has been removed/replaced in this CL.

#rnx

diff --git a/Engine/Source/Runtime/Engine/Private/SceneManagement.cpp b/Engine/Source/Runtime/Engine/Private/SceneManagement.cpp
--- a/Engine/Source/Runtime/Engine/Private/SceneManagement.cpp
+++ b/Engine/Source/Runtime/Engine/Private/SceneManagement.cpp
@@ -1180,8 +1180,8 @@ FViewUniformShaderParameters::FViewUniformShaderParameters()
 	SkyViewLutTexture = GBlackTexture->TextureRHI;
 	SkyViewLutTextureSampler = TStaticSamplerState<SF_Bilinear>::GetRHI();
 
-	DistantSkyLightLutBufferSRV = GBlackVertexBufferWithSRV->ShaderResourceViewRHI;
-	MobileDistantSkyLightLutBufferSRV = GBlackStructuredBufferWithSRV->ShaderResourceViewRHI;
+	DistantSkyLightLutBufferSRV = GBlackFloat4StructuredBufferWithSRV->ShaderResourceViewRHI;
+	MobileDistantSkyLightLutBufferSRV = GBlackFloat4VertexBufferWithSRV->ShaderResourceViewRHI;
 
 	CameraAerialPerspectiveVolume = GBlackAlpha1VolumeTexture->TextureRHI;
 	CameraAerialPerspectiveVolumeSampler = TStaticSamplerState<SF_Bilinear>::GetRHI();
diff --git a/Engine/Source/Runtime/RenderCore/Private/GlobalRenderResources.cpp b/Engine/Source/Runtime/RenderCore/Private/GlobalRenderResources.cpp
--- a/Engine/Source/Runtime/RenderCore/Private/GlobalRenderResources.cpp
+++ b/Engine/Source/Runtime/RenderCore/Private/GlobalRenderResources.cpp
@@ -106,16 +106,16 @@ public:
 	}
 };
 
-class FBlackStructuredBufferWithSRV : public FVertexBufferWithSRV
+class FBlackFloat4StructuredBufferWithSRV : public FVertexBufferWithSRV
 {
 public:
 	virtual void InitRHI(FRHICommandListBase& RHICmdList) override
 	{
 		// Create the buffer RHI.  		
-		FRHIResourceCreateInfo CreateInfo(TEXT("BlackStructuredBuffer"));
+		FRHIResourceCreateInfo CreateInfo(TEXT("BlackFloat4StructuredBuffer"));
 
-		const uint32 BufferSize = sizeof(float) * 4u;
-		VertexBufferRHI = RHICmdList.CreateStructuredBuffer(sizeof(float), BufferSize, BUF_Static | BUF_ShaderResource | BUF_UnorderedAccess, ERHIAccess::SRVMask, CreateInfo);
+		const uint32 BufferSize = sizeof(FVector4f);
+		VertexBufferRHI = RHICmdList.CreateStructuredBuffer(sizeof(FVector4f), BufferSize, BUF_Static | BUF_ShaderResource | BUF_UnorderedAccess, ERHIAccess::SRVMask, CreateInfo);
 
 		FVector4f* BufferData = (FVector4f*)RHICmdList.LockBuffer(VertexBufferRHI, 0, sizeof(FVector4f), RLM_WriteOnly);
 		*BufferData = FVector4f(0.0f, 0.0f, 0.0f, 0.0f);
@@ -127,6 +127,26 @@ public:
 	}
 };
 
+class FBlackFloat4VertexBuffer : public FVertexBufferWithSRV
+{
+public:
+	virtual void InitRHI(FRHICommandListBase& RHICmdList) override
+	{
+		// Create the texture RHI.  		
+		FRHIResourceCreateInfo CreateInfo(TEXT("BlackFloat4VertexBuffer"));
+
+		VertexBufferRHI = RHICmdList.CreateVertexBuffer(sizeof(FVector4f), BUF_Static | BUF_ShaderResource | BUF_UnorderedAccess, ERHIAccess::SRVMask, CreateInfo);
+
+		FVector4f* BufferData = (FVector4f*)RHICmdList.LockBuffer(VertexBufferRHI, 0, sizeof(FVector4f), RLM_WriteOnly);
+		*BufferData = FVector4f(0.0f, 0.0f, 0.0f, 0.0f);
+		RHICmdList.UnlockBuffer(VertexBufferRHI);
+
+		// Create a view of the buffer
+		ShaderResourceViewRHI = RHICmdList.CreateShaderResourceView(VertexBufferRHI, sizeof(FVector4f), PF_A32B32G32R32F);
+		UnorderedAccessViewRHI = RHICmdList.CreateUnorderedAccessView(VertexBufferRHI, PF_A32B32G32R32F);
+	}
+};
+
 class FBlackTextureWithSRV : public FColoredTexture<0, 0, 0, 255>
 {
 public:
@@ -152,7 +172,8 @@ FTexture* GTransparentBlackTexture = GTransparentBlackTextureWithSRV;
 
 FVertexBufferWithSRV* GEmptyVertexBufferWithUAV = new TGlobalResource<FEmptyVertexBuffer, FRenderResource::EInitPhase::Pre>;
 FVertexBufferWithSRV* GEmptyStructuredBufferWithUAV = new TGlobalResource<FEmptyStructuredBuffer, FRenderResource::EInitPhase::Pre>;
-FVertexBufferWithSRV* GBlackStructuredBufferWithSRV = new TGlobalResource<FBlackStructuredBufferWithSRV, FRenderResource::EInitPhase::Pre>;
+FVertexBufferWithSRV* GBlackFloat4StructuredBufferWithSRV = new TGlobalResource<FBlackFloat4StructuredBufferWithSRV, FRenderResource::EInitPhase::Pre>;
+FVertexBufferWithSRV* GBlackFloat4VertexBufferWithSRV = new TGlobalResource<FBlackFloat4VertexBuffer, FRenderResource::EInitPhase::Pre>;
 
 class FWhiteVertexBuffer : public FVertexBufferWithSRV
 {
diff --git a/Engine/Source/Runtime/RenderCore/Public/GlobalRenderResources.h b/Engine/Source/Runtime/RenderCore/Public/GlobalRenderResources.h
--- a/Engine/Source/Runtime/RenderCore/Public/GlobalRenderResources.h
+++ b/Engine/Source/Runtime/RenderCore/Public/GlobalRenderResources.h
@@ -28,8 +28,11 @@ extern RENDERCORE_API FVertexBufferWithSRV* GEmptyVertexBufferWithUAV;
 // An empty structured buffer with a UAV
 extern RENDERCORE_API FVertexBufferWithSRV* GEmptyStructuredBufferWithUAV;
 
-// An empty structured buffer with a UAV
-extern RENDERCORE_API FVertexBufferWithSRV* GBlackStructuredBufferWithSRV;
+// An empty float4 structured buffer
+extern RENDERCORE_API FVertexBufferWithSRV* GBlackFloat4StructuredBufferWithSRV;
+
+// An empty float4 vertex buffer
+extern RENDERCORE_API FVertexBufferWithSRV* GBlackFloat4VertexBufferWithSRV;
 
 // An white vertex buffer with a SRV
 extern RENDERCORE_API FVertexBufferWithSRV* GWhiteVertexBufferWithSRV;
diff --git a/Engine/Source/Runtime/Renderer/Private/SceneRendering.cpp b/Engine/Source/Runtime/Renderer/Private/SceneRendering.cpp
--- a/Engine/Source/Runtime/Renderer/Private/SceneRendering.cpp
+++ b/Engine/Source/Runtime/Renderer/Private/SceneRendering.cpp
@@ -1513,7 +1513,7 @@ void FViewInfo::SetupUniformBufferParameters(
 	}
 	else
 	{
-		ViewUniformShaderParameters.DistantSkyLightLutBufferSRV = GBlackStructuredBufferWithSRV->ShaderResourceViewRHI.GetReference();
+		ViewUniformShaderParameters.DistantSkyLightLutBufferSRV = GBlackFloat4StructuredBufferWithSRV->ShaderResourceViewRHI.GetReference();
 	}
 	if(MobileDistantSkyLightLutBufferSRVFound != nullptr)
 	{
@@ -1521,7 +1521,7 @@ void FViewInfo::SetupUniformBufferParameters(
 	}
 	else
 	{
-		ViewUniformShaderParameters.MobileDistantSkyLightLutBufferSRV = GBlackVertexBufferWithSRV->ShaderResourceViewRHI.GetReference();
+		ViewUniformShaderParameters.MobileDistantSkyLightLutBufferSRV = GBlackFloat4VertexBufferWithSRV->ShaderResourceViewRHI.GetReference();
 	}
 	ViewUniformShaderParameters.SkyViewLutTexture = OrBlack2DIfNull(SkyViewLutTextureFound);
 	ViewUniformShaderParameters.SkyViewLutTextureSampler = TStaticSamplerState<SF_Bilinear>::GetRHI();
diff --git a/Engine/Source/Runtime/Renderer/Private/SkyAtmosphereRendering.cpp b/Engine/Source/Runtime/Renderer/Private/SkyAtmosphereRendering.cpp
--- a/Engine/Source/Runtime/Renderer/Private/SkyAtmosphereRendering.cpp
+++ b/Engine/Source/Runtime/Renderer/Private/SkyAtmosphereRendering.cpp
@@ -523,7 +523,7 @@ FRHIShaderResourceView* FSkyAtmosphereRenderSceneInfo::GetDistantSkyLightLutBuff
 	{
 		return DistantSkyLightLutBufferSRV;
 	}
-	return GBlackVertexBufferWithSRV->ShaderResourceViewRHI;
+	return GBlackFloat4StructuredBufferWithSRV->ShaderResourceViewRHI;
 }
 
 FRHIShaderResourceView* FSkyAtmosphereRenderSceneInfo::GetMobileDistantSkyLightLutBufferSRV()
@@ -532,7 +532,7 @@ FRHIShaderResourceView* FSkyAtmosphereRenderSceneInfo::GetMobileDistantSkyLightL
 	{
 		return MobileDistantSkyLightLutBufferSRV;
 	}
-	return GBlackStructuredBufferWithSRV->ShaderResourceViewRHI;
+	return GBlackFloat4VertexBufferWithSRV->ShaderResourceViewRHI;
 }
 
 
