Removed awkward CreateFromStreamSourceSlot in favor of constructor (uses nullptr_t to disambiguate nullptr).

diff --git a/Engine/Source/Runtime/Engine/Private/GPUSkinVertexFactory.cpp b/Engine/Source/Runtime/Engine/Private/GPUSkinVertexFactory.cpp
--- a/Engine/Source/Runtime/Engine/Private/GPUSkinVertexFactory.cpp
+++ b/Engine/Source/Runtime/Engine/Private/GPUSkinVertexFactory.cpp
@@ -710,7 +710,7 @@ void FGPUBaseSkinVertexFactory::GetOverrideVertexStreams(FVertexInputStreamArray
 {
 	if (MorphDeltaStreamIndex >= 0)
 	{
-		VertexStreams.Emplace(FVertexInputStream::CreateFromStreamSourceSlot(MorphDeltaStreamIndex, 0, MorphDeltaBufferSlot));
+		VertexStreams.Emplace(MorphDeltaStreamIndex, 0, MorphDeltaBufferSlot);
 	}
 }
 
@@ -1435,7 +1435,7 @@ void FGPUSkinPassthroughVertexFactory::GetOverrideVertexStreams(FVertexInputStre
 	{
 		if (EnumHasAnyFlags(VertexAttributesToBind, static_cast<EVertexAttributeFlags>(1 << Index)))
 		{
-			VertexStreams.Add(FVertexInputStream::CreateFromStreamSourceSlot(Index, 0, StreamSourceSlots[Index]));
+			VertexStreams.Emplace(Index, 0, StreamSourceSlots[Index]);
 		}
 	}
 }
diff --git a/Engine/Source/Runtime/RenderCore/Public/VertexFactory.h b/Engine/Source/Runtime/RenderCore/Public/VertexFactory.h
--- a/Engine/Source/Runtime/RenderCore/Public/VertexFactory.h
+++ b/Engine/Source/Runtime/RenderCore/Public/VertexFactory.h
@@ -71,29 +71,26 @@ struct FVertexInputStream
 		bStreamSourceSlot(0),
 		VertexBuffer(nullptr)
 	{}
-
-	FVertexInputStream(uint32 InStreamIndex, uint32 InOffset, FRHIBuffer* InVertexBuffer)
-		: StreamIndex(InStreamIndex), Offset(InOffset), bStreamSourceSlot(0), VertexBuffer(InVertexBuffer)
+	
+	FVertexInputStream(uint32 InStreamIndex, uint32 InOffset, nullptr_t)
+		: StreamIndex(InStreamIndex), Offset(InOffset), bStreamSourceSlot(0), VertexBuffer(nullptr)
 	{
 		// Verify no overflow
 		checkSlow(InStreamIndex == StreamIndex && InOffset == Offset);
 	}
 
-	// Creates a vertex input stream from an RHI stream source slot.
-	static FVertexInputStream CreateFromStreamSourceSlot(uint32 InStreamIndex, uint32 InOffset, FRHIStreamSourceSlot* StreamSourceSlot)
+	FVertexInputStream(uint32 InStreamIndex, uint32 InOffset, FRHIStreamSourceSlot* InStreamSourceSlot)
+		: StreamIndex(InStreamIndex), Offset(InOffset), bStreamSourceSlot(1), StreamSourceSlot(InStreamSourceSlot)
 	{
-		check(StreamSourceSlot);
-
-		FVertexInputStream Stream;
-		Stream.StreamIndex = InStreamIndex;
-		Stream.bStreamSourceSlot = true;
-		Stream.Offset = InOffset;
-		Stream.StreamSourceSlot = StreamSourceSlot;
-
 		// Verify no overflow
-		checkSlow(InStreamIndex == Stream.StreamIndex && InOffset == Stream.Offset);
+		checkSlow(InStreamIndex == StreamIndex && InOffset == Offset);
+	}
 
-		return Stream;
+	FVertexInputStream(uint32 InStreamIndex, uint32 InOffset, FRHIBuffer* InVertexBuffer)
+		: StreamIndex(InStreamIndex), Offset(InOffset), bStreamSourceSlot(0), VertexBuffer(InVertexBuffer)
+	{
+		// Verify no overflow
+		checkSlow(InStreamIndex == StreamIndex && InOffset == Offset);
 	}
 
 	inline bool operator==(const FVertexInputStream& rhs) const
