Disable timestamp render query batching introduced in CL 32763001 by removing the calls to RHIBeginRenderQueryBatch_TopOfPipe and RHIEndRenderQueryBatch_TopOfPipe in the realtime GPU profiler
 - The batches weren't getting flushed when RHICmdLists are submitted which causes the GPU benchmark to hang, since it issues and then immediately reads back timestamp queries.
 - Ideally the batches get ended whenever RHICmdLists are submitted, but the way the active batch is propagated to RDG passes means there isn't a simple way to achieve this.
 - There is no good solution until we've refactored the GPU profiler, at which point this becomes a non-issue.

diff --git a/Engine/Source/Runtime/RenderCore/Private/ProfilingDebugging/RealtimeGPUProfiler.cpp b/Engine/Source/Runtime/RenderCore/Private/ProfilingDebugging/RealtimeGPUProfiler.cpp
--- a/Engine/Source/Runtime/RenderCore/Private/ProfilingDebugging/RealtimeGPUProfiler.cpp
+++ b/Engine/Source/Runtime/RenderCore/Private/ProfilingDebugging/RealtimeGPUProfiler.cpp
@@ -932,12 +932,6 @@ void FRealtimeGPUProfiler::BeginFrame(FRHICommandListImmediate& RHICmdList)
 
 	Frames[WriteBufferIndex]->TimestampCalibrationQuery = new FRHITimestampCalibrationQuery();
 	RHICmdList.CalibrateTimers(Frames[WriteBufferIndex]->TimestampCalibrationQuery);
-
-	check(!bQueryBatchStarted);
-	bQueryBatchStarted = true;
-	// Using GDynamicRHI directly to avoid exposing this as a public API on RHICmdList.
-	// This is a temporary fix until we refactor the GPU profiler.
-	GDynamicRHI->RHIBeginRenderQueryBatch_TopOfPipe(RHICmdList, RQT_AbsoluteTime);
 }
 
 bool AreGPUStatsEnabled()
@@ -977,15 +971,6 @@ bool AreGPUStatsEnabled()
 
 void FRealtimeGPUProfiler::EndFrame(FRHICommandListImmediate& RHICmdList)
 {
-	// Do this before the AreGPUStatsEnabled check, in case the cvar changes mid-frame.
-	if (bQueryBatchStarted)
-	{
-	    // Using GDynamicRHI directly to avoid exposing this as a public API on RHICmdList.
-	    // This is a temporary fix until we refactor the GPU profiler.
-		GDynamicRHI->RHIEndRenderQueryBatch_TopOfPipe(RHICmdList, RQT_AbsoluteTime);
-		bQueryBatchStarted = false;
-	}
-
 	if (!AreGPUStatsEnabled())
 	{
 		return;
diff --git a/Engine/Source/Runtime/RenderCore/Public/ProfilingDebugging/RealtimeGPUProfiler.h b/Engine/Source/Runtime/RenderCore/Public/ProfilingDebugging/RealtimeGPUProfiler.h
--- a/Engine/Source/Runtime/RenderCore/Public/ProfilingDebugging/RealtimeGPUProfiler.h
+++ b/Engine/Source/Runtime/RenderCore/Public/ProfilingDebugging/RealtimeGPUProfiler.h
@@ -319,7 +319,6 @@ private:
 	bool bStatGatheringPaused;
 	bool bInBeginEndBlock;
 	bool bLocked = false;
-	bool bQueryBatchStarted = false;
 
 #if GPUPROFILERTRACE_ENABLED
 	FRealtimeGPUProfilerHistoryByDescription HistoryByDescription;
