Added CreateByteAddressBuffer variants to mirror CreateStructuredBuffer

diff --git a/Engine/Source/Runtime/RenderCore/Private/RenderGraphUtils.cpp b/Engine/Source/Runtime/RenderCore/Private/RenderGraphUtils.cpp
--- a/Engine/Source/Runtime/RenderCore/Private/RenderGraphUtils.cpp
+++ b/Engine/Source/Runtime/RenderCore/Private/RenderGraphUtils.cpp
@@ -910,6 +910,31 @@ FRDGBufferRef CreateStructuredBuffer(
 	return Buffer;
 }
 
+FRDGBufferRef CreateByteAddressBuffer(
+	FRDGBuilder& GraphBuilder,
+	const TCHAR* Name,
+	uint32 NumBytes,
+	const void* InitialData,
+	uint64 InitialDataSize,
+	ERDGInitialDataFlags InitialDataFlags)
+{
+	FRDGBufferRef Buffer = GraphBuilder.CreateBuffer(FRDGBufferDesc::CreateByteAddressDesc(NumBytes), Name);
+	GraphBuilder.QueueBufferUpload(Buffer, InitialData, InitialDataSize, InitialDataFlags);
+	return Buffer;
+}
+
+FRDGBufferRef CreateByteAddressBuffer(
+	FRDGBuilder& GraphBuilder,
+	const TCHAR* Name,
+	FRDGBufferNumElementsCallback&& NumElementsCallback,
+	FRDGBufferInitialDataCallback&& InitialDataCallback,
+	FRDGBufferInitialDataSizeCallback&& InitialDataSizeCallback)
+{
+	FRDGBufferRef Buffer = GraphBuilder.CreateBuffer(FRDGBufferDesc::CreateByteAddressDesc(4), Name, MoveTemp(NumElementsCallback));
+	GraphBuilder.QueueBufferUpload(Buffer, MoveTemp(InitialDataCallback), MoveTemp(InitialDataSizeCallback));
+	return Buffer;
+}
+
 FRDGBufferRef CreateUploadBuffer(
 	FRDGBuilder& GraphBuilder,
 	const TCHAR* Name,
diff --git a/Engine/Source/Runtime/RenderCore/Public/RenderGraphUtils.h b/Engine/Source/Runtime/RenderCore/Public/RenderGraphUtils.h
--- a/Engine/Source/Runtime/RenderCore/Public/RenderGraphUtils.h
+++ b/Engine/Source/Runtime/RenderCore/Public/RenderGraphUtils.h
@@ -966,7 +966,7 @@ FORCEINLINE FRDGBufferRef CreateStructuredBuffer_Impl(
 }
 
 /** Same as the previous function but where the type of the array is automatically inferred, so we can do : 
- *  TArray<FSomeType> Array;
+ *  TArray<FSomeType>& Array = GraphBuilder.AllocArray<...>();
  *  CreateStructuredBuffer(..., [&]() -> auto&{ return Array; });
  */
 template <typename GetArrayRefCallback, typename Type = TInvokeResult_T<GetArrayRefCallback>>
@@ -992,6 +992,107 @@ FORCEINLINE FRDGBufferRef CreateStructuredBuffer(
 	return CreateStructuredBuffer(GraphBuilder, Name, InitialData.GetTypeSize(), InitialData.Num(), InitialData.GetData(), InitialData.GetTotalSize(), ERDGInitialDataFlags::NoCopy);
 }
 
+/** Creates a byte address buffer with initial data by creating an upload pass. */
+RENDERCORE_API FRDGBufferRef CreateByteAddressBuffer(
+	FRDGBuilder& GraphBuilder,
+	const TCHAR* Name,
+	uint32 NumBytes,
+	const void* InitialData,
+	uint64 InitialDataSize,
+	ERDGInitialDataFlags InitialDataFlags = ERDGInitialDataFlags::None);
+
+/** A variant where NumElements, InitialData, and InitialDataSize are supplied through callbacks. This allows creating a buffer with
+ *  information unknown at creation time. Though, data must be ready before the most recent RDG pass that references the buffer
+ *  is executed. For byte address buffers, NumElements must be Size / 4.
+ */
+RENDERCORE_API FRDGBufferRef CreateByteAddressBuffer(
+	FRDGBuilder& GraphBuilder,
+	const TCHAR* Name,
+	FRDGBufferNumElementsCallback&& NumElementsCallback,
+	FRDGBufferInitialDataCallback&& InitialDataCallback,
+	FRDGBufferInitialDataSizeCallback&& InitialDataSizeCallback);
+
+/**
+ * Helper to create a byte address buffer with initial data from a TArray with move semantics, this can be cheaper as it guarantees the lifetimes of the data & permits copy-free upload.
+ */
+template <typename ElementType, typename AllocatorType>
+FORCEINLINE FRDGBufferRef CreateByteAddressBuffer(
+	FRDGBuilder& GraphBuilder,
+	const TCHAR* Name,
+	TArray<ElementType, AllocatorType>&& InitialData)
+{
+	static const uint32 DummyElement = 0;
+	if (InitialData.Num() == 0)
+	{
+		return CreateByteAddressBuffer(GraphBuilder, Name, 4, &DummyElement, 4, ERDGInitialDataFlags::NoCopy);
+	}
+
+	// Create a move-initialized copy of the TArray with RDG lifetime & move the data there.
+	TArray<ElementType, AllocatorType>& UploadData = *GraphBuilder.AllocObject<TArray<ElementType, AllocatorType> >(MoveTemp(InitialData));
+	return CreateByteAddressBuffer(GraphBuilder, Name, UploadData.Num() * UploadData.GetTypeSize(), UploadData.GetData(), UploadData.Num() * UploadData.GetTypeSize(), ERDGInitialDataFlags::NoCopy);
+}
+
+/**
+ * Helper to create a byte address buffer with initial data from a TConstArrayView.
+ */
+template <typename ElementType>
+FORCEINLINE FRDGBufferRef CreateByteAddressBuffer(
+	FRDGBuilder& GraphBuilder,
+	const TCHAR* Name,
+	TConstArrayView<ElementType> InitialData,
+	ERDGInitialDataFlags InitialDataFlags = ERDGInitialDataFlags::None)
+{
+	static const ElementType DummyElement = ElementType();
+	if (InitialData.Num() == 0)
+	{
+		return CreateByteAddressBuffer(GraphBuilder, Name, InitialData.GetTypeSize(), &DummyElement, InitialData.GetTypeSize(), ERDGInitialDataFlags::NoCopy);
+	}
+	return CreateByteAddressBuffer(GraphBuilder, Name, InitialData.Num() * InitialData.GetTypeSize(), InitialData.GetData(), InitialData.Num() * InitialData.GetTypeSize(), InitialDataFlags);
+}
+
+/** A variant where the TArray is supplied through callbacks. This allows creating a buffer with
+ *  information unknown at creation time. Though, data must be ready before the most recent RDG pass that references the buffer
+ *  is executed.
+ */
+template <typename ArrayType>
+FORCEINLINE FRDGBufferRef CreateByteAddressBuffer_Impl(
+	FRDGBuilder& GraphBuilder,
+	const TCHAR* Name,
+	TRDGBufferArrayCallback<ArrayType>&& ArrayCallback)
+{
+	return CreateByteAddressBuffer(GraphBuilder, Name,
+		/*NumElementsCallback = */[ArrayCallback]() { const auto& Array = ArrayCallback(); return Array.Num() * Array.GetTypeSize() / 4; },
+		/*InitialDataCallback = */[ArrayCallback]() { return ArrayCallback().GetData(); },
+		/*InitialDataSizeCallback = */[ArrayCallback]() { const auto& Array = ArrayCallback(); return Array.Num() * Array.GetTypeSize(); });
+}
+
+/** Same as the previous function but where the type of the array is automatically inferred, so we can do : 
+ *  TArray<FSomeType> Array = GraphBuilder.AllocArray<...>();
+ *  CreateByteAddressBuffer(..., [&]() -> auto&{ return Array; });
+ */
+template <typename GetArrayRefCallback, typename Type = TInvokeResult_T<GetArrayRefCallback>>
+FORCEINLINE FRDGBufferRef CreateByteAddressBuffer(
+	FRDGBuilder& GraphBuilder,
+	const TCHAR* Name,
+	GetArrayRefCallback&& ArrayCallback)
+{
+	return CreateByteAddressBuffer_Impl<Type>(GraphBuilder, Name, MoveTemp(ArrayCallback));
+}
+
+template <typename ElementType>
+FORCEINLINE FRDGBufferRef CreateByteAddressBuffer(
+	FRDGBuilder& GraphBuilder,
+	const TCHAR* Name,
+	const FRDGUploadData<ElementType>& InitialData)
+{
+	static const uint32 DummyElement = 0;
+	if (InitialData.Num() == 0)
+	{
+		return CreateByteAddressBuffer(GraphBuilder, Name, 4, &DummyElement, 4, ERDGInitialDataFlags::NoCopy);
+	}
+	return CreateByteAddressBuffer(GraphBuilder, Name, InitialData.Num() * InitialData.GetTypeSize(), InitialData.GetData(), InitialData.GetTotalSize(), ERDGInitialDataFlags::NoCopy);
+}
+
 RENDERCORE_API FRDGBufferRef CreateUploadBuffer(
 	FRDGBuilder& GraphBuilder,
 	const TCHAR* Name,
@@ -1028,10 +1129,10 @@ FORCEINLINE FRDGBufferRef CreateUploadBuffer(
 	TConstArrayView<ElementType> InitialData,
 	ERDGInitialDataFlags InitialDataFlags = ERDGInitialDataFlags::None)
 {
-	static const ElementType DummyElement = ElementType();
+	static const uint32 DummyElement = 0;
 	if (InitialData.Num() == 0)
 	{
-		return CreateUploadBuffer(GraphBuilder, Name, sizeof(ElementType), 1, &DummyElement, sizeof(ElementType), ERDGInitialDataFlags::NoCopy);
+		return CreateUploadBuffer(GraphBuilder, Name, 4, 1, &DummyElement, 4, ERDGInitialDataFlags::NoCopy);
 	}
 	return CreateUploadBuffer(GraphBuilder, Name, sizeof(ElementType), InitialData.Num(), InitialData.GetData(), sizeof(ElementType) * InitialData.Num(), InitialDataFlags);
 }
