Increase inline storage to avoid dynamic allocations due to cached skelmesh MDC

- After 32834511 enabled caching mesh draw commands for skeletal meshes, the number needs to be adjusted to prevent/reduce dynamic allocations.

diff --git a/Engine/Source/Runtime/RenderCore/Public/VertexFactory.h b/Engine/Source/Runtime/RenderCore/Public/VertexFactory.h
--- a/Engine/Source/Runtime/RenderCore/Public/VertexFactory.h
+++ b/Engine/Source/Runtime/RenderCore/Public/VertexFactory.h
@@ -117,10 +117,10 @@ struct FVertexInputStream
 
 /** 
  * Number of vertex input bindings to allocate inline within a FMeshDrawCommand.
- * This is tweaked so that the bindings for FLocalVertexFactory fit into the inline storage.
+ * This is tweaked so that the bindings for commonly cached factories (FLocalVertexFactory, FGPUBaseSkinVertexFactory) fit into the inline storage.
  * Overflow of the inline storage will cause a heap allocation per draw (and corresponding cache miss on traversal)
  */
-typedef TArray<FVertexInputStream, TInlineAllocator<4>> FVertexInputStreamArray;
+typedef TArray<FVertexInputStream, TInlineAllocator<7>> FVertexInputStreamArray;
 
 ENUM_CLASS_FLAGS(EVertexStreamUsage);
 
