Added SRVNonPixel, SHADER_PARAMETER_RDG_NON_PIXEL_SRV, and modified RDG_TEXTURE_ACCESS to support texture subresources.

 - SRVNonPixel is needed by mobile to insert a barrier between fragment -> vertex texture fetch, but since this is a heavyweight barrier, it is opt-in with SHADER_PARAMETER_RDG_NON_PIXEL_SRV.
 - Small refactor to FRDGTextureAccess to allow for arbitrary subresources, as the current model only allows full resource transitions.

diff --git a/Engine/Source/Developer/Apple/MetalShaderFormat/Private/MetalCompileShaderSPIRV.cpp b/Engine/Source/Developer/Apple/MetalShaderFormat/Private/MetalCompileShaderSPIRV.cpp
--- a/Engine/Source/Developer/Apple/MetalShaderFormat/Private/MetalCompileShaderSPIRV.cpp
+++ b/Engine/Source/Developer/Apple/MetalShaderFormat/Private/MetalCompileShaderSPIRV.cpp
@@ -1108,6 +1108,7 @@ void FMetalCompileShaderSPIRV::DoCompileMetalShader(
 							case UBMT_TEXTURE:
 							case UBMT_RDG_TEXTURE:
 							case UBMT_RDG_TEXTURE_SRV:
+							case UBMT_RDG_TEXTURE_NON_PIXEL_SRV:
 							case UBMT_SRV:
 							case UBMT_SAMPLER:
 							case UBMT_RDG_BUFFER_SRV:
@@ -1132,6 +1133,7 @@ void FMetalCompileShaderSPIRV::DoCompileMetalShader(
 										case UBMT_TEXTURE:
 										case UBMT_RDG_TEXTURE:
 										case UBMT_RDG_TEXTURE_SRV:
+										case UBMT_RDG_TEXTURE_NON_PIXEL_SRV:
 										case UBMT_SRV:
 										{
 											std::string typeName = "texture_buffer<float, access::read>";
diff --git a/Engine/Source/Developer/ShaderCompilerCommon/Private/ShaderCompilerCommon.cpp b/Engine/Source/Developer/ShaderCompilerCommon/Private/ShaderCompilerCommon.cpp
--- a/Engine/Source/Developer/ShaderCompilerCommon/Private/ShaderCompilerCommon.cpp
+++ b/Engine/Source/Developer/ShaderCompilerCommon/Private/ShaderCompilerCommon.cpp
@@ -162,6 +162,7 @@ bool BuildResourceTableMapping(
 				break;
 			case UBMT_SRV:
 			case UBMT_RDG_TEXTURE_SRV:
+			case UBMT_RDG_TEXTURE_NON_PIXEL_SRV:
 			case UBMT_RDG_BUFFER_SRV:
 				OutSRT.ShaderResourceViewMap.Add(ResourceMap);
 				break;
diff --git a/Engine/Source/Developer/VulkanShaderFormat/Private/VulkanShaderCompiler.cpp b/Engine/Source/Developer/VulkanShaderFormat/Private/VulkanShaderCompiler.cpp
--- a/Engine/Source/Developer/VulkanShaderFormat/Private/VulkanShaderCompiler.cpp
+++ b/Engine/Source/Developer/VulkanShaderFormat/Private/VulkanShaderCompiler.cpp
@@ -868,6 +868,7 @@ static FShaderResourceTable BuildSRTFromHeader(const FVulkanShaderHeader& NEWHea
 					break;
 				case UBMT_SRV:
 				case UBMT_RDG_TEXTURE_SRV:
+				case UBMT_RDG_TEXTURE_NON_PIXEL_SRV:
 				case UBMT_RDG_BUFFER_SRV:
 					GenericSRT.ShaderResourceViewMap.Add(ResourceMap);
 					break;
diff --git a/Engine/Source/Runtime/RHI/Private/RHIStrings.cpp b/Engine/Source/Runtime/RHI/Private/RHIStrings.cpp
--- a/Engine/Source/Runtime/RHI/Private/RHIStrings.cpp
+++ b/Engine/Source/Runtime/RHI/Private/RHIStrings.cpp
@@ -687,6 +687,8 @@ const TCHAR* GetUniformBufferBaseTypeString(EUniformBufferBaseType BaseType)
 		return TEXT("UBMT_RDG_TEXTURE_ACCESS_ARRAY");
 	case UBMT_RDG_TEXTURE_SRV:
 		return TEXT("UBMT_RDG_TEXTURE_SRV");
+	case UBMT_RDG_TEXTURE_NON_PIXEL_SRV:
+		return TEXT("UBMT_RDG_TEXTURE_NON_PIXEL_SRV");
 	case UBMT_RDG_TEXTURE_UAV:
 		return TEXT("UBMT_RDG_TEXTURE_UAV");
 	case UBMT_RDG_BUFFER_ACCESS:
diff --git a/Engine/Source/Runtime/RHI/Public/RHIAccess.h b/Engine/Source/Runtime/RHI/Public/RHIAccess.h
--- a/Engine/Source/Runtime/RHI/Public/RHIAccess.h
+++ b/Engine/Source/Runtime/RHI/Public/RHIAccess.h
@@ -14,41 +14,45 @@ enum class ERHIAccess : uint32
 	Unknown = 0,
 
 	// Read states
-	CPURead             	= 1 <<  0,
-	Present             	= 1 <<  1,
-	IndirectArgs        	= 1 <<  2,
-	VertexOrIndexBuffer 	= 1 <<  3,
-	SRVCompute          	= 1 <<  4,
-	SRVGraphics         	= 1 <<  5,
-	CopySrc             	= 1 <<  6,
-	ResolveSrc          	= 1 <<  7,
-	DSVRead					= 1 <<  8,
+	CPURead                 = 1 <<  0,
+	Present                 = 1 <<  1,
+	IndirectArgs            = 1 <<  2,
+	VertexOrIndexBuffer     = 1 <<  3,
+	SRVCompute              = 1 <<  4,
+	SRVGraphicsPixel        = 1 <<  5,
+	SRVGraphicsNonPixel     = 1 <<  6,
+	CopySrc                 = 1 <<  7,
+	ResolveSrc              = 1 <<  8,
+	DSVRead                 = 1 <<  9,
 
 	// Read-write states
-	UAVCompute          	= 1 <<  9,
-	UAVGraphics         	= 1 << 10,
-	RTV                 	= 1 << 11,
-	CopyDest            	= 1 << 12,
-	ResolveDst          	= 1 << 13,
-	DSVWrite            	= 1 << 14,
+	UAVCompute              = 1 << 10,
+	UAVGraphics             = 1 << 11,
+	RTV                     = 1 << 12,
+	CopyDest                = 1 << 13,
+	ResolveDst              = 1 << 14,
+	DSVWrite                = 1 << 15,
 
 	// Ray tracing acceleration structure states.
 	// Buffer that contains an AS must always be in either of these states.
 	// BVHRead -- required for AS inputs to build/update/copy/trace commands.
 	// BVHWrite -- required for AS outputs of build/update/copy commands.
-	BVHRead                  = 1 << 15,
-	BVHWrite                 = 1 << 16,
+	BVHRead                  = 1 << 16,
+	BVHWrite                 = 1 << 17,
 
 	// Invalid released state (transient resources)
-	Discard					= 1 << 17,
+	Discard					= 1 << 18,
 
 	// Shading Rate Source
-	ShadingRateSource	= 1 << 18,
+	ShadingRateSource	= 1 << 19,
 
 	Last = ShadingRateSource,
 	None = Unknown,
 	Mask = (Last << 1) - 1,
 
+	// Graphics is a combination of pixel and non-pixel
+	SRVGraphics = SRVGraphicsPixel | SRVGraphicsNonPixel,
+
 	// A mask of the two possible SRV states
 	SRVMask = SRVCompute | SRVGraphics,
 
diff --git a/Engine/Source/Runtime/RHI/Public/RHIDefinitions.h b/Engine/Source/Runtime/RHI/Public/RHIDefinitions.h
--- a/Engine/Source/Runtime/RHI/Public/RHIDefinitions.h
+++ b/Engine/Source/Runtime/RHI/Public/RHIDefinitions.h
@@ -564,6 +564,7 @@ enum EUniformBufferBaseType : uint8
 	UBMT_RDG_TEXTURE_ACCESS,
 	UBMT_RDG_TEXTURE_ACCESS_ARRAY,
 	UBMT_RDG_TEXTURE_SRV,
+	UBMT_RDG_TEXTURE_NON_PIXEL_SRV,
 	UBMT_RDG_TEXTURE_UAV,
 	UBMT_RDG_BUFFER_ACCESS,
 	UBMT_RDG_BUFFER_ACCESS_ARRAY,
@@ -1320,6 +1321,7 @@ inline bool IsRDGTextureReferenceShaderParameterType(EUniformBufferBaseType Base
 	return
 		BaseType == UBMT_RDG_TEXTURE ||
 		BaseType == UBMT_RDG_TEXTURE_SRV ||
+		BaseType == UBMT_RDG_TEXTURE_NON_PIXEL_SRV ||
 		BaseType == UBMT_RDG_TEXTURE_UAV ||
 		BaseType == UBMT_RDG_TEXTURE_ACCESS ||
 		BaseType == UBMT_RDG_TEXTURE_ACCESS_ARRAY;
diff --git a/Engine/Source/Runtime/RHICore/Internal/RHIUniformBufferDataShared.h b/Engine/Source/Runtime/RHICore/Internal/RHIUniformBufferDataShared.h
--- a/Engine/Source/Runtime/RHICore/Internal/RHIUniformBufferDataShared.h
+++ b/Engine/Source/Runtime/RHICore/Internal/RHIUniformBufferDataShared.h
@@ -56,6 +56,7 @@ namespace UE::RHICore
 		}
 		break;
 		case UBMT_RDG_TEXTURE_SRV:
+		case UBMT_RDG_TEXTURE_NON_PIXEL_SRV:
 		case UBMT_RDG_BUFFER_SRV:
 		{
 			FRDGShaderResourceView* RDGShaderResourceView = Reader.Read<FRDGShaderResourceView*>(Resource);
diff --git a/Engine/Source/Runtime/RHICore/Public/RHICoreShader.h b/Engine/Source/Runtime/RHICore/Public/RHICoreShader.h
--- a/Engine/Source/Runtime/RHICore/Public/RHICoreShader.h
+++ b/Engine/Source/Runtime/RHICore/Public/RHICoreShader.h
@@ -209,7 +209,9 @@ void SetResourcesFromTables(TBinder&& Binder, FRHIShader const& Shader, TBitMask
 				{
 					ERHIAccess Access = IsComputeShaderFrequency(Shader.GetFrequency())
 						? ERHIAccess::SRVCompute
-						: ERHIAccess::SRVGraphics;
+						: Shader.GetFrequency() == SF_Pixel
+							? ERHIAccess::SRVGraphicsPixel
+							: ERHIAccess::SRVGraphicsNonPixel;
 
 					// Textures bound here only have their "common" plane accessible. Stencil etc is ignored.
 					// (i.e. only access the color plane of a color texture, or depth plane of a depth texture)
@@ -229,7 +231,9 @@ void SetResourcesFromTables(TBinder&& Binder, FRHIShader const& Shader, TBitMask
 				{
 					ERHIAccess Access = IsComputeShaderFrequency(Shader.GetFrequency())
 						? ERHIAccess::SRVCompute
-						: ERHIAccess::SRVGraphics;
+						: Shader.GetFrequency() == SF_Pixel
+							? ERHIAccess::SRVGraphicsPixel
+							: ERHIAccess::SRVGraphicsNonPixel;
 
 					Tracker->Assert(SRV->GetViewIdentity(), Access);
 				}
@@ -268,7 +272,9 @@ void SetResourcesFromTables(TBinder&& Binder, FRHIShader const& Shader, TBitMask
 				{
 					ERHIAccess Access = IsComputeShaderFrequency(Shader.GetFrequency())
 						? ERHIAccess::UAVCompute
-						: ERHIAccess::UAVGraphics;
+						: Shader.GetFrequency() == SF_Pixel
+							? ERHIAccess::SRVGraphicsPixel
+							: ERHIAccess::SRVGraphicsNonPixel;
 
 					Tracker->AssertUAV(UAV, Access, Index);
 				}
diff --git a/Engine/Source/Runtime/RenderCore/Private/DumpGPU.cpp b/Engine/Source/Runtime/RenderCore/Private/DumpGPU.cpp
--- a/Engine/Source/Runtime/RenderCore/Private/DumpGPU.cpp
+++ b/Engine/Source/Runtime/RenderCore/Private/DumpGPU.cpp
@@ -2471,6 +2471,7 @@ void FRDGBuilder::DumpResourcePassOutputs(const FRDGPass* Pass)
 		}
 		break;
 		case UBMT_RDG_TEXTURE_SRV:
+		case UBMT_RDG_TEXTURE_NON_PIXEL_SRV:
 		{
 			if (FRDGTextureSRVRef SRV = Parameter.GetAsTextureSRV())
 			{
diff --git a/Engine/Source/Runtime/RenderCore/Private/RenderGraphBuilder.cpp b/Engine/Source/Runtime/RenderCore/Private/RenderGraphBuilder.cpp
--- a/Engine/Source/Runtime/RenderCore/Private/RenderGraphBuilder.cpp
+++ b/Engine/Source/Runtime/RenderCore/Private/RenderGraphBuilder.cpp
@@ -116,7 +116,7 @@ void EnumerateTextureAccess(FRDGParameterStruct PassParameters, ERDGPassFlags Pa
 		{
 			if (FRDGTextureAccess TextureAccess = Parameter.GetAsTextureAccess())
 			{
-				AccessFunction(nullptr, TextureAccess.GetTexture(), TextureAccess.GetAccess(), NoneFlags, TextureAccess->GetSubresourceRange());
+				AccessFunction(nullptr, TextureAccess.GetTexture(), TextureAccess.GetAccess(), NoneFlags, TextureAccess.GetSubresourceRange());
 			}
 		}
 		break;
@@ -126,14 +126,16 @@ void EnumerateTextureAccess(FRDGParameterStruct PassParameters, ERDGPassFlags Pa
 
 			for (FRDGTextureAccess TextureAccess : TextureAccessArray)
 			{
-				AccessFunction(nullptr, TextureAccess.GetTexture(), TextureAccess.GetAccess(), NoneFlags, TextureAccess->GetSubresourceRange());
+				AccessFunction(nullptr, TextureAccess.GetTexture(), TextureAccess.GetAccess(), NoneFlags, TextureAccess.GetSubresourceRange());
 			}
 		}
 		break;
 		case UBMT_RDG_TEXTURE_SRV:
+		case UBMT_RDG_TEXTURE_NON_PIXEL_SRV:
 			if (FRDGTextureSRVRef SRV = Parameter.GetAsTextureSRV())
 			{
-				AccessFunction(SRV, SRV->GetParent(), SRVAccess, NoneFlags, SRV->GetSubresourceRange());
+				const ERHIAccess SRVGraphicsNonPixel = Parameter.GetType() == UBMT_RDG_TEXTURE_NON_PIXEL_SRV ? ERHIAccess::SRVGraphicsNonPixel : ERHIAccess::None;
+				AccessFunction(SRV, SRV->GetParent(), SRVAccess | SRVGraphicsNonPixel, NoneFlags, SRV->GetSubresourceRange());
 			}
 		break;
 		case UBMT_RDG_TEXTURE_UAV:
@@ -681,6 +683,7 @@ FRHIUniformBuffer* FRDGBuilder::ConvertToExternalUniformBuffer(FRDGUniformBuffer
 			}
 			break;
 			case UBMT_RDG_TEXTURE_SRV:
+			case UBMT_RDG_TEXTURE_NON_PIXEL_SRV:
 			{
 				ConvertTexture(this, Param.GetAsTextureSRV()->Desc.Texture);
 				ConvertView(Param.GetAsView());
@@ -4342,7 +4345,7 @@ void FRDGBuilder::VisualizePassOutputs(const FRDGPass* Pass)
 			{
 				if (IsWritableAccess(TextureAccess.GetAccess()))
 				{
-					if (TOptional<uint32> CaptureId = GVisualizeTexture.ShouldCapture(TextureAccess->Name, /* MipIndex = */ 0))
+					if (TOptional<uint32> CaptureId = GVisualizeTexture.ShouldCapture(TextureAccess->Name, TextureAccess.GetSubresourceRange().MipIndex))
 					{
 						GVisualizeTexture.CreateContentCapturePass(*this, TextureAccess.GetTexture(), *CaptureId);
 					}
@@ -4358,7 +4361,7 @@ void FRDGBuilder::VisualizePassOutputs(const FRDGPass* Pass)
 			{
 				if (IsWritableAccess(TextureAccess.GetAccess()))
 				{
-					if (TOptional<uint32> CaptureId = GVisualizeTexture.ShouldCapture(TextureAccess->Name, /* MipIndex = */ 0))
+					if (TOptional<uint32> CaptureId = GVisualizeTexture.ShouldCapture(TextureAccess->Name, TextureAccess.GetSubresourceRange().MipIndex))
 					{
 						GVisualizeTexture.CreateContentCapturePass(*this, TextureAccess.GetTexture(), *CaptureId);
 					}
diff --git a/Engine/Source/Runtime/RenderCore/Private/RenderGraphUtils.cpp b/Engine/Source/Runtime/RenderCore/Private/RenderGraphUtils.cpp
--- a/Engine/Source/Runtime/RenderCore/Private/RenderGraphUtils.cpp
+++ b/Engine/Source/Runtime/RenderCore/Private/RenderGraphUtils.cpp
@@ -32,6 +32,7 @@ void ClearUnusedGraphResourcesImpl(
 
 		if (Type == UBMT_RDG_TEXTURE ||
 			Type == UBMT_RDG_TEXTURE_SRV ||
+			Type == UBMT_RDG_TEXTURE_NON_PIXEL_SRV ||
 			Type == UBMT_RDG_TEXTURE_UAV ||
 			Type == UBMT_RDG_BUFFER_SRV ||
 			Type == UBMT_RDG_BUFFER_UAV)
@@ -116,6 +117,7 @@ void ClearUnusedGraphResourcesImpl(
 
 		if (Type == UBMT_RDG_TEXTURE ||
 			Type == UBMT_RDG_TEXTURE_SRV ||
+			Type == UBMT_RDG_TEXTURE_NON_PIXEL_SRV ||
 			Type == UBMT_RDG_TEXTURE_UAV ||
 			Type == UBMT_RDG_BUFFER_SRV ||
 			Type == UBMT_RDG_BUFFER_UAV)
diff --git a/Engine/Source/Runtime/RenderCore/Private/RenderGraphValidation.cpp b/Engine/Source/Runtime/RenderCore/Private/RenderGraphValidation.cpp
--- a/Engine/Source/Runtime/RenderCore/Private/RenderGraphValidation.cpp
+++ b/Engine/Source/Runtime/RenderCore/Private/RenderGraphValidation.cpp
@@ -835,6 +835,7 @@ void FRDGUserValidation::ValidateAddPass(const FRDGPass* Pass, bool bSkipPassAcc
 		}
 		break;
 		case UBMT_RDG_TEXTURE_SRV:
+		case UBMT_RDG_TEXTURE_NON_PIXEL_SRV:
 		{
 			if (FRDGTextureSRVRef SRV = Parameter.GetAsTextureSRV())
 			{
diff --git a/Engine/Source/Runtime/RenderCore/Private/ShaderParameterMetadata.cpp b/Engine/Source/Runtime/RenderCore/Private/ShaderParameterMetadata.cpp
--- a/Engine/Source/Runtime/RenderCore/Private/ShaderParameterMetadata.cpp
+++ b/Engine/Source/Runtime/RenderCore/Private/ShaderParameterMetadata.cpp
@@ -195,6 +195,7 @@ const TCHAR* const kShaderParameterMacroNames[] = {
 	TEXT("RDG_TEXTURE_ACCESS"), // UBMT_RDG_TEXTURE_ACCESS,
 	TEXT("RDG_TEXTURE_ACCESS_ARRAY"), // UBMT_RDG_TEXTURE_ACCESS,
 	TEXT("SHADER_PARAMETER_RDG_TEXTURE_SRV"), // UBMT_RDG_TEXTURE_SRV,
+	TEXT("SHADER_PARAMETER_RDG_TEXTURE_NON_PIXEL_SRV"), // UBMT_RDG_TEXTURE_NON_PIXEL_SRV,
 	TEXT("SHADER_PARAMETER_RDG_TEXTURE_UAV"), // UBMT_RDG_TEXTURE_UAV,
 	TEXT("RDG_BUFFER_ACCESS"), // UBMT_RDG_BUFFER_ACCESS,
 	TEXT("RDG_BUFFER_ACCESS_ARRAY"), // UBMT_RDG_BUFFER_ACCESS_ARRAY,
@@ -664,13 +665,15 @@ void FShaderParametersMetadata::InitializeLayout(FRHIUniformBufferLayoutInitiali
 						BaseType == UBMT_TEXTURE ||
 						BaseType == UBMT_SRV ||
 						BaseType == UBMT_RDG_TEXTURE ||
-						BaseType == UBMT_RDG_TEXTURE_SRV);
+						BaseType == UBMT_RDG_TEXTURE_SRV ||
+						BaseType == UBMT_RDG_TEXTURE_NON_PIXEL_SRV);
 				}
 				else if (BindingType == EShaderCodeResourceBindingType::TextureMetadata)
 				{
 					bIsValidBindingType = (
 						BaseType == UBMT_SRV ||
-						BaseType == UBMT_RDG_TEXTURE_SRV);
+						BaseType == UBMT_RDG_TEXTURE_SRV ||
+						BaseType == UBMT_RDG_TEXTURE_NON_PIXEL_SRV);
 				}
 				else if (
 					BindingType == EShaderCodeResourceBindingType::Buffer ||
diff --git a/Engine/Source/Runtime/RenderCore/Private/ShaderParameterStruct.cpp b/Engine/Source/Runtime/RenderCore/Private/ShaderParameterStruct.cpp
--- a/Engine/Source/Runtime/RenderCore/Private/ShaderParameterStruct.cpp
+++ b/Engine/Source/Runtime/RenderCore/Private/ShaderParameterStruct.cpp
@@ -6,6 +6,18 @@
 
 #include "ShaderParameterStruct.h"
 
+FRDGTextureAccess::FRDGTextureAccess(FRDGTexture* InTexture, ERHIAccess InAccess)
+	: FRDGTextureAccess(InTexture, InTexture->GetSubresourceRange(), InAccess)
+{}
+
+FRDGTextureAccess::FRDGTextureAccess(FRDGTextureSRV* InTextureSRV, ERHIAccess InAccess)
+	: FRDGTextureAccess(InTextureSRV->GetParent(), InTextureSRV->GetSubresourceRange(), InAccess)
+{}
+
+FRDGTextureAccess::FRDGTextureAccess(FRDGTextureUAV* InTextureUAV, ERHIAccess InAccess)
+	: FRDGTextureAccess(InTextureUAV->GetParent(), InTextureUAV->GetSubresourceRange(), InAccess)
+{}
+
 /** Context of binding a map. */
 struct FShaderParameterStructBindingContext
 {
@@ -533,6 +545,7 @@ void ValidateShaderParameters(const TShaderRef<FShader>& Shader, const FShaderPa
 				break;
 			}
 			case UBMT_RDG_TEXTURE_SRV:
+			case UBMT_RDG_TEXTURE_NON_PIXEL_SRV:
 			case UBMT_RDG_BUFFER_SRV:
 			case UBMT_RDG_TEXTURE_UAV:
 			case UBMT_RDG_BUFFER_UAV:
@@ -582,6 +595,7 @@ void ValidateShaderParameterResourcesRHI(const void* Contents, const FRHIUniform
 		const bool bSRV =
 			Parameter.MemberType == UBMT_SRV ||
 			Parameter.MemberType == UBMT_RDG_TEXTURE_SRV ||
+			Parameter.MemberType == UBMT_RDG_TEXTURE_NON_PIXEL_SRV ||
 			Parameter.MemberType == UBMT_RDG_BUFFER_SRV;
 
 		// Allow null SRV's in uniform buffers for feature levels that don't support SRV's in shaders
@@ -661,6 +675,7 @@ FRHIShaderParameterResource ExtractShaderParameterResource(FShaderParameterReade
 		return FRHIShaderParameterResource(RDGTexture->GetRHI(), GetParameterIndex(Parameter));
 	}
 	case UBMT_RDG_TEXTURE_SRV:
+	case UBMT_RDG_TEXTURE_NON_PIXEL_SRV:
 	case UBMT_RDG_BUFFER_SRV:
 	{
 		FRDGShaderResourceView* RDGShaderResourceView = Reader.Read<FRDGShaderResourceView*>(Parameter);
@@ -858,6 +873,7 @@ void SetShaderParameters(
 		}
 		break;
 		case UBMT_RDG_TEXTURE_SRV:
+		case UBMT_RDG_TEXTURE_NON_PIXEL_SRV:
 		case UBMT_RDG_BUFFER_SRV:
 		{
 			FRDGShaderResourceView* RDGShaderResourceView = Reader.Read<FRDGShaderResourceView*>(Parameter);
diff --git a/Engine/Source/Runtime/RenderCore/Public/RenderGraphParameter.h b/Engine/Source/Runtime/RenderCore/Public/RenderGraphParameter.h
--- a/Engine/Source/Runtime/RenderCore/Public/RenderGraphParameter.h
+++ b/Engine/Source/Runtime/RenderCore/Public/RenderGraphParameter.h
@@ -18,7 +18,7 @@ public:
 
 	bool IsSRV() const
 	{
-		return MemberType == UBMT_RDG_TEXTURE_SRV || MemberType == UBMT_RDG_BUFFER_SRV;
+		return MemberType == UBMT_RDG_TEXTURE_SRV || MemberType == UBMT_RDG_TEXTURE_NON_PIXEL_SRV || MemberType == UBMT_RDG_BUFFER_SRV;
 	}
 
 	bool IsUAV() const
@@ -162,7 +162,7 @@ public:
 
 	FRDGTextureSRVRef GetAsTextureSRV() const
 	{
-		check(MemberType == UBMT_RDG_TEXTURE_SRV);
+		check(MemberType == UBMT_RDG_TEXTURE_SRV || MemberType == UBMT_RDG_TEXTURE_NON_PIXEL_SRV);
 		return *GetAs<FRDGTextureSRVRef>();
 	}
 
diff --git a/Engine/Source/Runtime/RenderCore/Public/ShaderParameterMacros.h b/Engine/Source/Runtime/RenderCore/Public/ShaderParameterMacros.h
--- a/Engine/Source/Runtime/RenderCore/Public/ShaderParameterMacros.h
+++ b/Engine/Source/Runtime/RenderCore/Public/ShaderParameterMacros.h
@@ -15,6 +15,7 @@
 #include "Traits/IsCharEncodingCompatibleWith.h"
 #include "Misc/AssertionMacros.h"
 #include "RHICommandList.h"
+#include "RenderGraphTextureSubresource.h"
 
 PRAGMA_DISABLE_BUFFER_OVERRUN_WARNING
 
@@ -422,13 +423,20 @@ class alignas(SHADER_PARAMETER_POINTER_ALIGNMENT) FRDGTextureAccess
 {
 public:
 	FRDGTextureAccess() = default;
-	FRDGTextureAccess(FRDGTexture* InTexture, ERHIAccess InAccess)
+
+	FRDGTextureAccess(FRDGTexture* InTexture, FRDGTextureSubresourceRange InSubresourceRange, ERHIAccess InAccess)
 		: Texture(InTexture)
+		, SubresourceRange(InSubresourceRange)
 		, Access(InAccess)
 	{}
 
+	RENDERCORE_API FRDGTextureAccess(FRDGTexture* InTexture, ERHIAccess InAccess);
+	RENDERCORE_API FRDGTextureAccess(FRDGTextureSRV* InTextureSRV, ERHIAccess InAccess);
+	RENDERCORE_API FRDGTextureAccess(FRDGTextureUAV* InTextureUAV, ERHIAccess InAccess);
+
 	FORCEINLINE FRDGTexture* GetTexture() const { return Texture; }
-	FORCEINLINE ERHIAccess   GetAccess() const  { return Access; }
+	FORCEINLINE FRDGTextureSubresourceRange GetSubresourceRange() const { return SubresourceRange; }
+	FORCEINLINE ERHIAccess GetAccess() const  { return Access; }
 
 	FORCEINLINE operator bool() const
 	{
@@ -458,7 +466,8 @@ public:
 
 private:
 	TAlignedShaderParameterPtr<FRDGTexture*> Texture = nullptr;
-	ERHIAccess   Access  = ERHIAccess::Unknown;
+	FRDGTextureSubresourceRange SubresourceRange;
+	ERHIAccess Access = ERHIAccess::Unknown;
 };
 
 template <ERHIAccess InAccess>
@@ -468,13 +477,23 @@ class alignas(SHADER_PARAMETER_POINTER_ALIGNMENT) TRDGTextureAccess
 public:
 	static_assert(IsValidAccess(InAccess), "Texture access is invalid.");
 
-	TRDGTextureAccess()
-		: FRDGTextureAccess(nullptr, InAccess)
+	TRDGTextureAccess() = default;
+
+	TRDGTextureAccess(FRDGTexture* InTexture, FRDGTextureSubresourceRange InSubresourceRange)
+		: FRDGTextureAccess(InTexture, InSubresourceRange, InAccess)
 	{}
 
 	TRDGTextureAccess(FRDGTexture* InTexture)
 		: FRDGTextureAccess(InTexture, InAccess)
 	{}
+
+	TRDGTextureAccess(FRDGTextureSRV* InTextureSRV)
+		: FRDGTextureAccess(InTextureSRV, InAccess)
+	{}
+
+	TRDGTextureAccess(FRDGTextureUAV* InTextureUAV)
+		: FRDGTextureAccess(InTextureUAV, InAccess)
+	{}
 };
 
 template <typename ResourceAccessType>
@@ -1130,7 +1149,7 @@ struct TShaderParameterTypeInfo<FMatrix44f>
 	static const FShaderParametersMetadata* GetStructMetadata() { return nullptr; }
 };
 
-template <typename ResourceAccessType>
+template <typename BufferAccessType>
 struct TRDGResourceAccessTypeInfo
 {
 	static constexpr int32 NumRows = 1;
@@ -1139,13 +1158,30 @@ struct TRDGResourceAccessTypeInfo
 	static constexpr int32 Alignment = SHADER_PARAMETER_POINTER_ALIGNMENT;
 	static constexpr bool bIsStoredInConstantBuffer = false;
 
-	using TAlignedType = ResourceAccessType;
+	using TAlignedType = BufferAccessType;
 
 	static const FShaderParametersMetadata* GetStructMetadata() { return nullptr; }
 
 	static_assert(sizeof(TAlignedType) == SHADER_PARAMETER_POINTER_ALIGNMENT * 2, "Uniform buffer layout must not be platform dependent.");
 };
 
+template <typename TextureAccessType>
+struct TRDGTextureAccessTypeInfo
+{
+	static constexpr int32 NumRows = 1;
+	static constexpr int32 NumColumns = 1;
+	static constexpr int32 NumElements = 0;
+	static constexpr int32 Alignment = SHADER_PARAMETER_POINTER_ALIGNMENT;
+	static constexpr bool bIsStoredInConstantBuffer = false;
+
+	using TAlignedType = TextureAccessType;
+
+	static const FShaderParametersMetadata* GetStructMetadata() { return nullptr; }
+
+	static_assert(sizeof(TAlignedType) == SHADER_PARAMETER_POINTER_ALIGNMENT * 3, "Uniform buffer layout must not be platform dependent.");
+};
+
+
 template<typename T, size_t InNumElements>
 struct TShaderParameterTypeInfo<T[InNumElements]>
 {
@@ -1673,6 +1709,18 @@ private:
 #define SHADER_PARAMETER_RDG_TEXTURE_SRV_ARRAY(ShaderType,MemberName, ArrayDecl) \
 	INTERNAL_SHADER_PARAMETER_EXPLICIT(UBMT_RDG_TEXTURE_SRV, TShaderResourceParameterTypeInfo<FRDGTextureSRV* ArrayDecl>, FRDGTextureSRV*,MemberName,ArrayDecl,,EShaderPrecisionModifier::Float,TEXT(#ShaderType),false)
 
+/** Adds a non-pixel shader resource view for a render graph tracked texture.
+ *
+ * Example:
+ *	SHADER_PARAMETER_RDG_TEXTURE_NON_PIXEL_SRV(Texture2D, MySRV)
+ *	SHADER_PARAMETER_RDG_TEXTURE_NON_PIXEL_SRV_ARRAY(Texture2D, MyArrayOfSRVs, [4])
+ */
+#define SHADER_PARAMETER_RDG_TEXTURE_NON_PIXEL_SRV(ShaderType,MemberName) \
+	INTERNAL_SHADER_PARAMETER_EXPLICIT(UBMT_RDG_TEXTURE_NON_PIXEL_SRV, TShaderResourceParameterTypeInfo<FRDGTextureSRV*>, FRDGTextureSRV*,MemberName,, = nullptr,EShaderPrecisionModifier::Float,TEXT(#ShaderType),false)
+
+#define SHADER_PARAMETER_RDG_TEXTURE_NON_PIXEL_SRV_ARRAY(ShaderType,MemberName, ArrayDecl) \
+	INTERNAL_SHADER_PARAMETER_EXPLICIT(UBMT_RDG_TEXTURE_NON_PIXEL_SRV, TShaderResourceParameterTypeInfo<FRDGTextureSRV* ArrayDecl>, FRDGTextureSRV*,MemberName,ArrayDecl,,EShaderPrecisionModifier::Float,TEXT(#ShaderType),false)
+
 /** Adds a unordered access view for a render graph tracked texture.
  *
  * Example:
@@ -1825,10 +1873,10 @@ private:
 
 /** Informs the RDG pass to transition the texture into the requested state. */
 #define RDG_TEXTURE_ACCESS(MemberName, Access) \
-	INTERNAL_SHADER_PARAMETER_EXPLICIT(UBMT_RDG_TEXTURE_ACCESS, TRDGResourceAccessTypeInfo<TRDGTextureAccess<Access>>, TRDGTextureAccess<Access>,MemberName,,,EShaderPrecisionModifier::Float,TEXT(""),false)
+	INTERNAL_SHADER_PARAMETER_EXPLICIT(UBMT_RDG_TEXTURE_ACCESS, TRDGTextureAccessTypeInfo<TRDGTextureAccess<Access>>, TRDGTextureAccess<Access>,MemberName,,,EShaderPrecisionModifier::Float,TEXT(""),false)
 
 #define RDG_TEXTURE_ACCESS_DYNAMIC(MemberName) \
-	INTERNAL_SHADER_PARAMETER_EXPLICIT(UBMT_RDG_TEXTURE_ACCESS, TRDGResourceAccessTypeInfo<FRDGTextureAccess>, FRDGTextureAccess,MemberName,,,EShaderPrecisionModifier::Float,TEXT(""),false)
+	INTERNAL_SHADER_PARAMETER_EXPLICIT(UBMT_RDG_TEXTURE_ACCESS, TRDGTextureAccessTypeInfo<FRDGTextureAccess>, FRDGTextureAccess,MemberName,,,EShaderPrecisionModifier::Float,TEXT(""),false)
 
 #define RDG_TEXTURE_ACCESS_ARRAY(MemberName) \
 	INTERNAL_SHADER_PARAMETER_EXPLICIT(UBMT_RDG_TEXTURE_ACCESS_ARRAY, TRDGResourceAccessTypeInfo<FRDGTextureAccessArray>, FRDGTextureAccessArray,MemberName,,,EShaderPrecisionModifier::Float,TEXT(""),false)
diff --git a/Engine/Source/Runtime/RenderCore/Public/ShaderParameterStruct.h b/Engine/Source/Runtime/RenderCore/Public/ShaderParameterStruct.h
--- a/Engine/Source/Runtime/RenderCore/Public/ShaderParameterStruct.h
+++ b/Engine/Source/Runtime/RenderCore/Public/ShaderParameterStruct.h
@@ -149,6 +149,7 @@ inline void CollectSRVsToUnset(FRHIBatchedShaderUnbinds& BatchedUnbinds, const T
 	{
 		if (Parameter.BaseType == UBMT_SRV ||
 			Parameter.BaseType == UBMT_RDG_TEXTURE_SRV ||
+			Parameter.BaseType == UBMT_RDG_TEXTURE_NON_PIXEL_SRV ||
 			Parameter.BaseType == UBMT_RDG_BUFFER_SRV)
 		{
 			BatchedUnbinds.UnsetSRV(GetParameterIndex(Parameter));
diff --git a/Engine/Source/Runtime/Renderer/Private/PostProcess/PostProcessLensFlares.cpp b/Engine/Source/Runtime/Renderer/Private/PostProcess/PostProcessLensFlares.cpp
--- a/Engine/Source/Runtime/Renderer/Private/PostProcess/PostProcessLensFlares.cpp
+++ b/Engine/Source/Runtime/Renderer/Private/PostProcess/PostProcessLensFlares.cpp
@@ -27,7 +27,7 @@ TAutoConsoleVariable<int32> CVarLensFlareQuality(
 
 // The RDG inputs shared by all lens flare passes.
 BEGIN_SHADER_PARAMETER_STRUCT(FLensFlarePassParameters, )
-	SHADER_PARAMETER_RDG_TEXTURE_SRV(Texture2D, InputTexture)
+	SHADER_PARAMETER_RDG_TEXTURE_NON_PIXEL_SRV(Texture2D, InputTexture)
 	RENDER_TARGET_BINDING_SLOTS()
 END_SHADER_PARAMETER_STRUCT()
 
diff --git a/Engine/Source/Runtime/VulkanRHI/Private/VulkanBarriers.cpp b/Engine/Source/Runtime/VulkanRHI/Private/VulkanBarriers.cpp
--- a/Engine/Source/Runtime/VulkanRHI/Private/VulkanBarriers.cpp
+++ b/Engine/Source/Runtime/VulkanRHI/Private/VulkanBarriers.cpp
@@ -1929,6 +1929,8 @@ VkImageLayout FVulkanLayoutManager::GetDefaultLayout(FVulkanCmdBuffer* CmdBuffer
 	{
 	case ERHIAccess::SRVCompute:
 	case ERHIAccess::SRVGraphics:
+	case ERHIAccess::SRVGraphicsNonPixel:
+	case ERHIAccess::SRVGraphicsPixel:
 	case ERHIAccess::SRVMask:
 	{
 		if (VulkanTexture.IsDepthOrStencilAspect())
