Fixed nullptr access with RDG_TEXTURE_ACCESS.

diff --git a/Engine/Source/Runtime/RenderCore/Private/ShaderParameterStruct.cpp b/Engine/Source/Runtime/RenderCore/Private/ShaderParameterStruct.cpp
--- a/Engine/Source/Runtime/RenderCore/Private/ShaderParameterStruct.cpp
+++ b/Engine/Source/Runtime/RenderCore/Private/ShaderParameterStruct.cpp
@@ -7,15 +7,15 @@
 #include "ShaderParameterStruct.h"
 
 FRDGTextureAccess::FRDGTextureAccess(FRDGTexture* InTexture, ERHIAccess InAccess)
-	: FRDGTextureAccess(InTexture, InTexture->GetSubresourceRange(), InAccess)
+	: FRDGTextureAccess(InTexture, InTexture ? InTexture->GetSubresourceRange() : FRDGTextureSubresourceRange(), InAccess)
 {}
 
 FRDGTextureAccess::FRDGTextureAccess(FRDGTextureSRV* InTextureSRV, ERHIAccess InAccess)
-	: FRDGTextureAccess(InTextureSRV->GetParent(), InTextureSRV->GetSubresourceRange(), InAccess)
+	: FRDGTextureAccess(InTextureSRV ? InTextureSRV->GetParent() : nullptr, InTextureSRV ? InTextureSRV->GetSubresourceRange() : FRDGTextureSubresourceRange(), InAccess)
 {}
 
 FRDGTextureAccess::FRDGTextureAccess(FRDGTextureUAV* InTextureUAV, ERHIAccess InAccess)
-	: FRDGTextureAccess(InTextureUAV->GetParent(), InTextureUAV->GetSubresourceRange(), InAccess)
+	: FRDGTextureAccess(InTextureUAV ? InTextureUAV->GetParent() : nullptr, InTextureUAV ? InTextureUAV->GetSubresourceRange() : FRDGTextureSubresourceRange(), InAccess)
 {}
 
 /** Context of binding a map. */
