Deduplicating missing shader parameter error/message box code and adding a log message when the message box is used.

diff --git a/Engine/Source/Runtime/RenderCore/Private/ShaderParameters.cpp b/Engine/Source/Runtime/RenderCore/Private/ShaderParameters.cpp
--- a/Engine/Source/Runtime/RenderCore/Private/ShaderParameters.cpp
+++ b/Engine/Source/Runtime/RenderCore/Private/ShaderParameters.cpp
@@ -18,21 +18,28 @@ IMPLEMENT_TYPE_LAYOUT(FShaderParameter);
 IMPLEMENT_TYPE_LAYOUT(FShaderResourceParameter);
 IMPLEMENT_TYPE_LAYOUT(FShaderUniformBufferParameter);
 
+static void FailureToBindNonOptionalParameter(const TCHAR* ParameterType, const TCHAR* ParameterName)
+{
+	if (!UE_LOG_ACTIVE(LogShaders, Log))
+	{
+		UE_LOG(LogShaders, Fatal, TEXT("Failure to bind non-optional %s %s!  The parameter is either not present in the shader, or the shader compiler optimized it out."), ParameterType, ParameterName);
+	}
+	else
+	{
+		UE_LOG(LogShaders, Log, TEXT("Failure to bind non-optional %s %s!  The parameter is either not present in the shader, or the shader compiler optimized it out."), ParameterType, ParameterName);
+
+		// We use a non-Slate message box to avoid problem where we haven't compiled the shaders for Slate.
+		FPlatformMisc::MessageBoxExt(EAppMsgType::Ok, *FText::Format(
+			NSLOCTEXT("UnrealEd", "Error_FailedToBindShaderParameter", "Failure to bind non-optional shader parameter {0}! The parameter is either not present in the shader, or the shader compiler optimized it out. This will be an assert with LogShaders suppressed!"),
+			FText::FromString(ParameterName)).ToString(), TEXT("Warning"));
+	}
+}
+
 void FShaderParameter::Bind(const FShaderParameterMap& ParameterMap,const TCHAR* ParameterName,EShaderParameterFlags Flags)
 {
 	if (!ParameterMap.FindParameterAllocation(ParameterName,BufferIndex,BaseIndex,NumBytes) && Flags == SPF_Mandatory)
 	{
-		if (!UE_LOG_ACTIVE(LogShaders, Log))
-		{
-			UE_LOG(LogShaders, Fatal,TEXT("Failure to bind non-optional shader parameter %s!  The parameter is either not present in the shader, or the shader compiler optimized it out."),ParameterName);
-		}
-		else
-		{
-			// We use a non-Slate message box to avoid problem where we haven't compiled the shaders for Slate.
-			FPlatformMisc::MessageBoxExt( EAppMsgType::Ok, *FText::Format(
-				NSLOCTEXT("UnrealEd", "Error_FailedToBindShaderParameter", "Failure to bind non-optional shader parameter {0}! The parameter is either not present in the shader, or the shader compiler optimized it out. This will be an assert with LogShaders suppressed!"),
-				FText::FromString(ParameterName)).ToString(), TEXT("Warning"));
-		}
+		FailureToBindNonOptionalParameter(TEXT("shader parameter"), ParameterName);
 	}
 }
 
@@ -56,17 +63,7 @@ void FShaderResourceParameter::Bind(const FShaderParameterMap& ParameterMap, con
 	}
 	else if (Flags == SPF_Mandatory)
 	{
-		if (!UE_LOG_ACTIVE(LogShaders, Log))
-		{
-			UE_LOG(LogShaders, Fatal,TEXT("Failure to bind non-optional shader resource parameter %s!  The parameter is either not present in the shader, or the shader compiler optimized it out."),ParameterName);
-		}
-		else
-		{
-			// We use a non-Slate message box to avoid problem where we haven't compiled the shaders for Slate.
-			FPlatformMisc::MessageBoxExt( EAppMsgType::Ok, *FText::Format(
-				NSLOCTEXT("UnrealEd", "Error_FailedToBindShaderParameter", "Failure to bind non-optional shader parameter {0}! The parameter is either not present in the shader, or the shader compiler optimized it out. This will be an assert with LogShaders suppressed!"),
-				FText::FromString(ParameterName)).ToString(), TEXT("Warning"));
-		}
+		FailureToBindNonOptionalParameter(TEXT("shader resource parameter"), ParameterName);
 	}
 }
 
@@ -110,17 +107,7 @@ void FShaderUniformBufferParameter::Bind(const FShaderParameterMap& ParameterMap
 		BaseIndex = 0xffff;
 		if(Flags == SPF_Mandatory)
 		{
-			if (!UE_LOG_ACTIVE(LogShaders, Log))
-			{
-				UE_LOG(LogShaders, Fatal,TEXT("Failure to bind non-optional shader resource parameter %s!  The parameter is either not present in the shader, or the shader compiler optimized it out."),ParameterName);
-			}
-			else
-			{
-				// We use a non-Slate message box to avoid problem where we haven't compiled the shaders for Slate.
-				FPlatformMisc::MessageBoxExt( EAppMsgType::Ok, *FText::Format(
-					NSLOCTEXT("UnrealEd", "Error_FailedToBindShaderParameter", "Failure to bind non-optional shader parameter {0}! The parameter is either not present in the shader, or the shader compiler optimized it out. This will be an assert with LogShaders suppressed!"),
-					FText::FromString(ParameterName)).ToString(), TEXT("Warning"));
-			}
+			FailureToBindNonOptionalParameter(TEXT("shader resource parameter"), ParameterName);
 		}
 	}
 	else
