- Don't make failing of global compute PSO precached shaders fatal
- Log info on shader creation failure due to missing RHI features

diff --git a/Engine/Source/Runtime/D3D12RHI/Private/D3D12Shaders.cpp b/Engine/Source/Runtime/D3D12RHI/Private/D3D12Shaders.cpp
--- a/Engine/Source/Runtime/D3D12RHI/Private/D3D12Shaders.cpp
+++ b/Engine/Source/Runtime/D3D12RHI/Private/D3D12Shaders.cpp
@@ -63,6 +63,7 @@ static bool ValidateShaderIsUsable(FD3D12ShaderData* InShader, EShaderFrequency
 #if D3D12RHI_NEEDS_SHADER_FEATURE_CHECKS
 	if ((InFrequency == SF_Mesh || InFrequency == SF_Amplification) && !GRHISupportsMeshShadersTier0)
 	{
+		UE_LOG(LogD3D12RHI, Log, TEXT("Trying to create Mesh or Amplication shader but RHI doesn't support MeshShaders"));
 		return false;
 	}
 
@@ -70,6 +71,7 @@ static bool ValidateShaderIsUsable(FD3D12ShaderData* InShader, EShaderFrequency
 	// but RT shaders are compiled to DXIL and can use wave ops (HW support is guaranteed too).
 	if (EnumHasAnyFlags(InShader->Features, EShaderCodeFeatures::WaveOps) && !GRHISupportsWaveOperations && !IsRayTracingShaderFrequency(InFrequency))
 	{
+		UE_LOG(LogD3D12RHI, Log, TEXT("Trying to create shader with WaveOps but RHI doesn't support WaveOperations"));
 		return false;
 	}
 
@@ -78,17 +80,20 @@ static bool ValidateShaderIsUsable(FD3D12ShaderData* InShader, EShaderFrequency
 		if (GRHIBindlessSupport == ERHIBindlessSupport::Unsupported ||
 			(GRHIBindlessSupport == ERHIBindlessSupport::RayTracingOnly && !IsRayTracingShaderFrequency(InFrequency)))
 		{
+			UE_LOG(LogD3D12RHI, Log, TEXT("Trying to create shader with bindless resources or samplers but RHI doesn't support Bindless"));
 			return false;
 		}
 	}
 
 	if (InFrequency == SF_Pixel && EnumHasAnyFlags(InShader->Features, EShaderCodeFeatures::StencilRef) && !GRHISupportsStencilRefFromPixelShader)
 	{
+		UE_LOG(LogD3D12RHI, Log, TEXT("Trying to create pixel shader with stencil ref but RHI doesn't support StencilRefFromPixelShader"));
 		return false;
 	}
 
 	if (EnumHasAnyFlags(InShader->Features, EShaderCodeFeatures::BarycentricsSemantic) && !GRHIGlobals.SupportsBarycentricsSemantic)
 	{
+		UE_LOG(LogD3D12RHI, Log, TEXT("Trying to create shader with BarycentricsSemantic but RHI doesn't support BarycentricsSemantic"));
 		return false;
 	}
 #endif
diff --git a/Engine/Source/Runtime/Engine/Private/ShaderCompiler/ShaderCompiler.cpp b/Engine/Source/Runtime/Engine/Private/ShaderCompiler/ShaderCompiler.cpp
--- a/Engine/Source/Runtime/Engine/Private/ShaderCompiler/ShaderCompiler.cpp
+++ b/Engine/Source/Runtime/Engine/Private/ShaderCompiler/ShaderCompiler.cpp
@@ -9356,10 +9356,14 @@ void PrecacheComputePipelineStatesForGlobalShaders(EShaderPlatform Platform, con
 			{
 				for (TShaderRef<FShader> GlobalShader : ComputeShadersToPrecache)
 				{
+					// PSO precache shaders are not required to all load correctly
+					bool bRequired = false;
 					const TCHAR* TypeName = GlobalShader.GetType()->GetName();
-					FRHIComputeShader* RHIComputeShader = GlobalShader.GetComputeShader();
-					check(RHIComputeShader);
-					PipelineStateCache::PrecacheComputePipelineState(RHIComputeShader, TypeName);
+					FRHIComputeShader* RHIComputeShader = static_cast<FRHIComputeShader*>(GlobalShader.GetRHIShaderBase(SF_Compute, bRequired));
+					if (RHIComputeShader)
+					{
+						PipelineStateCache::PrecacheComputePipelineState(RHIComputeShader, TypeName);
+					}
 				}
 			});
 	}
diff --git a/Engine/Source/Runtime/RenderCore/Private/ShaderCodeLibrary.cpp b/Engine/Source/Runtime/RenderCore/Private/ShaderCodeLibrary.cpp
--- a/Engine/Source/Runtime/RenderCore/Private/ShaderCodeLibrary.cpp
+++ b/Engine/Source/Runtime/RenderCore/Private/ShaderCodeLibrary.cpp
@@ -278,7 +278,7 @@ public:
 
 	// FShaderMapResource interface
 	virtual FSHAHash GetShaderHash(int32 ShaderIndex) override;
-	virtual FRHIShader* CreateRHIShaderOrCrash(int32 ShaderIndex) override;
+	virtual FRHIShader* CreateRHIShaderOrCrash(int32 ShaderIndex, bool bRequired) override;
 	virtual void ReleasePreloadedShaderCode(int32 ShaderIndex) override;
 	virtual bool TryRelease() override;
 	virtual uint32 GetSizeBytes() const override { return sizeof(*this) + GetAllocatedSize(); }
@@ -1203,7 +1203,7 @@ FSHAHash FShaderMapResource_SharedCode::GetShaderHash(int32 ShaderIndex)
 	return LibraryInstance->Library->GetShaderHash(ShaderMapIndex, ShaderIndex);
 }
 
-FRHIShader* FShaderMapResource_SharedCode::CreateRHIShaderOrCrash(int32 ShaderIndex)
+FRHIShader* FShaderMapResource_SharedCode::CreateRHIShaderOrCrash(int32 ShaderIndex, bool bRequired)
 {
 	SCOPED_LOADTIMER(FShaderMapResource_SharedCode_InitRHI);
 #if STATS
@@ -1219,8 +1219,11 @@ FRHIShader* FShaderMapResource_SharedCode::CreateRHIShaderOrCrash(int32 ShaderIn
 	TRefCountPtr<FRHIShader> CreatedShader = LibraryInstance->GetOrCreateShader(LibraryShaderIndex);
 	if (UNLIKELY(CreatedShader == nullptr))
 	{
-		UE_LOG(LogShaders, Fatal, TEXT("FShaderMapResource_SharedCode::InitRHI is unable to create a shader"));
-		// unreachable
+		if (bRequired)
+		{
+			UE_LOG(LogShaders, Fatal, TEXT("FShaderMapResource_SharedCode::InitRHI is unable to create a shader"));
+		}
+
 		return nullptr;
 	}
 
diff --git a/Engine/Source/Runtime/RenderCore/Private/ShaderResource.cpp b/Engine/Source/Runtime/RenderCore/Private/ShaderResource.cpp
--- a/Engine/Source/Runtime/RenderCore/Private/ShaderResource.cpp
+++ b/Engine/Source/Runtime/RenderCore/Private/ShaderResource.cpp
@@ -523,17 +523,22 @@ void FShaderMapResource::BeginCreateAllShaders()
 	{
 		for (int32 ShaderIndex = 0; ShaderIndex < Resource->GetNumShaders(); ++ShaderIndex)
 		{
-			Resource->GetShader(ShaderIndex);
+			Resource->GetShader(ShaderIndex, true /*bRequired*/);
 		}
 	});
 }
 
-FRHIShader* FShaderMapResource::CreateShaderOrCrash(int32 ShaderIndex)
+FRHIShader* FShaderMapResource::CreateShaderOrCrash(int32 ShaderIndex, bool bRequired)
 {
 	FRHIShader* Shader = nullptr;
 	// create before taking the lock. This may cause multiple creations, but it's better
 	// than a potential oversubscription deadlock, since CreateShader can spawn async tasks
-	FRHIShader* CreatedShader = CreateRHIShaderOrCrash(ShaderIndex);	// guaranteed to return non-null
+	FRHIShader* CreatedShader = CreateRHIShaderOrCrash(ShaderIndex, bRequired);	// guaranteed to return non-null if required is set
+	if (CreatedShader == nullptr)
+	{
+		check(!bRequired);
+		return nullptr;
+	}
 
 	{
 		// Most shadermaps have <100 shaders, and less than a half of them can be created. 
@@ -593,7 +598,7 @@ FSHAHash FShaderMapResource_InlineCode::GetShaderHash(int32 ShaderIndex)
 	return Code->ShaderHashes[ShaderIndex];
 }
 
-FRHIShader* FShaderMapResource_InlineCode::CreateRHIShaderOrCrash(int32 ShaderIndex)
+FRHIShader* FShaderMapResource_InlineCode::CreateRHIShaderOrCrash(int32 ShaderIndex, bool bRequired)
 {
 	TRACE_CPUPROFILER_EVENT_SCOPE(FShaderMapResource_InlineCode::CreateRHIShaderOrCrash);
 #if STATS
@@ -656,8 +661,10 @@ FRHIShader* FShaderMapResource_InlineCode::CreateRHIShaderOrCrash(int32 ShaderIn
 	}
 	if (UNLIKELY(RHIShader == nullptr))
 	{
-		UE_LOG(LogShaders, Fatal, TEXT("FShaderMapResource_InlineCode::InitRHI is unable to create a shader: frequency=%d, hash=%s."), static_cast<int32>(Frequency), *ShaderHash.ToString());
-		// unreachable
+		if (bRequired)
+		{
+			UE_LOG(LogShaders, Fatal, TEXT("FShaderMapResource_InlineCode::InitRHI is unable to create a shader: frequency=%d, hash=%s."), static_cast<int32>(Frequency), *ShaderHash.ToString());
+		}
 		return nullptr;
 	}
 
diff --git a/Engine/Source/Runtime/RenderCore/Public/Shader.h b/Engine/Source/Runtime/RenderCore/Public/Shader.h
--- a/Engine/Source/Runtime/RenderCore/Public/Shader.h
+++ b/Engine/Source/Runtime/RenderCore/Public/Shader.h
@@ -359,7 +359,7 @@ public:
 		return RHIShaders[ShaderIndex].load(std::memory_order_acquire) != nullptr;
 	}
 
-	inline FRHIShader* GetShader(int32 ShaderIndex)
+	inline FRHIShader* GetShader(int32 ShaderIndex, bool bRequired = true)
 	{
 		// This is a double checked locking. This trickery arises from the fact that we're
 		// synchronizing two threads: one that takes a lock and another that doesn't.
@@ -368,7 +368,7 @@ public:
 		FRHIShader* Shader = RHIShaders[ShaderIndex].load(std::memory_order_acquire);
 		if (UNLIKELY(Shader == nullptr))
 		{
-			Shader = CreateShaderOrCrash(ShaderIndex);
+			Shader = CreateShaderOrCrash(ShaderIndex, bRequired);
 		}
 		return Shader;
 	}
@@ -434,7 +434,7 @@ protected:
 	}
 
 	/** Creates RHI shader, with a reference (so the caller can release). Never returns nullptr (inability to create is Fatal) */
-	virtual FRHIShader*	CreateRHIShaderOrCrash(int32 ShaderIndex) = 0;
+	virtual FRHIShader*	CreateRHIShaderOrCrash(int32 ShaderIndex, bool bRequired) = 0;
 
 	/** Signal the shader library that it can release compressed shader code for a shader that it keeps preloaded in memory. */
 	virtual void ReleasePreloadedShaderCode(int32 ShaderIndex) { /* no-op when not using shader library */ };
@@ -446,7 +446,7 @@ protected:
 private:
 
 	/** Creates an entry in RHIShaders array and registers it among the raytracing libs if needed. Created shader is returned. */
-	RENDERCORE_API FRHIShader* CreateShaderOrCrash(int32 ShaderIndex);
+	RENDERCORE_API FRHIShader* CreateShaderOrCrash(int32 ShaderIndex, bool bRequired);
 
 	/** This lock is to prevent two threads creating the same RHIShaders element. It is only taken if the element is to be created. */
 	FCriticalSection RHIShadersCreationGuard;
@@ -548,7 +548,7 @@ public:
 
 	// FShaderMapResource interface
 	RENDERCORE_API virtual FSHAHash GetShaderHash(int32 ShaderIndex) override;
-	RENDERCORE_API virtual FRHIShader* CreateRHIShaderOrCrash(int32 ShaderIndex) override;
+	RENDERCORE_API virtual FRHIShader* CreateRHIShaderOrCrash(int32 ShaderIndex, bool bRequired) override;
 	virtual uint32 GetSizeBytes() const override { return sizeof(*this) + GetAllocatedSize(); }
 
 	TRefCountPtr<FShaderMapResourceCode> Code;
@@ -1081,13 +1081,18 @@ public:
 
 	inline ShaderType* operator->() const { return ShaderContent; }
 
-	inline FRHIShader* GetRHIShaderBase(EShaderFrequency Frequency) const
+	inline FRHIShader* GetRHIShaderBase(EShaderFrequency Frequency, bool bRequired = true) const
 	{
 		FRHIShader* RHIShader = nullptr;
 		if(ShaderContent)
 		{
 			checkSlow(ShaderContent->GetFrequency() == Frequency);
-			RHIShader = GetResourceChecked().GetShader(ShaderContent->GetResourceIndex());
+			RHIShader = GetResourceChecked().GetShader(ShaderContent->GetResourceIndex(), bRequired);
+			if (RHIShader == nullptr)
+			{
+				UE_LOG(LogShaders, Log, TEXT("Failed to create shader for type %s with resource index %d."), GetType()->GetName(), ShaderContent->GetResourceIndex());
+				return nullptr;
+			}
 			checkSlow(RHIShader->GetFrequency() == Frequency);
 		}
 		return RHIShader;
