[Shaders] Output absolute shader file paths with native platform separator.

#rnx

diff --git a/Engine/Source/Runtime/RenderCore/Private/ShaderDiagnostics.cpp b/Engine/Source/Runtime/RenderCore/Private/ShaderDiagnostics.cpp
--- a/Engine/Source/Runtime/RenderCore/Private/ShaderDiagnostics.cpp
+++ b/Engine/Source/Runtime/RenderCore/Private/ShaderDiagnostics.cpp
@@ -24,6 +24,13 @@ static TAutoConsoleVariable<int32> CVarShaderWarningsFilter(
 );
 
 
+static FString ConvertToNativePlatformAbsolutePath(const TCHAR* InPath)
+{
+	FString Path = IFileManager::Get().ConvertToAbsolutePathForExternalAppForRead(InPath);
+	FPaths::MakePlatformFilename(Path);
+	return Path;
+}
+
 int32 AddAndProcessErrorsForFailedJobFiltered(FShaderCompileJob& CurrentJob, FShaderDiagnosticInfo& OutShaderDiagInfo, const TCHAR* FilterMessage)
 {
 	int32 NumAddedErrors = 0;
@@ -77,26 +84,28 @@ int32 AddAndProcessErrorsForFailedJobFiltered(FShaderCompileJob& CurrentJob, FSh
 				// If we dumped the shader info, add it before the first error string
 				if (!GIsBuildMachine && !bReportedDebugInfo && CurrentJob.Input.DumpDebugInfoPath.Len() > 0)
 				{
-					UniqueErrorPrefix += FString::Printf(TEXT("Shader debug info dumped to: \"%s\"\n"), *CurrentJob.Input.DumpDebugInfoPath);
+					const FString DebugInfoPath = ConvertToNativePlatformAbsolutePath(*CurrentJob.Input.DumpDebugInfoPath);
+					UniqueErrorPrefix += FString::Printf(TEXT("Shader debug info dumped to: \"%s\"\n"), *DebugInfoPath);
 					bReportedDebugInfo = true;
 				}
 
+				const FString ShaderFilePath = ConvertToNativePlatformAbsolutePath(*CurrentError.GetShaderSourceFilePath());
+				const FString ShaderErrorLineString = CurrentError.ErrorLineString.IsEmpty() ? TEXT("0") : *CurrentError.ErrorLineString;
 				if (CurrentJob.Key.ShaderType)
 				{
 					// Construct a path that will enable VS.NET to find the shader file, relative to the solution
-					const FString SolutionPath = FPaths::RootDir();
-					FString ShaderFilePath = IFileManager::Get().ConvertToAbsolutePathForExternalAppForRead(*CurrentError.GetShaderSourceFilePath());
 					UniqueErrorPrefix += FString::Printf(TEXT("%s(%s): Shader %s, Permutation %d, VF %s:\n\t"),
 						*ShaderFilePath,
-						*CurrentError.ErrorLineString,
+						*ShaderErrorLineString,
 						CurrentJob.Key.ShaderType->GetName(),
 						CurrentJob.Key.PermutationId,
 						CurrentJob.Key.VFType ? CurrentJob.Key.VFType->GetName() : TEXT("None"));
 				}
 				else
 				{
-					UniqueErrorPrefix += FString::Printf(TEXT("%s(0): "),
-						*CurrentJob.Input.VirtualSourceFilePath);
+					UniqueErrorPrefix += FString::Printf(TEXT("%s(%s): "),
+						*ShaderFilePath,
+						*ShaderErrorLineString);
 				}
 
 				OutShaderDiagInfo.UniqueErrors.Add(CurrentErrorString);
