Moved an ensure into FRayTracingGeometry::HasPendingBuildRequest().

diff --git a/Engine/Source/Runtime/Engine/Private/InstancedStaticMesh.cpp b/Engine/Source/Runtime/Engine/Private/InstancedStaticMesh.cpp
--- a/Engine/Source/Runtime/Engine/Private/InstancedStaticMesh.cpp
+++ b/Engine/Source/Runtime/Engine/Private/InstancedStaticMesh.cpp
@@ -1640,7 +1640,6 @@ void FInstancedStaticMeshSceneProxy::GetDynamicRayTracingInstances(struct FRayTr
 
 		if (CurrentLODResources.RayTracingGeometry.HasPendingBuildRequest())
 		{
-			ensure(CurrentLODResources.RayTracingGeometry.IsValid());
 			CurrentLODResources.RayTracingGeometry.BoostBuildPriority();
 		}
 		else if (CurrentLODResources.RayTracingGeometry.IsValid())
diff --git a/Engine/Source/Runtime/Engine/Private/StaticMeshRender.cpp b/Engine/Source/Runtime/Engine/Private/StaticMeshRender.cpp
--- a/Engine/Source/Runtime/Engine/Private/StaticMeshRender.cpp
+++ b/Engine/Source/Runtime/Engine/Private/StaticMeshRender.cpp
@@ -1932,7 +1932,6 @@ void FStaticMeshSceneProxy::GetDynamicRayTracingInstances(FRayTracingMaterialGat
 
 			if (CurrentLODResources.RayTracingGeometry.HasPendingBuildRequest())
 			{
-				ensure(CurrentLODResources.RayTracingGeometry.IsValid());
 				CurrentLODResources.RayTracingGeometry.BoostBuildPriority();
 			}
 			else if (CurrentLODResources.RayTracingGeometry.IsValid())
diff --git a/Engine/Source/Runtime/RenderCore/Private/RayTracingGeometry.cpp b/Engine/Source/Runtime/RenderCore/Private/RayTracingGeometry.cpp
--- a/Engine/Source/Runtime/RenderCore/Private/RayTracingGeometry.cpp
+++ b/Engine/Source/Runtime/RenderCore/Private/RayTracingGeometry.cpp
@@ -285,6 +285,18 @@ void FRayTracingGeometry::ReleaseResource()
 	Initializer = FRayTracingGeometryInitializer{};
 }
 
+bool FRayTracingGeometry::HasPendingBuildRequest() const
+{
+	const bool bHasPendingBuildRequest = RayTracingBuildRequestIndex != INDEX_NONE;
+
+	if (bHasPendingBuildRequest)
+	{
+		ensure(IsValid());
+	}
+
+	return bHasPendingBuildRequest;
+}
+
 void FRayTracingGeometry::BoostBuildPriority(float InBoostValue) const
 {
 	check(HasPendingBuildRequest());
diff --git a/Engine/Source/Runtime/RenderCore/Public/RayTracingGeometry.h b/Engine/Source/Runtime/RenderCore/Public/RayTracingGeometry.h
--- a/Engine/Source/Runtime/RenderCore/Public/RayTracingGeometry.h
+++ b/Engine/Source/Runtime/RenderCore/Public/RayTracingGeometry.h
@@ -140,10 +140,8 @@ public:
 	RENDERCORE_API void MakeResident(FRHICommandList& RHICmdList);
 	RENDERCORE_API void Evict();
 	
-	bool HasPendingBuildRequest() const
-	{
-		return RayTracingBuildRequestIndex != INDEX_NONE;
-	}
+	RENDERCORE_API bool HasPendingBuildRequest() const;
+
 	RENDERCORE_API void BoostBuildPriority(float InBoostValue = 0.01f) const;
 
 	// FRenderResource interface
diff --git a/Engine/Source/Runtime/Renderer/Private/PrimitiveSceneInfo.cpp b/Engine/Source/Runtime/Renderer/Private/PrimitiveSceneInfo.cpp
--- a/Engine/Source/Runtime/Renderer/Private/PrimitiveSceneInfo.cpp
+++ b/Engine/Source/Runtime/Renderer/Private/PrimitiveSceneInfo.cpp
@@ -419,7 +419,6 @@ FRayTracingGeometry* FPrimitiveSceneInfo::GetValidStaticRayTracingGeometry(int8&
 	{
 		if (StaticRayTracingGeometries[InOutLODIndex]->HasPendingBuildRequest())
 		{
-			ensure(StaticRayTracingGeometries[InOutLODIndex]->IsValid());
 			StaticRayTracingGeometries[InOutLODIndex]->BoostBuildPriority();
 		}
 		else if (StaticRayTracingGeometries[InOutLODIndex]->IsValid())
