[shaders] change param metadata layout hash to use xxhash64 instead of GetTypeHash (the latter is not intended to be persisted across sessions). speculative fix for param struct layout hash mismatches, but a good change even if it doesn't solve the issue at hand.

diff --git a/Engine/Shaders/Public/ShaderVersion.ush b/Engine/Shaders/Public/ShaderVersion.ush
--- a/Engine/Shaders/Public/ShaderVersion.ush
+++ b/Engine/Shaders/Public/ShaderVersion.ush
@@ -3,4 +3,4 @@
 // Each time the console command is executed it generates a new GUID. As this file is included
 // in Platform.ush (which should be included in any shader) it allows to invalidate the shader DDC.
 // 
-#pragma message("UESHADERMETADATA_VERSION 6D4B37CE-368C-4654-BCF8-BF58084A1941")
+#pragma message("UESHADERMETADATA_VERSION F5BD4964-B08E-431F-B710-DC98AD2EA1FC")
diff --git a/Engine/Source/Runtime/RenderCore/Private/ShaderParameterMetadata.cpp b/Engine/Source/Runtime/RenderCore/Private/ShaderParameterMetadata.cpp
--- a/Engine/Source/Runtime/RenderCore/Private/ShaderParameterMetadata.cpp
+++ b/Engine/Source/Runtime/RenderCore/Private/ShaderParameterMetadata.cpp
@@ -476,42 +476,6 @@ void FShaderParametersMetadata::Append(FShaderKeyGenerator& KeyGen) const
 	KeyGen.Append(LayoutSignature);
 }
 
-void FShaderParametersMetadata::FMember::HashLayout(FMemoryHasherBlake3& Hasher)
-{
-	Hasher << Offset;
-	Hasher << reinterpret_cast<uint8&>(BaseType);
-
-	Hasher.Serialize(const_cast<TCHAR*>(Name), FCString::Strlen(Name));
-	Hasher << NumElements;
-
-	const bool bIsRHIResource = IsShaderParameterTypeReadOnlyRHIResource(BaseType);
-	const bool bIsRDGResource = IsRDGResourceReferenceShaderParameterType(BaseType);
-
-	if (BaseType == UBMT_INT32 ||
-		BaseType == UBMT_UINT32 ||
-		BaseType == UBMT_FLOAT32)
-	{
-		Hasher << reinterpret_cast<uint8&>(Precision);
-		Hasher << NumRows;
-		Hasher << NumColumns;
-	}
-	else if (BaseType == UBMT_INCLUDED_STRUCT || BaseType == UBMT_NESTED_STRUCT)
-	{
-		const_cast<FShaderParametersMetadata*>(Struct)->HashLayout(Hasher);
-	}
-	else if (bIsRHIResource || bIsRDGResource)
-	{
-		Hasher.Serialize(const_cast<TCHAR*>(ShaderType), FCString::Strlen(ShaderType));
-	}
-}
-
-void FShaderParametersMetadata::HashLayout(FMemoryHasherBlake3& SignatureData) 
-{
-	for (FMember& CurrentMember : Members)
-	{
-		CurrentMember.HashLayout(SignatureData);
-	}
-}
 #endif // WITH_EDITOR
 
 void FShaderParametersMetadata::InitializeLayout(FRHIUniformBufferLayoutInitializer* OutLayoutInitializer)
@@ -809,51 +773,11 @@ void FShaderParametersMetadata::InitializeLayout(FRHIUniformBufferLayoutInitiali
 	// Compute the hash of the RHI layout.
 	LayoutInitializer.ComputeHash();
 	
-	// Compute the hash about the entire layout of the structure.
+	// Fast runtime-compatible hash about the entire layout of the structure.
 	{
-		uint32 RootStructureHash = 0;
-		RootStructureHash = HashCombine(RootStructureHash, GetTypeHash(int32(GetSize())));
-
-		for (const FMember& CurrentMember : Members)
-		{
-			EUniformBufferBaseType BaseType = CurrentMember.GetBaseType();
-			const FShaderParametersMetadata* ChildStruct = CurrentMember.GetStructMetadata();
-
-			uint32 MemberHash = 0;
-			MemberHash = HashCombine(MemberHash, GetTypeHash(int32(CurrentMember.GetOffset())));
-			MemberHash = HashCombine(MemberHash, GetTypeHash(uint8(BaseType)));
-			static_assert(EUniformBufferBaseType_NumBits <= 8, "Invalid EUniformBufferBaseType_NumBits");
-			MemberHash = HashCombine(MemberHash, FCrc::Strihash_DEPRECATED(CurrentMember.GetName()));
-			MemberHash = HashCombine(MemberHash, GetTypeHash(int32(CurrentMember.GetNumElements())));
-
-			const bool bIsRHIResource = IsShaderParameterTypeReadOnlyRHIResource(BaseType);
-			const bool bIsRDGResource = IsRDGResourceReferenceShaderParameterType(BaseType);
-
-			if (BaseType == UBMT_INT32 ||
-				BaseType == UBMT_UINT32 ||
-				BaseType == UBMT_FLOAT32)
-			{
-				MemberHash = HashCombine(MemberHash, GetTypeHash(uint8(CurrentMember.GetNumRows())));
-				MemberHash = HashCombine(MemberHash, GetTypeHash(uint8(CurrentMember.GetNumColumns())));
-			}
-			else if (BaseType == UBMT_INCLUDED_STRUCT || BaseType == UBMT_NESTED_STRUCT)
-			{
-				if (!ChildStruct->IsLayoutInitialized())
-				{
-					const_cast<FShaderParametersMetadata*>(ChildStruct)->InitializeLayout();
-				}
-
-				MemberHash = HashCombine(MemberHash, ChildStruct->GetLayoutHash());
-			}
-			else if (bIsRHIResource || bIsRDGResource)
-			{
-				MemberHash = HashCombine(MemberHash, FCrc::Strihash_DEPRECATED(CurrentMember.GetShaderType()));
-			}
-
-			RootStructureHash = HashCombine(RootStructureHash, MemberHash);
-		}
-
-		LayoutHash = RootStructureHash;
+		TMemoryHasher<FXxHash64Builder, FXxHash64> FastHasher;
+		HashLayout(FastHasher);
+		LayoutHash = (uint32)FastHasher.Finalize().Hash; // note: decimating 64 bit hash to 32 bits to avoid data format changes & API deprecation
 	}
 
 	if (UseCase == EUseCase::UniformBuffer)
@@ -864,6 +788,7 @@ void FShaderParametersMetadata::InitializeLayout(FRHIUniformBufferLayoutInitiali
 	Layout = RHICreateUniformBufferLayout(LayoutInitializer);
 
 #if WITH_EDITOR
+	// second stronger hash for use in DDC keys
 	FMemoryHasherBlake3 Hasher;
 	HashLayout(Hasher);
 	LayoutSignature = Hasher.Finalize();
diff --git a/Engine/Source/Runtime/RenderCore/Public/ShaderParameterMetadata.h b/Engine/Source/Runtime/RenderCore/Public/ShaderParameterMetadata.h
--- a/Engine/Source/Runtime/RenderCore/Public/ShaderParameterMetadata.h
+++ b/Engine/Source/Runtime/RenderCore/Public/ShaderParameterMetadata.h
@@ -266,9 +266,36 @@ public:
 
 	private:
 		friend class FShaderParametersMetadata;
-#if WITH_EDITOR
-		void HashLayout(FMemoryHasherBlake3& SignatureData);
-#endif
+		
+		template<typename FHasherType, typename FHashType>
+		void HashLayout(TMemoryHasher<FHasherType, FHashType>& Hasher)
+		{
+			Hasher << Offset;
+			Hasher << reinterpret_cast<uint8&>(BaseType);
+
+			Hasher.Serialize(const_cast<TCHAR*>(Name), FCString::Strlen(Name));
+			Hasher << NumElements;
+
+			const bool bIsRHIResource = IsShaderParameterTypeReadOnlyRHIResource(BaseType);
+			const bool bIsRDGResource = IsRDGResourceReferenceShaderParameterType(BaseType);
+
+			if (BaseType == UBMT_INT32 ||
+				BaseType == UBMT_UINT32 ||
+				BaseType == UBMT_FLOAT32)
+			{
+				Hasher << reinterpret_cast<uint8&>(Precision);
+				Hasher << NumRows;
+				Hasher << NumColumns;
+			}
+			else if (BaseType == UBMT_INCLUDED_STRUCT || BaseType == UBMT_NESTED_STRUCT)
+			{
+				const_cast<FShaderParametersMetadata*>(Struct)->HashLayout(Hasher);
+			}
+			else if (bIsRHIResource || bIsRDGResource)
+			{
+				Hasher.Serialize(const_cast<TCHAR*>(ShaderType), FCString::Strlen(ShaderType));
+			}
+		}
 
 		const TCHAR* Name;
 		const TCHAR* ShaderType;
@@ -491,8 +518,16 @@ private:
 	/** Hash about the entire memory layout of the structure. */
 	uint32 LayoutHash = 0;
 
+	template<typename FHasherType, typename FHashType>
+	void HashLayout(TMemoryHasher<FHasherType, FHashType>& Hasher)
+	{
+		for (FMember& CurrentMember : Members)
+		{
+			CurrentMember.HashLayout(Hasher);
+		}
+	}
+
 #if WITH_EDITOR
-	void HashLayout(FMemoryHasherBlake3& SignatureData);
 	
 	/** Strong persistable hash representing the binary layout of the entire parameter structure */
 	FBlake3Hash LayoutSignature;
