Fix assert when toggling GPU stats
 - The "r.GPUCsvStatsEnabled" cvar can change mid-frame, but AreGPUStatsEnabled() always re-read it, meaning the value can change between BeginFrame and EndFrame.
 - Fix is to cache the value at BeginFrame() time.

diff --git a/Engine/Source/Runtime/RenderCore/Private/ProfilingDebugging/RealtimeGPUProfiler.cpp b/Engine/Source/Runtime/RenderCore/Private/ProfilingDebugging/RealtimeGPUProfiler.cpp
--- a/Engine/Source/Runtime/RenderCore/Private/ProfilingDebugging/RealtimeGPUProfiler.cpp
+++ b/Engine/Source/Runtime/RenderCore/Private/ProfilingDebugging/RealtimeGPUProfiler.cpp
@@ -920,53 +920,67 @@ namespace GpuProfilerTrace
 }
 #endif
 
-void FRealtimeGPUProfiler::BeginFrame(FRHICommandListImmediate& RHICmdList)
-{
-	if (!AreGPUStatsEnabled())
-	{
-		return;
-	}
+static std::atomic<bool> GAreGPUStatsEnabled{ false };
 
-	check(bInBeginEndBlock == false);
-	bInBeginEndBlock = true;
-
-	Frames[WriteBufferIndex]->TimestampCalibrationQuery = new FRHITimestampCalibrationQuery();
-	RHICmdList.CalibrateTimers(Frames[WriteBufferIndex]->TimestampCalibrationQuery);
+bool AreGPUStatsEnabled()
+{
+	return GAreGPUStatsEnabled;
 }
 
-bool AreGPUStatsEnabled()
+void LatchAreGPUStatsEnabled()
 {
-	if (GSupportsTimestampRenderQueries == false || !CVarGPUStatsEnabled.GetValueOnRenderThread())
+	auto GetValue = []()
 	{
-		return false;
-	}
+		if (GSupportsTimestampRenderQueries == false || !CVarGPUStatsEnabled.GetValueOnRenderThread())
+		{
+			return false;
+		}
 
 #if GPUPROFILERTRACE_ENABLED
-	// Force GPU profiler on if Unreal Insights is running
-	if (UE_TRACE_CHANNELEXPR_IS_ENABLED(GpuProfilerTrace::GpuChannel))
-	{
-		return true;
-	}
+		// Force GPU profiler on if Unreal Insights is running
+		if (UE_TRACE_CHANNELEXPR_IS_ENABLED(GpuProfilerTrace::GpuChannel))
+		{
+			return true;
+		}
 #endif
 
 #if STATS 
-	return true;
+		return true;
 #elif !CSV_PROFILER
-	return false;
+		return false;
 #else
 
-	// If we only have CSV stats, only capture if CSV GPU stats are enabled, and we're capturing
-	if (!CVarGPUCsvStatsEnabled.GetValueOnRenderThread())
-	{
-		return false;
-	}
-	if (!FCsvProfiler::Get()->IsCapturing_Renderthread())
+		// If we only have CSV stats, only capture if CSV GPU stats are enabled, and we're capturing
+		if (!CVarGPUCsvStatsEnabled.GetValueOnRenderThread())
+		{
+			return false;
+		}
+		if (!FCsvProfiler::Get()->IsCapturing_Renderthread())
+		{
+			return false;
+		}
+
+		return true;
+#endif
+	};
+
+	GAreGPUStatsEnabled = GetValue();
+}
+
+void FRealtimeGPUProfiler::BeginFrame(FRHICommandListImmediate& RHICmdList)
+{
+	LatchAreGPUStatsEnabled();
+
+	if (!AreGPUStatsEnabled())
 	{
-		return false;
+		return;
 	}
 
-	return true;
-#endif
+	check(bInBeginEndBlock == false);
+	bInBeginEndBlock = true;
+
+	Frames[WriteBufferIndex]->TimestampCalibrationQuery = new FRHITimestampCalibrationQuery();
+	RHICmdList.CalibrateTimers(Frames[WriteBufferIndex]->TimestampCalibrationQuery);
 }
 
 void FRealtimeGPUProfiler::EndFrame(FRHICommandListImmediate& RHICmdList)
