[LLM] Scope filesystem callbacks to shader code into Shader LLM, added more asset tags
#tests PIE

diff --git a/Engine/Source/Runtime/PakFile/Private/IoDispatcherFileBackend.cpp b/Engine/Source/Runtime/PakFile/Private/IoDispatcherFileBackend.cpp
--- a/Engine/Source/Runtime/PakFile/Private/IoDispatcherFileBackend.cpp
+++ b/Engine/Source/Runtime/PakFile/Private/IoDispatcherFileBackend.cpp
@@ -5,6 +5,7 @@
 #include "Misc/Parse.h"
 #include "Misc/ScopeRWLock.h"
 #include "Misc/StringBuilder.h"
+#include "ProfilingDebugging/AssetMetadataTrace.h"
 #include "ProfilingDebugging/CountersTrace.h"
 #include "HAL/PlatformFileManager.h"
 #include "GenericPlatform/GenericPlatformFile.h"
@@ -670,6 +671,8 @@ uint64 FFileIoStoreReader::GetTocAllocatedSize() const
 
 FIoStatus FFileIoStoreReader::Initialize(const TCHAR* InTocFilePath, int32 InOrder)
 {
+	UE_TRACE_METADATA_SCOPE_ASSET_FNAME(FName(InTocFilePath), FName(TEXT("FileIoStoreReader")), FName(InTocFilePath));
+
 	FStringView ContainerPathView(InTocFilePath);
 	if (!ContainerPathView.EndsWith(TEXT(".utoc")))
 	{
diff --git a/Engine/Source/Runtime/RenderCore/Private/ShaderCodeLibrary.cpp b/Engine/Source/Runtime/RenderCore/Private/ShaderCodeLibrary.cpp
--- a/Engine/Source/Runtime/RenderCore/Private/ShaderCodeLibrary.cpp
+++ b/Engine/Source/Runtime/RenderCore/Private/ShaderCodeLibrary.cpp
@@ -30,6 +30,7 @@ ShaderCodeLibrary.cpp: Bound shader state cache implementation.
 #include "Misc/SecureHash.h"
 #include "Misc/StringBuilder.h"
 #include "PipelineFileCache.h"
+#include "ProfilingDebugging/AssetMetadataTrace.h"
 #include "ProfilingDebugging/LoadTimeTracker.h"
 #include "RenderingThread.h"
 #include "Shader.h"
@@ -875,6 +876,9 @@ class FShaderLibraryInstance
 public:
 	static FShaderLibraryInstance* Create(EShaderPlatform InShaderPlatform, const FString& ShaderCodeDir, FString const& InLibraryName)
 	{
+		LLM_SCOPE(ELLMTag::Shaders);
+		UE_TRACE_METADATA_SCOPE_ASSET_FNAME(FName(*InLibraryName), FName(TEXT("ShaderLibraryInstance")),  FName(*TStringBuilder<520>(InPlace, ShaderCodeDir, TEXT("/"), InLibraryName)));
+
 		FRHIShaderLibraryRef Library;
 		FString ShaderCodeDirectory;
 		if (RHISupportsNativeShaderLibraries(InShaderPlatform))
@@ -2826,9 +2830,10 @@ public:
 	}
 
 	void OnPakFileMounted(const UE::ShaderLibrary::Private::FMountedPakFileInfo& MountInfo)
-	{		
+	{
 		if (IsLibraryInitializedForRuntime())
 		{
+			LLM_SCOPE(ELLMTag::Shaders);
 			TArray<TUniqueFunction<void()>> ShaderLibraryStateChanges;
 			{
 				FRWScopeLock WriteLock(NamedLibrariesMutex, SLT_Write);
@@ -4159,6 +4164,7 @@ void FShaderCodeLibrary::OpenPluginShaderLibrary(IPlugin& Plugin, bool bMonolith
 // At runtime, open shader code collection for specified shader platform. Returns true if new code was opened
 bool UE::ShaderLibrary::Private::FNamedShaderLibrary::OpenShaderCode(const FString& ShaderCodeDir, FString const& Library)
 {
+	LLM_SCOPE(ELLMTag::Shaders);
 	// check if any of the components has this content
 	{
 		FRWScopeLock ReadLock(ComponentsMutex, SLT_Write);
