Lumen - allow Metal to run Lumen with SM5 until it won't get solid SM6 support.

diff --git a/Engine/Source/Runtime/RenderCore/Private/RenderUtils.cpp b/Engine/Source/Runtime/RenderCore/Private/RenderUtils.cpp
--- a/Engine/Source/Runtime/RenderCore/Private/RenderUtils.cpp
+++ b/Engine/Source/Runtime/RenderCore/Private/RenderUtils.cpp
@@ -1963,7 +1963,7 @@ bool DoesPlatformSupportLumenGI(EShaderPlatform Platform, bool bSkipProjectCheck
 
 	return (bSkipProjectCheck || bLumenSupported)
 		&& FDataDrivenShaderPlatformInfo::GetSupportsLumenGI(Platform)
-		&& (IsFeatureLevelSupported(Platform, ERHIFeatureLevel::SM6) || Platform == SP_VULKAN_SM5) // SP_VULKAN_SM5 is needed to run on mobile
+		&& (IsFeatureLevelSupported(Platform, ERHIFeatureLevel::SM6) || Platform == SP_VULKAN_SM5 || Platform == SP_METAL_SM5) // Android and Mac can't rely on SM6 yet, but want to run Lumen
 		&& !IsForwardShadingEnabled(Platform);
 }
 
