Added BUNDLE_MODE_* enum defines to shader bundle dispatch shaders

diff --git a/Engine/Source/Runtime/RenderCore/Private/ShaderBundles.cpp b/Engine/Source/Runtime/RenderCore/Private/ShaderBundles.cpp
--- a/Engine/Source/Runtime/RenderCore/Private/ShaderBundles.cpp
+++ b/Engine/Source/Runtime/RenderCore/Private/ShaderBundles.cpp
@@ -30,6 +30,9 @@ void FDispatchShaderBundleCS::ModifyCompilationEnvironment(const FGlobalShaderPe
 
 	OutEnvironment.SetDefine(TEXT("THREADGROUP_SIZEX"), ThreadGroupSizeX);
 	OutEnvironment.SetDefine(TEXT("USE_SHADER_ROOT_CONSTANTS"), RHISupportsShaderRootConstants(Parameters.Platform) ? 1 : 0);
+	OutEnvironment.SetDefine(TEXT("BUNDLE_MODE_CS"), int32(ERHIShaderBundleMode::CS));
+	OutEnvironment.SetDefine(TEXT("BUNDLE_MODE_MSPS"), int32(ERHIShaderBundleMode::MSPS));
+	OutEnvironment.SetDefine(TEXT("BUNDLE_MODE_VSPS"), int32(ERHIShaderBundleMode::VSPS));
 	FGlobalShader::ModifyCompilationEnvironment(Parameters, OutEnvironment);
 }
 
