Added variants of SetShaderParametersMixed* that can directly record into FRHIBatchedShaderParameters instead of taking an RHI command list and using the batched shader parameter scratch

diff --git a/Engine/Source/Runtime/RenderCore/Public/ShaderParameterUtils.h b/Engine/Source/Runtime/RenderCore/Public/ShaderParameterUtils.h
--- a/Engine/Source/Runtime/RenderCore/Public/ShaderParameterUtils.h
+++ b/Engine/Source/Runtime/RenderCore/Public/ShaderParameterUtils.h
@@ -271,25 +271,77 @@ void SetSingleShaderValue(
 // Mixed mode binding utilities
 
 /// Utility to set all legacy and non-legacy parameters for a shader. Requires the shader type to implement SetParameters(FRHIBatchedShaderParameters& BatchedParameters, ...)
-template<typename TRHICmdList, typename TShaderType, typename TShaderTypeRHI, typename... TArguments>
-inline void SetShaderParametersMixed(
-	TRHICmdList& RHICmdList,
+template<typename TShaderType, typename TShaderTypeRHI, typename... TArguments>
+inline void SetBatchedShaderParametersMixed(
+	FRHIBatchedShaderParameters& BatchedParameters,
 	const TShaderRef<TShaderType>& InShader,
 	TShaderTypeRHI* InShaderRHI,
 	const typename TShaderType::FParameters& Parameters,
 	TArguments&&... InArguments)
 {
-	FRHIBatchedShaderParameters& BatchedParameters = RHICmdList.GetScratchShaderParameters();
-
 	// New Style first
 	SetShaderParameters(BatchedParameters, InShader, Parameters);
 
 	// Legacy second
 	InShader->SetParameters(BatchedParameters, Forward<TArguments>(InArguments)...);
+}
 
+/// Utility to set all legacy and non-legacy parameters for a shader. Requires the shader type to implement SetParameters(FRHIBatchedShaderParameters& BatchedParameters, ...)
+template<typename TRHICmdList, typename TShaderType, typename TShaderTypeRHI, typename... TArguments>
+inline void SetShaderParametersMixed(
+	TRHICmdList& RHICmdList,
+	const TShaderRef<TShaderType>& InShader,
+	TShaderTypeRHI* InShaderRHI,
+	const typename TShaderType::FParameters& Parameters,
+	TArguments&&... InArguments)
+{
+	FRHIBatchedShaderParameters& BatchedParameters = RHICmdList.GetScratchShaderParameters();
+	SetBatchedShaderParametersMixed(BatchedParameters, InShader, InShaderRHI, Parameters, Forward<TArguments>(InArguments)...);
 	RHICmdList.SetBatchedShaderParameters(InShaderRHI, BatchedParameters);
 }
 
+/// Utility to set all legacy and non-legacy parameters for a Vertex shader. Requires the shader type to implement SetParameters(FRHIBatchedShaderParameters& BatchedParameters, ...)
+template<typename TShaderType, typename... TArguments>
+inline void SetShaderParametersMixedVS(FRHIBatchedShaderParameters& BatchedParameters, const TShaderRef<TShaderType>& InShader, const typename TShaderType::FParameters& Parameters, TArguments&&... InArguments)
+{
+	SetBatchedShaderParametersMixed(BatchedParameters, InShader, InShader.GetVertexShader(), Parameters, Forward<TArguments>(InArguments)...);
+}
+
+/// Utility to set all legacy and non-legacy parameters for a Mesh shader. Requires the shader type to implement SetParameters(FRHIBatchedShaderParameters& BatchedParameters, ...)
+template<typename TShaderType, typename... TArguments>
+inline void SetShaderParametersMixedMS(FRHIBatchedShaderParameters& BatchedParameters, const TShaderRef<TShaderType>& InShader, const typename TShaderType::FParameters& Parameters, TArguments&&... InArguments)
+{
+	SetBatchedShaderParametersMixed(BatchedParameters, InShader, InShader.GetMeshShader(), Parameters, Forward<TArguments>(InArguments)...);
+}
+
+/// Utility to set all legacy and non-legacy parameters for an Amplification shader. Requires the shader type to implement SetParameters(FRHIBatchedShaderParameters& BatchedParameters, ...)
+template<typename TShaderType, typename... TArguments>
+inline void SetShaderParametersMixedAS(FRHIBatchedShaderParameters& BatchedParameters, const TShaderRef<TShaderType>& InShader, const typename TShaderType::FParameters& Parameters, TArguments&&... InArguments)
+{
+	SetBatchedShaderParametersMixed(BatchedParameters, InShader, InShader.GetAmplificationShader(), Parameters, Forward<TArguments>(InArguments)...);
+}
+
+/// Utility to set all legacy and non-legacy parameters for a Pixel shader. Requires the shader type to implement SetParameters(FRHIBatchedShaderParameters& BatchedParameters, ...)
+template<typename TShaderType, typename... TArguments>
+inline void SetShaderParametersMixedPS(FRHIBatchedShaderParameters& BatchedParameters, const TShaderRef<TShaderType>& InShader, const typename TShaderType::FParameters& Parameters, TArguments&&... InArguments)
+{
+	SetBatchedShaderParametersMixed(BatchedParameters, InShader, InShader.GetPixelShader(), Parameters, Forward<TArguments>(InArguments)...);
+}
+
+/// Utility to set all legacy and non-legacy parameters for a Geometry shader. Requires the shader type to implement SetParameters(FRHIBatchedShaderParameters& BatchedParameters, ...)
+template<typename TShaderType, typename... TArguments>
+inline void SetShaderParametersMixedGS(FRHIBatchedShaderParameters& BatchedParameters, const TShaderRef<TShaderType>& InShader, const typename TShaderType::FParameters& Parameters, TArguments&&... InArguments)
+{
+	SetBatchedShaderParametersMixed(BatchedParameters, InShader, InShader.GetGeometryShader(), Parameters, Forward<TArguments>(InArguments)...);
+}
+
+/// Utility to set all legacy and non-legacy parameters for a Compute shader. Requires the shader type to implement SetParameters(FRHIBatchedShaderParameters& BatchedParameters, ...)
+template<typename TShaderType, typename... TArguments>
+inline void SetShaderParametersMixedCS(FRHIBatchedShaderParameters& BatchedParameters, const TShaderRef<TShaderType>& InShader, const typename TShaderType::FParameters& Parameters, TArguments&&... InArguments)
+{
+	SetBatchedShaderParametersMixed(BatchedParameters, InShader, InShader.GetComputeShader(), Parameters, Forward<TArguments>(InArguments)...);
+}
+
 /// Utility to set all legacy and non-legacy parameters for a Vertex shader. Requires the shader type to implement SetParameters(FRHIBatchedShaderParameters& BatchedParameters, ...)
 template<typename TRHICmdList, typename TShaderType, typename... TArguments>
 inline void SetShaderParametersMixedVS(TRHICmdList& RHICmdList, const TShaderRef<TShaderType>& InShader, const typename TShaderType::FParameters& Parameters, TArguments&&... InArguments)
