HWRT: Add new platform setting for better control over which passes in HWRT are enabled on platforms that support all modes:
* New options:
 - "Disabled" - will turn off cooking and running with ray tracing.
 - "Inline" - will not compile any ray tracing material shaders and disable support for ray tracing shaders at runtime (saving time and memory) only inline passes will work.
 - "Full" - will enable full support for any ray tracing pass using inline or material shaders.

* Set "Inline" as a default mode on those platforms.
* Made barycentrics the new default ray tracing visualization mode because it's supported in all modes.

#tests project with different RayTracingMode configurations

diff --git a/Engine/Source/Developer/TargetPlatform/Private/TargetPlatformBase.cpp b/Engine/Source/Developer/TargetPlatform/Private/TargetPlatformBase.cpp
--- a/Engine/Source/Developer/TargetPlatform/Private/TargetPlatformBase.cpp
+++ b/Engine/Source/Developer/TargetPlatform/Private/TargetPlatformBase.cpp
@@ -59,6 +59,11 @@ bool FTargetPlatformBase::UsesRayTracing() const
 	return CVar ? (CVar->GetInt() != 0) : false;
 }
 
+ERayTracingRuntimeMode FTargetPlatformBase::GetRayTracingMode() const
+{
+	return ERayTracingRuntimeMode::Disabled;
+}
+
 uint32 FTargetPlatformBase::GetSupportedHardwareMask() const
 {
 	return 0;
diff --git a/Engine/Source/Developer/TargetPlatform/Private/TargetPlatformSettingsBase.cpp b/Engine/Source/Developer/TargetPlatform/Private/TargetPlatformSettingsBase.cpp
--- a/Engine/Source/Developer/TargetPlatform/Private/TargetPlatformSettingsBase.cpp
+++ b/Engine/Source/Developer/TargetPlatform/Private/TargetPlatformSettingsBase.cpp
@@ -45,6 +45,11 @@ bool FTargetPlatformSettingsBase::UsesRayTracing() const
 	return CVar ? (CVar->GetInt() != 0) : false;
 }
 
+ERayTracingRuntimeMode FTargetPlatformSettingsBase::GetRayTracingMode() const
+{
+	return UsesRayTracing() ? ERayTracingRuntimeMode::Full : ERayTracingRuntimeMode::Disabled;
+}
+
 uint32 FTargetPlatformSettingsBase::GetSupportedHardwareMask() const
 {
 	return 0;
diff --git a/Engine/Source/Developer/TargetPlatform/Public/Common/TargetPlatformBase.h b/Engine/Source/Developer/TargetPlatform/Public/Common/TargetPlatformBase.h
--- a/Engine/Source/Developer/TargetPlatform/Public/Common/TargetPlatformBase.h
+++ b/Engine/Source/Developer/TargetPlatform/Public/Common/TargetPlatformBase.h
@@ -90,6 +90,8 @@ public:
 
 	TARGETPLATFORM_API virtual bool UsesRayTracing() const override;
 
+	TARGETPLATFORM_API virtual ERayTracingRuntimeMode GetRayTracingMode() const override;
+
 	TARGETPLATFORM_API virtual uint32 GetSupportedHardwareMask() const override;
 
 	TARGETPLATFORM_API virtual EOfflineBVHMode GetStaticMeshOfflineBVHMode() const override;
@@ -804,6 +806,10 @@ public:
 	{
 		return TargetPlatformSettings->UsesRayTracing();
 	}
+	virtual ERayTracingRuntimeMode GetRayTracingMode() const override
+	{
+		return TargetPlatformSettings->GetRayTracingMode();
+	}
 	virtual uint32 GetSupportedHardwareMask() const override
 	{
 		return TargetPlatformSettings->GetSupportedHardwareMask();
diff --git a/Engine/Source/Developer/TargetPlatform/Public/Common/TargetPlatformSettingsBase.h b/Engine/Source/Developer/TargetPlatform/Public/Common/TargetPlatformSettingsBase.h
--- a/Engine/Source/Developer/TargetPlatform/Public/Common/TargetPlatformSettingsBase.h
+++ b/Engine/Source/Developer/TargetPlatform/Public/Common/TargetPlatformSettingsBase.h
@@ -27,6 +27,7 @@ public:
 	TARGETPLATFORM_API virtual bool UsesSelectiveBasePassOutputs() const override;
 	TARGETPLATFORM_API virtual bool UsesDistanceFields() const override;
 	TARGETPLATFORM_API virtual bool UsesRayTracing() const override;
+	TARGETPLATFORM_API virtual ERayTracingRuntimeMode GetRayTracingMode() const override;
 	TARGETPLATFORM_API virtual uint32 GetSupportedHardwareMask() const override;
 	TARGETPLATFORM_API virtual EOfflineBVHMode GetStaticMeshOfflineBVHMode() const override;
 	TARGETPLATFORM_API virtual bool GetStaticMeshOfflineBVHCompression() const override;
diff --git a/Engine/Source/Developer/TargetPlatform/Public/Interfaces/ITargetPlatformSettings.h b/Engine/Source/Developer/TargetPlatform/Public/Interfaces/ITargetPlatformSettings.h
--- a/Engine/Source/Developer/TargetPlatform/Public/Interfaces/ITargetPlatformSettings.h
+++ b/Engine/Source/Developer/TargetPlatform/Public/Interfaces/ITargetPlatformSettings.h
@@ -10,6 +10,13 @@ enum class EOfflineBVHMode
 	MinimizeMemory,
 };
 
+enum class ERayTracingRuntimeMode
+{
+	Disabled,
+	Inline,
+	Full
+};
+
 /**
  * Enumerates features that may be supported by target platforms.
  */
@@ -163,6 +170,11 @@ public:
 	*/
 	virtual bool UsesRayTracing() const = 0;
 
+	/**
+	* Gets ray tracing runtime mode (disabled/inline/hit-lighting).
+	*/
+	virtual ERayTracingRuntimeMode GetRayTracingMode() const = 0;
+
 	/**
 	 * Gets a platform-dependent bitfield describing which hardware generations are supported.
 	 * Applies to platforms with multiple iterations of the device, which may vary in performance.
diff --git a/Engine/Source/Runtime/RenderCore/Private/RenderUtils.cpp b/Engine/Source/Runtime/RenderCore/Private/RenderUtils.cpp
--- a/Engine/Source/Runtime/RenderCore/Private/RenderUtils.cpp
+++ b/Engine/Source/Runtime/RenderCore/Private/RenderUtils.cpp
@@ -696,6 +696,10 @@ RENDERCORE_API ShaderPlatformMaskType GSimpleSkyDiffusePlatformMask;
 // Safe to use to make cook-time decisions, such as whether to compile ray tracing shaders.
 RENDERCORE_API ShaderPlatformMaskType GRayTracingPlatformMask;
 
+// Specifies whether ray tracing shaders *can* be used on a particular platform.
+// This takes into account whether RT is globally enabled for the project and specifically enabled on a target platform.
+RENDERCORE_API ShaderPlatformMaskType GRayTracingShadersPlatformMask;
+
 // Specifies whether ray tracing *is* enabled on the current running system (in current game or editor process).
 // This takes into account additional factors, such as concrete current GPU/OS/Driver capability, user-set game graphics options, etc.
 // Only safe to make run-time decisions, such as whether to build acceleration structures and render ray tracing effects.
@@ -742,6 +746,7 @@ RENDERCORE_API void RenderUtilsInit()
 	GSimpleSkyDiffusePlatformMask.Init(false, EShaderPlatform::SP_NumPlatforms);
 	GVelocityEncodeDepthPlatformMask.Init(false, EShaderPlatform::SP_NumPlatforms);
 	GRayTracingPlatformMask.Init(false, EShaderPlatform::SP_NumPlatforms);
+	GRayTracingShadersPlatformMask.Init(false, EShaderPlatform::SP_NumPlatforms);
 
 	static IConsoleVariable* MobileAmbientOcclusionCVar = IConsoleManager::Get().FindConsoleVariable(TEXT("r.Mobile.AmbientOcclusion"));
 	GMobileAmbientOcclusionPlatformMask.Init(MobileAmbientOcclusionCVar && MobileAmbientOcclusionCVar->GetInt(), EShaderPlatform::SP_NumPlatforms);
@@ -795,6 +800,8 @@ RENDERCORE_API void RenderUtilsInit()
 				TArray<FName> PlatformRayTracingShaderFormats;
 				TargetPlatformSettings->GetRayTracingShaderFormats(PlatformRayTracingShaderFormats);
 
+				const bool bRayTracingShadersEnabled = TargetPlatformSettings->GetRayTracingMode() == ERayTracingRuntimeMode::Full;
+
 				for (FName FormatName : PlatformRayTracingShaderFormats)
 				{
 					EShaderPlatform MainShaderPlatform = ShaderFormatNameToShaderPlatform(FormatName);
@@ -807,6 +814,9 @@ RENDERCORE_API void RenderUtilsInit()
 					{
 						uint32 ShaderPlatformIndex = static_cast<uint32>(ShaderPlatform);
 						GRayTracingPlatformMask[ShaderPlatformIndex] = true;
+						GRayTracingShadersPlatformMask[ShaderPlatformIndex] = FDataDrivenShaderPlatformInfo::GetIsPreviewPlatform(ShaderPlatform)
+							? true
+							: bRayTracingShadersEnabled;
 					}
 				}
 			}
@@ -823,6 +833,11 @@ RENDERCORE_API void RenderUtilsInit()
 	if (RayTracingCVar && RayTracingCVar->GetInt() && GRHISupportsRayTracing)
 	{
 		GRayTracingPlatformMask.Init(true, EShaderPlatform::SP_NumPlatforms);
+
+		if (GRHISupportsRayTracingShaders)
+		{
+			GRayTracingShadersPlatformMask.Init(true, EShaderPlatform::SP_NumPlatforms);
+		}
 	}
 
 	// Load runtime values from and *.ini file used by a current platform
@@ -842,10 +857,31 @@ RENDERCORE_API void RenderUtilsInit()
 				GDistanceFieldsPlatformMask.Init(false, EShaderPlatform::SP_NumPlatforms);
 			}
 
-			bool bRayTracing = false;
-			if (PlatformIniFile.GetBool(*CategoryName, TEXT("bEnableRayTracing"), bRayTracing) && !bRayTracing)
 			{
-				GRayTracingPlatformMask.Init(false, EShaderPlatform::SP_NumPlatforms);
+				FString RayTracingModeString;
+				GConfig->GetString(FPlatformProperties::GetRuntimeSettingsClassName(), TEXT("RayTracingMode"), RayTracingModeString, GEngineIni);
+
+				if (RayTracingModeString.Equals(TEXT("Inline"), ESearchCase::IgnoreCase))
+				{
+					GRayTracingShadersPlatformMask.Init(false, EShaderPlatform::SP_NumPlatforms);
+				}
+				else if (RayTracingModeString.Equals(TEXT("Full"), ESearchCase::IgnoreCase))
+				{
+					// nothing
+				}
+				else if (RayTracingModeString.Equals(TEXT("Disabled"), ESearchCase::IgnoreCase))
+				{
+					GRayTracingShadersPlatformMask.Init(false, EShaderPlatform::SP_NumPlatforms);
+					GRayTracingPlatformMask.Init(false, EShaderPlatform::SP_NumPlatforms);
+				}
+				else
+				{
+					bool bRayTracing = false;
+					if (PlatformIniFile.GetBool(*CategoryName, TEXT("bEnableRayTracing"), bRayTracing) && !bRayTracing)
+					{
+						GRayTracingPlatformMask.Init(false, EShaderPlatform::SP_NumPlatforms);
+					}
+				}
 			}
 		}
 	}
@@ -922,6 +958,15 @@ RENDERCORE_API void RenderUtilsInit()
 				UE_LOG(LogRendererCore, Fatal, TEXT("Ray tracing requires skin cache to be enabled. Set r.SkinCache.CompileShaders=1."));
 			}
 
+			const bool bRayTracingShadersAllowedOnCurrentPlatform = GRayTracingShadersPlatformMask[(int)GMaxRHIShaderPlatform];
+			if (GRHISupportsRayTracingShaders && bRayTracingShadersAllowedOnCurrentPlatform)
+			{
+				UE_LOG(LogRendererCore, Log, TEXT("Ray tracing shaders are enabled."));
+			}
+			else
+			{
+				UE_LOG(LogRendererCore, Log, TEXT("Ray tracing shaders are disabled."));
+			}
 		}
 		else
 		{
@@ -1521,13 +1566,25 @@ bool IsRayTracingEnabledForProject(EShaderPlatform ShaderPlatform)
 	}
 }
 
+bool AreRayTracingShadersEnabledForProject(EShaderPlatform ShaderPlatform)
+{
+	if (RHISupportsRayTracing(ShaderPlatform) && RHISupportsRayTracingShaders(ShaderPlatform))
+	{
+		return GRayTracingShadersPlatformMask[(int)ShaderPlatform];
+	}
+	else
+	{
+		return false;
+	}
+}
+
 bool ShouldCompileRayTracingShadersForProject(EShaderPlatform ShaderPlatform)
 {
 	if (RHISupportsRayTracingShaders(ShaderPlatform))
 	{
 		const bool bRayTracingRequireSM6 = (GRayTracingRequireSM6 != 0);
 
-		return IsRayTracingEnabledForProject(ShaderPlatform) && (IsFeatureLevelSupported(ShaderPlatform, ERHIFeatureLevel::SM6) || !bRayTracingRequireSM6);		
+		return IsRayTracingEnabledForProject(ShaderPlatform) && AreRayTracingShadersEnabledForProject(ShaderPlatform) && (IsFeatureLevelSupported(ShaderPlatform, ERHIFeatureLevel::SM6) || !bRayTracingRequireSM6);
 	}
 	else
 	{
diff --git a/Engine/Source/Runtime/Renderer/Private/DeferredShadingRenderer.cpp b/Engine/Source/Runtime/Renderer/Private/DeferredShadingRenderer.cpp
--- a/Engine/Source/Runtime/Renderer/Private/DeferredShadingRenderer.cpp
+++ b/Engine/Source/Runtime/Renderer/Private/DeferredShadingRenderer.cpp
@@ -1679,7 +1679,7 @@ void FDeferredShadingSceneRenderer::Render(FRDGBuilder& GraphBuilder)
 			PreparePathTracingCloudMaterial(Scene, Views);
 		}
 
-		if (IsRayTracingEnabled(ViewFamily.GetShaderPlatform()) && RHISupportsRayTracingShaders(ViewFamily.GetShaderPlatform()))
+		if (IsRayTracingEnabled(ViewFamily.GetShaderPlatform()) && GRHISupportsRayTracingShaders)
 		{
 			// Nanite raytracing manager update must run before GPUScene update since it can modify primitive data
 			Nanite::GRayTracingManager.Update();
diff --git a/Engine/Source/Runtime/Renderer/Private/RayTracing/RayTracingDebug.cpp b/Engine/Source/Runtime/Renderer/Private/RayTracing/RayTracingDebug.cpp
--- a/Engine/Source/Runtime/Renderer/Private/RayTracing/RayTracingDebug.cpp
+++ b/Engine/Source/Runtime/Renderer/Private/RayTracing/RayTracingDebug.cpp
@@ -1357,7 +1357,7 @@ void FDeferredShadingSceneRenderer::RenderRayTracingDebug(FRDGBuilder& GraphBuil
 	else
 	{
 		// Set useful default value
-		DebugVisualizationMode = RAY_TRACING_DEBUG_VIZ_BASE_COLOR;
+		DebugVisualizationMode = RAY_TRACING_DEBUG_VIZ_BARYCENTRICS;
 	}
 
 	if (DebugVisualizationMode == RAY_TRACING_DEBUG_VIZ_BARYCENTRICS)
@@ -1419,6 +1419,7 @@ void FDeferredShadingSceneRenderer::RenderRayTracingDebug(FRDGBuilder& GraphBuil
 	const bool bRayTracingPipeline = ShouldRenderRayTracingEffect(ERayTracingPipelineCompatibilityFlags::FullPipeline);
 	if (!bRayTracingPipeline)
 	{
+		AddClearUAVPass(GraphBuilder, GraphBuilder.CreateUAV(SceneColorTexture), FLinearColor::Black);
 		return;
 	}
 
