Don't create geometries that use shared buffers when FRayTracingGeometry::MakeResident() is called.
- added an ensure to check this.

diff --git a/Engine/Source/Runtime/RenderCore/Private/RayTracingGeometry.cpp b/Engine/Source/Runtime/RenderCore/Private/RayTracingGeometry.cpp
--- a/Engine/Source/Runtime/RenderCore/Private/RayTracingGeometry.cpp
+++ b/Engine/Source/Runtime/RenderCore/Private/RayTracingGeometry.cpp
@@ -75,6 +75,16 @@ void FRayTracingGeometry::RequestBuildIfNeeded(FRHICommandListBase& RHICmdList,
 void FRayTracingGeometry::MakeResident(FRHICommandList& RHICmdList)
 {
 	check(EnumHasAllFlags(GeometryState, EGeometryStateFlags::Evicted) && RayTracingGeometryRHI == nullptr);
+
+	if (!ensureMsgf(DynamicGeometrySharedBufferGenerationID == NonSharedVertexBuffers,
+		TEXT("Cannot call MakeResident(...) on FRayTracingGeometry using shared vertex buffers.\n")
+		TEXT("Dynamic geometry should be rebuilt instead.")))
+	{
+		// if geometry is using shared buffers those buffers might not be valid at this point
+		// instead of being made resident here, dynamic geometries need to be manually updated as necessary
+		return;
+	}
+
 	EnumRemoveFlags(GeometryState, EGeometryStateFlags::Evicted);
 
 	// Streaming BLAS needs special handling to not get their "streaming" type wiped out as it will cause issues down the line.	
