Add memory tracking for FShaderMapResourceCode.

For example, the output of `obj list class=MaterialInterface` will now report larger sizes to account for the in-memory shader code.

                                                                                         Object      NumKB      MaxKB   ResExcKB  ResExcDedSysKB  ResExcDedVidKB     ResExcUnkKB
Material /Engine/EngineMaterials/WorldGridMaterial.WorldGridMaterial       3.21       3.21    2105.67         2105.67            0.00            0.00
#tests Running the command "obj list class=MaterialInterface" in the editor in EngineTest

diff --git a/Engine/Source/Runtime/Engine/Private/Materials/MaterialShared.cpp b/Engine/Source/Runtime/Engine/Private/Materials/MaterialShared.cpp
--- a/Engine/Source/Runtime/Engine/Private/Materials/MaterialShared.cpp
+++ b/Engine/Source/Runtime/Engine/Private/Materials/MaterialShared.cpp
@@ -2404,8 +2404,6 @@ void FMaterialResource::GetResourceSizeEx(FResourceSizeEx& CumulativeResourceSiz
 			{
 				CumulativeResourceSize.AddDedicatedSystemMemoryBytes(Resource->GetSizeBytes());
 			}
-
-			// TODO - account for shader code size somehow, either inline or shared code
 		}
 	}
 }
diff --git a/Engine/Source/Runtime/RenderCore/Private/ShaderResource.cpp b/Engine/Source/Runtime/RenderCore/Private/ShaderResource.cpp
--- a/Engine/Source/Runtime/RenderCore/Private/ShaderResource.cpp
+++ b/Engine/Source/Runtime/RenderCore/Private/ShaderResource.cpp
@@ -675,3 +675,18 @@ FRHIShader* FShaderMapResource_InlineCode::CreateRHIShaderOrCrash(int32 ShaderIn
 	RHIShader->AddRef();
 	return RHIShader;
 }
+
+uint32 FShaderMapResource_InlineCode::GetSizeBytes() const
+{
+	uint32 TotalSize = 0;
+
+	if (Code)
+	{
+		TotalSize += Code->GetSizeBytes();
+	}
+
+	TotalSize += sizeof(FShaderMapResource_InlineCode);
+	TotalSize += GetAllocatedSize();
+
+	return TotalSize;
+}
diff --git a/Engine/Source/Runtime/RenderCore/Public/Shader.h b/Engine/Source/Runtime/RenderCore/Public/Shader.h
--- a/Engine/Source/Runtime/RenderCore/Public/Shader.h
+++ b/Engine/Source/Runtime/RenderCore/Public/Shader.h
@@ -549,7 +549,7 @@ public:
 	// FShaderMapResource interface
 	RENDERCORE_API virtual FSHAHash GetShaderHash(int32 ShaderIndex) override;
 	RENDERCORE_API virtual FRHIShader* CreateRHIShaderOrCrash(int32 ShaderIndex, bool bRequired) override;
-	virtual uint32 GetSizeBytes() const override { return sizeof(*this) + GetAllocatedSize(); }
+	virtual uint32 GetSizeBytes() const override;
 
 	TRefCountPtr<FShaderMapResourceCode> Code;
 };
