Support particle lights in mobile forward renderer

Requires r.Mobile.Forward.EnableParticleLights to be enabled, disabled by default.

diff --git a/Engine/Plugins/FX/Niagara/Source/Niagara/Private/NiagaraRendererLights.cpp b/Engine/Plugins/FX/Niagara/Source/Niagara/Private/NiagaraRendererLights.cpp
--- a/Engine/Plugins/FX/Niagara/Source/Niagara/Private/NiagaraRendererLights.cpp
+++ b/Engine/Plugins/FX/Niagara/Source/Niagara/Private/NiagaraRendererLights.cpp
@@ -62,11 +62,18 @@ FPrimitiveViewRelevance FNiagaraRendererLights::GetViewRelevance(const FSceneVie
 FNiagaraDynamicDataBase* FNiagaraRendererLights::GenerateDynamicData(const FNiagaraSceneProxy* Proxy, const UNiagaraRendererProperties* InProperties, const FNiagaraEmitterInstance* Emitter) const
 {
 	// particle (simple) lights are only supported with deferred shading
-	
-	if (!bHasLights || (Proxy->GetScene().GetShadingPath() != EShadingPath::Deferred && !IsMobileDeferredShadingEnabled(Proxy->GetScene().GetShaderPlatform())))
+	if (!bHasLights)
 	{
 		return nullptr;
 	}
+	if (Proxy->GetScene().GetShadingPath() != EShadingPath::Deferred)
+	{
+		const EShaderPlatform ShaderPlatform = Proxy->GetScene().GetShaderPlatform();
+		if (!IsMobileDeferredShadingEnabled(ShaderPlatform) && !MobileForwardEnableParticleLights(ShaderPlatform))
+		{
+			return nullptr;
+		}
+	}
 
 	if (!IsRendererEnabled(InProperties, Emitter))
 	{
diff --git a/Engine/Source/Runtime/Core/Private/HAL/ConsoleManager.cpp b/Engine/Source/Runtime/Core/Private/HAL/ConsoleManager.cpp
--- a/Engine/Source/Runtime/Core/Private/HAL/ConsoleManager.cpp
+++ b/Engine/Source/Runtime/Core/Private/HAL/ConsoleManager.cpp
@@ -3772,6 +3772,13 @@ static TAutoConsoleVariable<int32> CVarMobileForwardEnableLocalLights(
 		"2: Local Lights Buffer Enabled\n"),
 	ECVF_ReadOnly | ECVF_RenderThreadSafe);
 
+static TAutoConsoleVariable<int32> CVarMobileForwardEnableParticleLights(
+	TEXT("r.Mobile.Forward.EnableParticleLights"),
+	0,
+	TEXT("0: Particle Lights Disabled (default)\n"
+		 "1: Particle Lights Enabled (requires local lights)\n"),
+	ECVF_ReadOnly | ECVF_RenderThreadSafe);
+
 static TAutoConsoleVariable<int32> CVarMobileForwardEnableClusteredReflections(
 	TEXT("r.Mobile.Forward.EnableClusteredReflections"),
 	0,
diff --git a/Engine/Source/Runtime/RenderCore/Private/ReadOnlyCVARCache.cpp b/Engine/Source/Runtime/RenderCore/Private/ReadOnlyCVARCache.cpp
--- a/Engine/Source/Runtime/RenderCore/Private/ReadOnlyCVARCache.cpp
+++ b/Engine/Source/Runtime/RenderCore/Private/ReadOnlyCVARCache.cpp
@@ -27,6 +27,7 @@ bool FReadOnlyCVARCache::bMobileSupportsGPUScene = false;
 int32 FReadOnlyCVARCache::MobileSkyLightPermutationValue = 0;
 int32 FReadOnlyCVARCache::MobileEarlyZPassValue = 0;
 int32 FReadOnlyCVARCache::MobileForwardLocalLightsValue = 1;
+bool FReadOnlyCVARCache::bMobileForwardParticleLights = false;
 bool FReadOnlyCVARCache::bMobileEnableNoPrecomputedLightingCSMShader = false;
 bool FReadOnlyCVARCache::bMobileDeferredShadingValue = false;
 bool FReadOnlyCVARCache::bMobileEnableMovableSpotlightsShadowValue = false;
@@ -43,6 +44,12 @@ int32 FReadOnlyCVARCache::MobileForwardLocalLightsIniValue(EShaderPlatform Platf
 	return CVar.Get(Platform);
 }
 
+bool FReadOnlyCVARCache::MobileForwardParticleLightsIniValue(EShaderPlatform Platform)
+{
+	static FShaderPlatformCachedIniValue<bool> CVar(TEXT("r.Mobile.Forward.EnableParticleLights"));
+	return CVar.Get(Platform);
+}
+
 bool FReadOnlyCVARCache::MobileDeferredShadingIniValue(EShaderPlatform Platform)
 {
 	static FShaderPlatformCachedIniValue<bool> MobileShadingPathIniValue(TEXT("r.Mobile.ShadingPath"));
@@ -109,6 +116,7 @@ void FReadOnlyCVARCache::Initialize()
 	bMobileDeferredShadingValue = MobileDeferredShadingIniValue(GMaxRHIShaderPlatform);
 	bMobileEnableMovableSpotlightsShadowValue = MobileEnableMovableSpotlightsShadowIniValue(GMaxRHIShaderPlatform);
 	bMobileSupportsGPUScene = CVarMobileSupportGPUScene->GetValueOnAnyThread() != 0;
+	bMobileForwardParticleLights = MobileForwardParticleLightsIniValue(GMaxRHIShaderPlatform);
 
 #ifdef PROJECT_CVAR_ALLOW_STATIC_LIGHTING
 	check(!!PROJECT_CVAR_ALLOW_STATIC_LIGHTING == bAllowStaticLighting);
diff --git a/Engine/Source/Runtime/RenderCore/Public/ReadOnlyCVARCache.h b/Engine/Source/Runtime/RenderCore/Public/ReadOnlyCVARCache.h
--- a/Engine/Source/Runtime/RenderCore/Public/ReadOnlyCVARCache.h
+++ b/Engine/Source/Runtime/RenderCore/Public/ReadOnlyCVARCache.h
@@ -111,6 +111,17 @@ struct FReadOnlyCVARCache
 			return MobileForwardLocalLightsValue;
 		#endif
 	}
+
+	static inline int32 MobileForwardParticleLights(EShaderPlatform Platform)
+	{
+		#if WITH_EDITOR
+				return MobileForwardParticleLightsIniValue(Platform);
+		#elif defined PROJECT_CVAR_MOBILE_FORWARD_PARTICLELIGHTS
+				return PROJECT_CVAR_MOBILE_FORWARD_PARTICLELIGHTS;
+		#else
+				return bMobileForwardParticleLights;
+		#endif
+	}
 	
 	static inline bool MobileDeferredShading(EShaderPlatform Platform)
 	{
@@ -151,6 +162,7 @@ private:
 	RENDERCORE_API static int32 MobileSkyLightPermutationValue;
 	RENDERCORE_API static int32 MobileEarlyZPassValue;
 	RENDERCORE_API static int32 MobileForwardLocalLightsValue;
+	RENDERCORE_API static bool bMobileForwardParticleLights;
 	RENDERCORE_API static bool bMobileEnableNoPrecomputedLightingCSMShader;
 	RENDERCORE_API static bool bMobileDeferredShadingValue;
 	RENDERCORE_API static bool bMobileEnableMovableSpotlightsShadowValue;
@@ -158,6 +170,7 @@ private:
 private:
 	RENDERCORE_API static int32 MobileEarlyZPassIniValue(EShaderPlatform Platform);
 	RENDERCORE_API static int32 MobileForwardLocalLightsIniValue(EShaderPlatform Platform);
+	RENDERCORE_API static bool MobileForwardParticleLightsIniValue(EShaderPlatform Platform);
 	RENDERCORE_API static bool MobileDeferredShadingIniValue(EShaderPlatform Platform);
 	RENDERCORE_API static bool MobileEnableMovableSpotlightsShadowIniValue(EShaderPlatform Platform);
 };
diff --git a/Engine/Source/Runtime/RenderCore/Public/RenderUtils.h b/Engine/Source/Runtime/RenderCore/Public/RenderUtils.h
--- a/Engine/Source/Runtime/RenderCore/Public/RenderUtils.h
+++ b/Engine/Source/Runtime/RenderCore/Public/RenderUtils.h
@@ -291,6 +291,11 @@ inline bool MobileForwardEnableLocalLights(const FStaticShaderPlatform Platform)
 	return FReadOnlyCVARCache::MobileForwardLocalLights(Platform) > 0;
 }
 
+inline bool MobileForwardEnableParticleLights(const FStaticShaderPlatform Platform)
+{
+	return MobileForwardEnableLocalLights(Platform) && FReadOnlyCVARCache::MobileForwardParticleLights(Platform);
+}
+
 RENDERCORE_API bool PlatformGPUSceneUsesUniformBufferView(const FStaticShaderPlatform Platform);
 
 RENDERCORE_API bool MobileRequiresSceneDepthAux(const FStaticShaderPlatform Platform);
diff --git a/Engine/Source/Runtime/Renderer/Private/MobileBasePass.cpp b/Engine/Source/Runtime/Renderer/Private/MobileBasePass.cpp
--- a/Engine/Source/Runtime/Renderer/Private/MobileBasePass.cpp
+++ b/Engine/Source/Runtime/Renderer/Private/MobileBasePass.cpp
@@ -38,9 +38,9 @@ FAutoConsoleVariableRef CVarMobileForwardLocalLightsSinglePermutation(
 	ECVF_Scalability | ECVF_RenderThreadSafe
 );
 
-bool MobileLocalLightsUseSinglePermutation()
+bool MobileLocalLightsUseSinglePermutation(EShaderPlatform ShaderPlatform)
 { 
-	return GMobileForwardLocalLightsSinglePermutation != 0;
+	return GMobileForwardLocalLightsSinglePermutation != 0 || MobileForwardEnableParticleLights(ShaderPlatform);
 }
 
 EMobileLocalLightSetting GetMobileForwardLocalLightSetting(EShaderPlatform ShaderPlatform)
@@ -921,7 +921,7 @@ bool FMobileBasePassMeshProcessor::Process(
 		if (!bPassUsesDeferredShading &&
 			// we can choose to use a single permutation regarless of local light state
 			// this is to avoid re-caching MDC on light state changes
-			(MobileLocalLightsUseSinglePermutation() || PrimitiveSceneProxy->GetPrimitiveSceneInfo()->NumMobileDynamicLocalLights > 0))
+			(MobileLocalLightsUseSinglePermutation(Scene->GetShaderPlatform()) || PrimitiveSceneProxy->GetPrimitiveSceneInfo()->NumMobileDynamicLocalLights > 0))
 		{
 			LocalLightSetting = GetMobileForwardLocalLightSetting(Scene->GetShaderPlatform());
 		}
diff --git a/Engine/Source/Runtime/Renderer/Private/MobileShadingRenderer.cpp b/Engine/Source/Runtime/Renderer/Private/MobileShadingRenderer.cpp
--- a/Engine/Source/Runtime/Renderer/Private/MobileShadingRenderer.cpp
+++ b/Engine/Source/Runtime/Renderer/Private/MobileShadingRenderer.cpp
@@ -277,6 +277,7 @@ FMobileSceneRenderer::FMobileSceneRenderer(const FSceneViewFamily* InViewFamily,
 	, bDeferredShading(IsMobileDeferredShadingEnabled(ShaderPlatform))
 	, bRequiresDBufferDecals(bDeferredShading ? false : IsUsingDBuffers(ShaderPlatform))
 	, bUseVirtualTexturing(UseVirtualTexturing(ShaderPlatform) && GetRendererOutput() != FSceneRenderer::ERendererOutput::DepthPrepassOnly)
+	, bSupportsSimpleLights(bDeferredShading || MobileForwardEnableParticleLights(ShaderPlatform))
 {
 	bRenderToSceneColor = false;
 	bRequiresMultiPass = false;
@@ -2336,5 +2337,5 @@ void FMobileSceneRenderer::RenderHZB(FRDGBuilder& GraphBuilder, FRDGTextureRef S
 
 bool FMobileSceneRenderer::AllowSimpleLights() const
 {
-	return FSceneRenderer::AllowSimpleLights() && bDeferredShading;
+	return FSceneRenderer::AllowSimpleLights() && bSupportsSimpleLights;
 }
diff --git a/Engine/Source/Runtime/Renderer/Private/RendererScene.cpp b/Engine/Source/Runtime/Renderer/Private/RendererScene.cpp
--- a/Engine/Source/Runtime/Renderer/Private/RendererScene.cpp
+++ b/Engine/Source/Runtime/Renderer/Private/RendererScene.cpp
@@ -3855,8 +3855,8 @@ void FScene::GetPrimitiveUniformShaderParameters_RenderThread(const FPrimitiveSc
 
 bool DoesPlatformNeedLocalLightPrimitiveInteraction(EShaderPlatform ShaderPlatform)
 {
-	extern bool MobileLocalLightsUseSinglePermutation();
-	return !IsMobilePlatform(ShaderPlatform) || !MobileLocalLightsUseSinglePermutation() || IsMobileMovableSpotlightShadowsEnabled(ShaderPlatform);
+	extern bool MobileLocalLightsUseSinglePermutation(EShaderPlatform ShaderPlatform);
+	return !IsMobilePlatform(ShaderPlatform) || !MobileLocalLightsUseSinglePermutation(ShaderPlatform) || IsMobileMovableSpotlightShadowsEnabled(ShaderPlatform);
 }
 
 void FScene::UpdateLightTransform_RenderThread(int32 LightId, FLightSceneInfo* LightSceneInfo, const FUpdateLightCommand::FTransformParameters& Parameters)
diff --git a/Engine/Source/Runtime/Renderer/Private/SceneCore.cpp b/Engine/Source/Runtime/Renderer/Private/SceneCore.cpp
--- a/Engine/Source/Runtime/Renderer/Private/SceneCore.cpp
+++ b/Engine/Source/Runtime/Renderer/Private/SceneCore.cpp
@@ -158,11 +158,11 @@ extern bool ShouldCreateObjectShadowForStationaryLight(const FLightSceneInfo* Li
 
 static bool MobileRequiresStaticMeshUpdateOnLocalLightChange(const FStaticShaderPlatform Platform)
 {
-	extern bool MobileLocalLightsUseSinglePermutation();
+	extern bool MobileLocalLightsUseSinglePermutation(EShaderPlatform);
 	
 	if (!IsMobileDeferredShadingEnabled(Platform))
 	{
-		return MobileForwardEnableLocalLights(Platform) && !MobileLocalLightsUseSinglePermutation();
+		return MobileForwardEnableLocalLights(Platform) && !MobileLocalLightsUseSinglePermutation(Platform);
 	}
 	return false;
 }
diff --git a/Engine/Source/Runtime/Renderer/Private/SceneRendering.h b/Engine/Source/Runtime/Renderer/Private/SceneRendering.h
--- a/Engine/Source/Runtime/Renderer/Private/SceneRendering.h
+++ b/Engine/Source/Runtime/Renderer/Private/SceneRendering.h
@@ -2783,6 +2783,7 @@ private:
 	const bool bDeferredShading;
 	const bool bRequiresDBufferDecals;
 	const bool bUseVirtualTexturing;
+	const bool bSupportsSimpleLights;
 	bool bTonemapSubpass;
 	bool bTonemapSubpassInline;
 	int32 NumMSAASamples;
