stat shaders: removed some obsolete stats, added a stat showing actually used shadermaps.

#tests ran a replay

diff --git a/Engine/Source/Runtime/RenderCore/Private/ShaderCodeArchive.cpp b/Engine/Source/Runtime/RenderCore/Private/ShaderCodeArchive.cpp
--- a/Engine/Source/Runtime/RenderCore/Private/ShaderCodeArchive.cpp
+++ b/Engine/Source/Runtime/RenderCore/Private/ShaderCodeArchive.cpp
@@ -1257,17 +1257,6 @@ void FShaderCodeArchive::ReleasePreloadedShader(int32 ShaderIndex)
 TRefCountPtr<FRHIShader> FShaderCodeArchive::CreateShader(int32 Index)
 {
 	LLM_SCOPE(ELLMTag::Shaders);
-#if STATS
-	double TimeFunctionEntered = FPlatformTime::Seconds();
-	ON_SCOPE_EXIT
-	{
-		if (IsInRenderingThread())
-		{
-			double ShaderCreationTime = FPlatformTime::Seconds() - TimeFunctionEntered;
-			INC_FLOAT_STAT_BY(STAT_Shaders_TotalRTShaderInitForRenderingTime, ShaderCreationTime);
-		}
-	};
-#endif
 
 	TRefCountPtr<FRHIShader> Shader;
 
@@ -1337,7 +1326,7 @@ TRefCountPtr<FRHIShader> FShaderCodeArchive::CreateShader(int32 Index)
 
 	if (Shader)
 	{
-		INC_DWORD_STAT(STAT_Shaders_NumShadersUsedForRendering);
+		INC_DWORD_STAT(STAT_Shaders_NumShadersCreated);
 		Shader->SetHash(ShaderHash);
 	}
 
@@ -2161,17 +2150,6 @@ int32 FIoStoreShaderCodeArchive::FindShaderIndex(const FSHAHash& Hash)
 TRefCountPtr<FRHIShader> FIoStoreShaderCodeArchive::CreateShader(int32 ShaderIndex)
 {
 	LLM_SCOPE(ELLMTag::Shaders);
-#if STATS
-	double TimeFunctionEntered = FPlatformTime::Seconds();
-	ON_SCOPE_EXIT
-	{
-		if (IsInRenderingThread())
-		{
-			double ShaderCreationTime = FPlatformTime::Seconds() - TimeFunctionEntered;
-			INC_FLOAT_STAT_BY(STAT_Shaders_TotalRTShaderInitForRenderingTime, ShaderCreationTime);
-		}
-	};
-#endif
 
 	TRACE_CPUPROFILER_EVENT_SCOPE(FIoStoreShaderCodeArchive::CreateShader);
 
@@ -2277,7 +2255,7 @@ TRefCountPtr<FRHIShader> FIoStoreShaderCodeArchive::CreateShader(int32 ShaderInd
 
 	if (Shader)
 	{
-		INC_DWORD_STAT(STAT_Shaders_NumShadersUsedForRendering);
+		INC_DWORD_STAT(STAT_Shaders_NumShadersCreated);
 		Shader->SetHash(ShaderHash);
 	}
 
diff --git a/Engine/Source/Runtime/RenderCore/Private/ShaderCodeLibrary.cpp b/Engine/Source/Runtime/RenderCore/Private/ShaderCodeLibrary.cpp
--- a/Engine/Source/Runtime/RenderCore/Private/ShaderCodeLibrary.cpp
+++ b/Engine/Source/Runtime/RenderCore/Private/ShaderCodeLibrary.cpp
@@ -1210,14 +1210,6 @@ FSHAHash FShaderMapResource_SharedCode::GetShaderHash(int32 ShaderIndex)
 FRHIShader* FShaderMapResource_SharedCode::CreateRHIShaderOrCrash(int32 ShaderIndex, bool bRequired)
 {
 	SCOPED_LOADTIMER(FShaderMapResource_SharedCode_InitRHI);
-#if STATS
-	double TimeFunctionEntered = FPlatformTime::Seconds();
-	ON_SCOPE_EXIT
-	{
-		double ShaderCreationTime = FPlatformTime::Seconds() - TimeFunctionEntered;
-		INC_FLOAT_STAT_BY(STAT_Shaders_TotalRTShaderInitForRenderingTime, ShaderCreationTime);
-	};
-#endif
 
 	const int32 LibraryShaderIndex = LibraryInstance->Library->GetShaderIndex(ShaderMapIndex, ShaderIndex);
 	TRefCountPtr<FRHIShader> CreatedShader = LibraryInstance->GetOrCreateShader(LibraryShaderIndex);
diff --git a/Engine/Source/Runtime/RenderCore/Private/ShaderCore.cpp b/Engine/Source/Runtime/RenderCore/Private/ShaderCore.cpp
--- a/Engine/Source/Runtime/RenderCore/Private/ShaderCore.cpp
+++ b/Engine/Source/Runtime/RenderCore/Private/ShaderCore.cpp
@@ -107,19 +107,14 @@ DEFINE_STAT(STAT_ShaderCompiling_NumMaskedMaterialShaders);
 
 
 DEFINE_STAT(STAT_Shaders_NumShadersLoaded);
-DEFINE_STAT(STAT_Shaders_NumShaderResourcesLoaded);
+DEFINE_STAT(STAT_Shaders_NumShadersCreated);
 DEFINE_STAT(STAT_Shaders_NumShaderMaps);
+DEFINE_STAT(STAT_Shaders_NumShaderMapsUsedForRendering);
 DEFINE_STAT(STAT_Shaders_RTShaderLoadTime);
-DEFINE_STAT(STAT_Shaders_NumShadersUsedForRendering);
-DEFINE_STAT(STAT_Shaders_TotalRTShaderInitForRenderingTime);
-DEFINE_STAT(STAT_Shaders_FrameRTShaderInitForRenderingTime);
 DEFINE_STAT(STAT_Shaders_ShaderMemory);
 DEFINE_STAT(STAT_Shaders_ShaderResourceMemory);
 DEFINE_STAT(STAT_Shaders_ShaderPreloadMemory);
 
-DEFINE_STAT(STAT_Shaders_NumShadersRegistered);
-DEFINE_STAT(STAT_Shaders_NumShadersDuplicated);
-
 PRAGMA_DISABLE_DEPRECATION_WARNINGS		// FShaderCompilerDefinitions will be made internal in the future, marked deprecated until then
 
 /**
diff --git a/Engine/Source/Runtime/RenderCore/Private/ShaderResource.cpp b/Engine/Source/Runtime/RenderCore/Private/ShaderResource.cpp
--- a/Engine/Source/Runtime/RenderCore/Private/ShaderResource.cpp
+++ b/Engine/Source/Runtime/RenderCore/Private/ShaderResource.cpp
@@ -422,7 +422,8 @@ void FShaderMapResourceCode::NotifyShadersCompiled(FName FormatName)
 #endif // WITH_EDITORONLY_DATA
 
 FShaderMapResource::FShaderMapResource(EShaderPlatform InPlatform, int32 NumShaders)
-	: NumRHIShaders(NumShaders)
+	: NumRHIShaders(static_cast<uint32>(NumShaders))
+	, bAtLeastOneRHIShaderCreated(0)
 	, Platform(InPlatform)
 	, NumRefs(0)
 {
@@ -466,16 +467,23 @@ void FShaderMapResource::ReleaseShaders()
 {
 	if (RHIShaders)
 	{
-		for (int32 Idx = 0; Idx < NumRHIShaders; ++Idx)
+		FScopeLock ScopeLock(&RHIShadersCreationGuard);
+
+		for (uint32 Idx = 0; Idx < NumRHIShaders; ++Idx)
 		{
 			if (FRHIShader* Shader = RHIShaders[Idx].load(std::memory_order_acquire))
 			{
 				Shader->Release();
-				DEC_DWORD_STAT(STAT_Shaders_NumShadersUsedForRendering);
+				DEC_DWORD_STAT(STAT_Shaders_NumShadersCreated);
 			}
 		}
 		RHIShaders = nullptr;
 		NumRHIShaders = 0;
+		if (bAtLeastOneRHIShaderCreated)
+		{
+			DEC_DWORD_STAT(STAT_Shaders_NumShaderMapsUsedForRendering);
+		}
+		bAtLeastOneRHIShaderCreated = false;
 	}
 }
 
@@ -485,9 +493,9 @@ void FShaderMapResource::ReleaseRHI()
 #if RHI_RAYTRACING
 	if (GRHISupportsRayTracing && GRHISupportsRayTracingShaders)
 	{
-		check(NumRHIShaders == RayTracingLibraryIndices.Num());
+		check(NumRHIShaders == static_cast<uint32>(RayTracingLibraryIndices.Num()));
 
-		for (int32 Idx = 0; Idx < NumRHIShaders; ++Idx)
+		for (uint32 Idx = 0; Idx < NumRHIShaders; ++Idx)
 		{
 			if (FRHIShader* Shader = RHIShaders[Idx].load(std::memory_order_acquire))
 			{
@@ -542,7 +550,7 @@ FRHIShader* FShaderMapResource::CreateShaderOrCrash(int32 ShaderIndex, bool bReq
 
 	{
 		// Most shadermaps have <100 shaders, and less than a half of them can be created. 
-		// However, if this path is often contended, you can slice this lock
+		// However, if this path is often contended, you can slice this lock (but remember to take care of STAT_Shaders_NumShaderMapsUsedForRendering!)
 		FScopeLock ScopeLock(&RHIShadersCreationGuard);
 
 		Shader = RHIShaders[ShaderIndex].load(std::memory_order_relaxed);
@@ -552,6 +560,12 @@ FRHIShader* FShaderMapResource::CreateShaderOrCrash(int32 ShaderIndex, bool bReq
 			CreatedShader = nullptr;
 			RHIShaders[ShaderIndex].store(Shader, std::memory_order_release);
 
+			if (!bAtLeastOneRHIShaderCreated)
+			{
+				INC_DWORD_STAT(STAT_Shaders_NumShaderMapsUsedForRendering);
+				bAtLeastOneRHIShaderCreated = 1;
+			}
+
 #if RHI_RAYTRACING
 			// Registers RT shaders in global "libraries" that track all shaders potentially usable in a scene for adding to RTPSO
 			EShaderFrequency Frequency = Shader->GetFrequency();
@@ -601,14 +615,6 @@ FSHAHash FShaderMapResource_InlineCode::GetShaderHash(int32 ShaderIndex)
 FRHIShader* FShaderMapResource_InlineCode::CreateRHIShaderOrCrash(int32 ShaderIndex, bool bRequired)
 {
 	TRACE_CPUPROFILER_EVENT_SCOPE(FShaderMapResource_InlineCode::CreateRHIShaderOrCrash);
-#if STATS
-	double TimeFunctionEntered = FPlatformTime::Seconds();
-	ON_SCOPE_EXIT
-	{
-		double ShaderCreationTime = FPlatformTime::Seconds() - TimeFunctionEntered;
-		INC_FLOAT_STAT_BY(STAT_Shaders_TotalRTShaderInitForRenderingTime, ShaderCreationTime);
-	};
-#endif
 
 	// we can't have this called on the wrong platform's shaders
 	if (!ArePlatformsCompatible(GMaxRHIShaderPlatform, GetPlatform()))
@@ -668,7 +674,7 @@ FRHIShader* FShaderMapResource_InlineCode::CreateRHIShaderOrCrash(int32 ShaderIn
 		return nullptr;
 	}
 
-	INC_DWORD_STAT(STAT_Shaders_NumShadersUsedForRendering);
+	INC_DWORD_STAT(STAT_Shaders_NumShadersCreated);
 	RHIShader->SetHash(ShaderHash);
 
 	// contract of this function is to return a shader with an already held reference
diff --git a/Engine/Source/Runtime/RenderCore/Public/Shader.h b/Engine/Source/Runtime/RenderCore/Public/Shader.h
--- a/Engine/Source/Runtime/RenderCore/Public/Shader.h
+++ b/Engine/Source/Runtime/RenderCore/Public/Shader.h
@@ -351,7 +351,7 @@ public:
 
 	inline bool IsValidShaderIndex(int32 ShaderIndex) const
 	{
-		return ShaderIndex >= 0 && ShaderIndex < NumRHIShaders;
+		return static_cast<uint32>(ShaderIndex) < NumRHIShaders;
 	}
 
 	inline bool HasShader(int32 ShaderIndex) const
@@ -451,19 +451,22 @@ private:
 	/** This lock is to prevent two threads creating the same RHIShaders element. It is only taken if the element is to be created. */
 	FCriticalSection RHIShadersCreationGuard;
 
+#if RHI_RAYTRACING
+	TArray<uint32> RayTracingLibraryIndices;
+#endif // RHI_RAYTRACING
+
 	/** An array of shader pointers (refcount is managed manually). */
 	TUniquePtr<std::atomic<FRHIShader*>[]> RHIShaders;
 
-	/** Since the shaders are no longer a TArray, this is their count (the size of the RHIShaders array). */
-	int32 NumRHIShaders;
+	/** Since the shaders are no longer a TArray, this is their count (the size of the RHIShaders array). This does not count actually created RHI shaders, just the size of the above array of pointers to them. */
+	uint32 NumRHIShaders : 31;
 
-#if RHI_RAYTRACING
-	TArray<uint32> RayTracingLibraryIndices;
-#endif // RHI_RAYTRACING
+	/** Whether we have at least one RHI shader created. */
+	uint32 bAtLeastOneRHIShaderCreated : 1;
 
 	EShaderPlatform Platform;
 
-	/** The number of references to this shader. */
+	/** The number of references to this shader map. */
 	std::atomic<int32> NumRefs;
 };
 
diff --git a/Engine/Source/Runtime/RenderCore/Public/ShaderCore.h b/Engine/Source/Runtime/RenderCore/Public/ShaderCore.h
--- a/Engine/Source/Runtime/RenderCore/Public/ShaderCore.h
+++ b/Engine/Source/Runtime/RenderCore/Public/ShaderCore.h
@@ -93,20 +93,21 @@ DECLARE_DWORD_ACCUMULATOR_STAT_EXTERN(TEXT("Num Transparent Material Shaders"),S
 DECLARE_DWORD_ACCUMULATOR_STAT_EXTERN(TEXT("Num Opaque Material Shaders"),STAT_ShaderCompiling_NumOpaqueMaterialShaders,STATGROUP_ShaderCompiling, RENDERCORE_API);
 DECLARE_DWORD_ACCUMULATOR_STAT_EXTERN(TEXT("Num Masked Material Shaders"),STAT_ShaderCompiling_NumMaskedMaterialShaders,STATGROUP_ShaderCompiling, RENDERCORE_API);
 
+// Shaders that have been loaded in memory (in a compressed form). Can also be called "preloaded" shaders.
 DECLARE_DWORD_ACCUMULATOR_STAT_EXTERN(TEXT("Num Shaders Loaded"),STAT_Shaders_NumShadersLoaded,STATGROUP_Shaders, RENDERCORE_API);
-DECLARE_DWORD_ACCUMULATOR_STAT_EXTERN(TEXT("Num Shader Resources Loaded"),STAT_Shaders_NumShaderResourcesLoaded,STATGROUP_Shaders, RENDERCORE_API);
-DECLARE_DWORD_ACCUMULATOR_STAT_EXTERN(TEXT("Num Shader Maps Registered"),STAT_Shaders_NumShaderMaps,STATGROUP_Shaders, RENDERCORE_API);
+// Shaders that have been created as an RHI shader (i.e. exist in terms of underlying graphics API). Can also be called "used" shaders, although in theory some code may create but not use.
+DECLARE_DWORD_ACCUMULATOR_STAT_EXTERN(TEXT("Num Shaders Created"), STAT_Shaders_NumShadersCreated, STATGROUP_Shaders, RENDERCORE_API);
+// Shaders maps (essentially, material assets) that have been loaded in memory. This squarely correlates with the number of assets in memory (but assets can share SMs or have more than one SM - e.g. diff quality levels).
+DECLARE_DWORD_ACCUMULATOR_STAT_EXTERN(TEXT("Num ShaderMaps Loaded"), STAT_Shaders_NumShaderMaps,STATGROUP_Shaders, RENDERCORE_API);
+// Shaders maps that has had RHI shaders created for them. This usually means that the asset referencing this shadermap was rendered at least in some pass.
+DECLARE_DWORD_ACCUMULATOR_STAT_EXTERN(TEXT("Num ShaderMaps Used"), STAT_Shaders_NumShaderMapsUsedForRendering, STATGROUP_Shaders, RENDERCORE_API);
+
 DECLARE_CYCLE_STAT_EXTERN(TEXT("RT Shader Load Time"),STAT_Shaders_RTShaderLoadTime,STATGROUP_Shaders, RENDERCORE_API);
-DECLARE_DWORD_ACCUMULATOR_STAT_EXTERN(TEXT("Num Shaders Used"),STAT_Shaders_NumShadersUsedForRendering,STATGROUP_Shaders, RENDERCORE_API);
-DECLARE_FLOAT_ACCUMULATOR_STAT_EXTERN(TEXT("Total RT Shader Init Time"),STAT_Shaders_TotalRTShaderInitForRenderingTime,STATGROUP_Shaders, RENDERCORE_API);
-DECLARE_CYCLE_STAT_EXTERN(TEXT("Frame RT Shader Init Time"),STAT_Shaders_FrameRTShaderInitForRenderingTime,STATGROUP_Shaders, RENDERCORE_API);
+
 DECLARE_MEMORY_STAT_EXTERN(TEXT("Shader Memory"),STAT_Shaders_ShaderMemory,STATGROUP_Shaders, RENDERCORE_API);
 DECLARE_MEMORY_STAT_EXTERN(TEXT("Shader Resource Mem"),STAT_Shaders_ShaderResourceMemory,STATGROUP_Shaders, RENDERCORE_API);
 DECLARE_MEMORY_STAT_EXTERN(TEXT("Shader Preload Mem"), STAT_Shaders_ShaderPreloadMemory, STATGROUP_Shaders, RENDERCORE_API);
 
-DECLARE_DWORD_ACCUMULATOR_STAT_EXTERN(TEXT("Num Shaders Registered"), STAT_Shaders_NumShadersRegistered, STATGROUP_Shaders, RENDERCORE_API);
-DECLARE_DWORD_ACCUMULATOR_STAT_EXTERN(TEXT("Num Shaders Duplicated"), STAT_Shaders_NumShadersDuplicated, STATGROUP_Shaders, RENDERCORE_API);
-
 inline TStatId GetMemoryStatType(EShaderFrequency ShaderFrequency)
 {
 	static_assert(12 == SF_NumFrequencies, "EShaderFrequency has a bad size.");
