Add a FGPUDebugCrashUtilsCS variant that triggers shader assert

* Use `gpudebugcrash assert` console command to trigger

diff --git a/Engine/Shaders/Private/Tools/GPUDebugCrashUtils.usf b/Engine/Shaders/Private/Tools/GPUDebugCrashUtils.usf
--- a/Engine/Shaders/Private/Tools/GPUDebugCrashUtils.usf
+++ b/Engine/Shaders/Private/Tools/GPUDebugCrashUtils.usf
@@ -24,6 +24,11 @@ void MainCS(uint3 ThreadID : SV_DispatchThreadID)
 	{
 		PageFaultUAV[index1] = what;
 	}
+
+#if ASSERT_REQUESTED && PLATFORM_SUPPORTS_DIAGNOSTIC_BUFFER
+	UEReportAssertWithPayload(0xBE4A3F32, uint4(1, 2, 3, 4));
+#endif
+
 #if  HANG_REQUESTED
 	// Try to trigger a crash (timeout) using a deliberate infinite loop
 	for (;;)
diff --git a/Engine/Source/Runtime/Engine/Private/UnrealEngine.cpp b/Engine/Source/Runtime/Engine/Private/UnrealEngine.cpp
--- a/Engine/Source/Runtime/Engine/Private/UnrealEngine.cpp
+++ b/Engine/Source/Runtime/Engine/Private/UnrealEngine.cpp
@@ -5982,6 +5982,10 @@ bool UEngine::HandleGPUDebugCrashCommand(const TCHAR* Cmd, FOutputDevice& Ar)
 	{
 		GRHIGlobals.TriggerGPUCrash = ERequestedGPUCrash::Type_PlatformBreak;
 	}
+	if (FParse::Command(&Cmd, TEXT("assert")) && bSupportsPageFaultsAndPlatformBreak)
+	{
+		GRHIGlobals.TriggerGPUCrash = ERequestedGPUCrash::Type_Assert;
+	}
 	if (FParse::Command(&Cmd, TEXT("hang")) || !bSupportsPageFaultsAndPlatformBreak)
 	{
 		if (!bSupportsPageFaultsAndPlatformBreak)
diff --git a/Engine/Source/Runtime/RHI/Public/RHIDefinitions.h b/Engine/Source/Runtime/RHI/Public/RHIDefinitions.h
--- a/Engine/Source/Runtime/RHI/Public/RHIDefinitions.h
+++ b/Engine/Source/Runtime/RHI/Public/RHIDefinitions.h
@@ -1310,9 +1310,10 @@ enum class ERequestedGPUCrash : uint8
 	Type_Hang = 1 << 0,
 	Type_PageFault = 1 << 1,
 	Type_PlatformBreak = 1 << 2,
+	Type_Assert = 1 << 3,
 
-	Queue_Direct = 1 << 3,
-	Queue_Compute = 1 << 4
+	Queue_Direct = 1 << 4,
+	Queue_Compute = 1 << 5,
 };
 ENUM_CLASS_FLAGS(ERequestedGPUCrash);
 
diff --git a/Engine/Source/Runtime/RenderCore/Private/GPUDebugCrashUtils.cpp b/Engine/Source/Runtime/RenderCore/Private/GPUDebugCrashUtils.cpp
--- a/Engine/Source/Runtime/RenderCore/Private/GPUDebugCrashUtils.cpp
+++ b/Engine/Source/Runtime/RenderCore/Private/GPUDebugCrashUtils.cpp
@@ -18,8 +18,9 @@ public:
 	SHADER_USE_PARAMETER_STRUCT(FGPUDebugCrashUtilsCS, FGlobalShader)
 
 	class FPlatformBreakRequested : SHADER_PERMUTATION_BOOL("PLATFORM_BREAK_REQUESTED");
+	class FAssertRequested : SHADER_PERMUTATION_BOOL("ASSERT_REQUESTED");
 	class FHangRequested : SHADER_PERMUTATION_BOOL("HANG_REQUESTED");
-	using FPermutationDomain = TShaderPermutationDomain<FPlatformBreakRequested, FHangRequested>;
+	using FPermutationDomain = TShaderPermutationDomain<FPlatformBreakRequested, FAssertRequested, FHangRequested>;
 
 	BEGIN_SHADER_PARAMETER_STRUCT(FParameters, )
 		SHADER_PARAMETER_RDG_TEXTURE_UAV(RWTexture2D<uint>, PageFaultUAV)
@@ -39,6 +40,7 @@ void ScheduleGPUDebugCrash(FRDGBuilder& GraphBuilder)
 {
 	FGPUDebugCrashUtilsCS::FPermutationDomain PermutationVector;
 	PermutationVector.Set<FGPUDebugCrashUtilsCS::FPlatformBreakRequested>(EnumHasAnyFlags(GRHIGlobals.TriggerGPUCrash, ERequestedGPUCrash::Type_PlatformBreak));
+	PermutationVector.Set<FGPUDebugCrashUtilsCS::FAssertRequested>(EnumHasAnyFlags(GRHIGlobals.TriggerGPUCrash, ERequestedGPUCrash::Type_Assert));
 	PermutationVector.Set<FGPUDebugCrashUtilsCS::FHangRequested>(EnumHasAnyFlags(GRHIGlobals.TriggerGPUCrash, ERequestedGPUCrash::Type_Hang));
 
 	auto ComputeShader = GetGlobalShaderMap(GMaxRHIFeatureLevel)->GetShader<FGPUDebugCrashUtilsCS>(PermutationVector);
