Fix memory leak in GPU profiler
 - Don't emit GPU events from RDG GPU scopes when AreGPUStatsEnabled() returns false, otherwise timestamp queries accumulate and don't get cleared out.

diff --git a/Engine/Source/Runtime/RenderCore/Public/RenderGraphEvent.inl b/Engine/Source/Runtime/RenderCore/Public/RenderGraphEvent.inl
--- a/Engine/Source/Runtime/RenderCore/Public/RenderGraphEvent.inl
+++ b/Engine/Source/Runtime/RenderCore/Public/RenderGraphEvent.inl
@@ -11,11 +11,11 @@
 #if HAS_GPU_STATS
 inline FRDGScope_GPU::FRDGScope_GPU(FRDGScopeState& State, FRHIGPUMask GPUMask, const FName& CsvStatName, const TStatId& Stat, const TCHAR* Description, FRHIDrawStatsCategory const& Category)
 	: CurrentCategory(Category.ShouldCountDraws() ? &Category : nullptr),
-	  bEmitDuringExecute(!State.ScopeState.bParallelExecute)
+	  bEmitDuringExecute(AreGPUStatsEnabled() && !State.ScopeState.bParallelExecute)
 {
 	if (AreGPUStatsEnabled())
 	{
-		if(bEmitDuringExecute)
+		if (bEmitDuringExecute)
 		{ 
 			StatName = CsvStatName;
 			StatId = Stat;
