Fix non unity issues
#rnx

diff --git a/Engine/Source/Runtime/Core/Public/Clang/ClangPlatformCompilerPreSetup.h b/Engine/Source/Runtime/Core/Public/Clang/ClangPlatformCompilerPreSetup.h
--- a/Engine/Source/Runtime/Core/Public/Clang/ClangPlatformCompilerPreSetup.h
+++ b/Engine/Source/Runtime/Core/Public/Clang/ClangPlatformCompilerPreSetup.h
@@ -302,3 +302,5 @@
 #endif
 
 #endif // __clang__
+
+#include "GenericPlatform/GenericPlatformCompilerPreSetup.h"
\ No newline at end of file
diff --git a/Engine/Source/Runtime/MovieScene/MovieScene.Build.cs b/Engine/Source/Runtime/MovieScene/MovieScene.Build.cs
--- a/Engine/Source/Runtime/MovieScene/MovieScene.Build.cs
+++ b/Engine/Source/Runtime/MovieScene/MovieScene.Build.cs
@@ -19,6 +19,7 @@ public class MovieScene : ModuleRules
 				"CoreUObject",
                 "InputCore",
                 "Engine",
+				"MovieSceneTracks",
 				"SlateCore",
 				"TimeManagement",
 				"UniversalObjectLocator"
diff --git a/Engine/Source/Runtime/RenderCore/Public/RayGenShaderUtils.h b/Engine/Source/Runtime/RenderCore/Public/RayGenShaderUtils.h
--- a/Engine/Source/Runtime/RenderCore/Public/RayGenShaderUtils.h
+++ b/Engine/Source/Runtime/RenderCore/Public/RayGenShaderUtils.h
@@ -6,9 +6,11 @@
 
 #pragma once
 
+#if RHI_RAYTRACING
 #include "RenderResource.h"
 #include "RenderGraphUtils.h"
 #include "PipelineStateCache.h"
+#include "ShaderParameters.h"
 
 /** All utils for ray generation shaders. */
 struct FRayGenShaderUtils
@@ -42,3 +44,4 @@ struct FRayGenShaderUtils
 		});
 	}
 };
+#endif //RHI_RAYTRACINGO
\ No newline at end of file
diff --git a/Engine/Source/Runtime/Renderer/Public/RayTracingDynamicGeometryCollection.h b/Engine/Source/Runtime/Renderer/Public/RayTracingDynamicGeometryCollection.h
--- a/Engine/Source/Runtime/Renderer/Public/RayTracingDynamicGeometryCollection.h
+++ b/Engine/Source/Runtime/Renderer/Public/RayTracingDynamicGeometryCollection.h
@@ -2,6 +2,8 @@
 
 #pragma once
 
+#include "RHIDefinitions.h"
+
 #if RHI_RAYTRACING
 
 #include "RHI.h"
