Fix static code analysis warning

diff --git a/Engine/Source/Runtime/RenderCore/Private/ShaderParameterMetadata.cpp b/Engine/Source/Runtime/RenderCore/Private/ShaderParameterMetadata.cpp
--- a/Engine/Source/Runtime/RenderCore/Private/ShaderParameterMetadata.cpp
+++ b/Engine/Source/Runtime/RenderCore/Private/ShaderParameterMetadata.cpp
@@ -157,7 +157,7 @@ void BuildShaderBindingLayout(TConstArrayView<FShaderParametersMetadata*> Unifor
 	FShaderBindingLayout ShaderBindingLayout;
 
 	EShaderBindingLayoutFlags ShaderBindingLayoutFlags = BaseShaderBindingLayoutFlags;
-	if (BindingType == FShaderBindingLayoutContainer::EBindingType::Bindless)
+	//if (BindingType == FShaderBindingLayoutContainer::EBindingType::Bindless)
 	{
 		ShaderBindingLayoutFlags |= EShaderBindingLayoutFlags::BindlessResources | EShaderBindingLayoutFlags::BindlessSamplers;
 	}
