Fix non-unity
#rnx

diff --git a/Engine/Source/Runtime/RenderCore/Public/RayGenShaderUtils.h b/Engine/Source/Runtime/RenderCore/Public/RayGenShaderUtils.h
--- a/Engine/Source/Runtime/RenderCore/Public/RayGenShaderUtils.h
+++ b/Engine/Source/Runtime/RenderCore/Public/RayGenShaderUtils.h
@@ -6,6 +6,8 @@
 
 #pragma once
 
+#include "RHIDefinitions.h"
+
 #if RHI_RAYTRACING
 #include "RenderResource.h"
 #include "RenderGraphUtils.h"
