Add Shader Maps csv stats to better track Shader regressions in perf tests.

diff --git a/Engine/Source/Runtime/Engine/Private/Materials/MaterialShader.cpp b/Engine/Source/Runtime/Engine/Private/Materials/MaterialShader.cpp
--- a/Engine/Source/Runtime/Engine/Private/Materials/MaterialShader.cpp
+++ b/Engine/Source/Runtime/Engine/Private/Materials/MaterialShader.cpp
@@ -34,6 +34,7 @@
 #include "PSOPrecache.h"
 #include "PSOPrecacheMaterial.h"
 #include "PSOPrecacheValidation.h"
+#include "ProfilingDebugging/CsvProfiler.h"
 
 #if WITH_EDITOR
 #include "Algo/Sort.h"
@@ -49,6 +50,8 @@
 #include "Serialization/MemoryReader.h"
 #endif
 
+CSV_DECLARE_CATEGORY_MODULE_EXTERN(RENDERCORE_API, Shaders);
+
 int32 GMaterialExcludeNonPipelinedShaders = 1;
 static FAutoConsoleVariableRef CVarMaterialExcludeNonPipelinedShaders(
 	TEXT("r.Material.ExcludeNonPipelinedShaders"),
@@ -3321,6 +3324,11 @@ void FMaterialShaderMap::Register(EShaderPlatform InShaderPlatform)
 	if (!bRegistered)
 	{
 		INC_DWORD_STAT(STAT_Shaders_NumShaderMaps);
+
+#if (CSV_PROFILER && !UE_BUILD_SHIPPING) 
+		TCsvPersistentCustomStat<int>* CsvStatNumShaderMaps = FCsvProfiler::Get()->GetOrCreatePersistentCustomStatInt(TEXT("NumShaderMaps"), CSV_CATEGORY_INDEX(Shaders));
+		CsvStatNumShaderMaps->Add(1);
+#endif
 	}
 
 	{
@@ -3398,6 +3406,11 @@ void FMaterialShaderMap::Release()
 				bRegistered = false;
 				DEC_DWORD_STAT(STAT_Shaders_NumShaderMaps);
 
+#if (CSV_PROFILER && !UE_BUILD_SHIPPING) 
+				TCsvPersistentCustomStat<int>* CsvStatNumShaderMaps = FCsvProfiler::Get()->GetOrCreatePersistentCustomStatInt(TEXT("NumShaderMaps"), CSV_CATEGORY_INDEX(Shaders));
+				CsvStatNumShaderMaps->Sub(1);
+#endif
+
 				FMaterialShaderMap *CachedMap = GIdToMaterialShaderMap[GetShaderPlatform()].FindRef(ShaderMapId);
 
 				// Map is marked as registered therefore we do expect it to be in the cache
@@ -3477,6 +3490,11 @@ FMaterialShaderMap::~FMaterialShaderMap()
 		uint32 TotalShadersIncludingBaseClass = GetContent()->GetNumShaders();
 		uint32 OwnShaders = TotalShadersIncludingBaseClass - BaseClassShaders;
 		DEC_DWORD_STAT_BY(STAT_Shaders_NumShadersLoaded, OwnShaders);
+
+#if (CSV_PROFILER && !UE_BUILD_SHIPPING) 
+		TCsvPersistentCustomStat<int>* CsvStatNumShadersLoaded = FCsvProfiler::Get()->GetOrCreatePersistentCustomStatInt(TEXT("NumShadersLoaded"), CSV_CATEGORY_INDEX(Shaders));
+		CsvStatNumShadersLoaded->Sub(OwnShaders);
+#endif
 	}
 #endif
 }
@@ -3542,6 +3560,11 @@ bool FMaterialShaderMap::Serialize(FArchive& Ar, bool bInlineShaderResources, bo
 		uint32 TotalShadersIncludingBaseClass = GetContent()->GetNumShaders();
 		uint32 OwnShaders = TotalShadersIncludingBaseClass - BaseClassShaders;
 		INC_DWORD_STAT_BY(STAT_Shaders_NumShadersLoaded, OwnShaders);
+
+#if (CSV_PROFILER && !UE_BUILD_SHIPPING) 
+		TCsvPersistentCustomStat<int>* CsvStatNumShadersLoaded = FCsvProfiler::Get()->GetOrCreatePersistentCustomStatInt(TEXT("NumShadersLoaded"), CSV_CATEGORY_INDEX(Shaders));
+		CsvStatNumShadersLoaded->Add(OwnShaders);
+#endif
 	}
 #endif // STATS
 	return bSerialized;
diff --git a/Engine/Source/Runtime/RenderCore/Private/ShaderCodeArchive.cpp b/Engine/Source/Runtime/RenderCore/Private/ShaderCodeArchive.cpp
--- a/Engine/Source/Runtime/RenderCore/Private/ShaderCodeArchive.cpp
+++ b/Engine/Source/Runtime/RenderCore/Private/ShaderCodeArchive.cpp
@@ -1327,6 +1327,12 @@ TRefCountPtr<FRHIShader> FShaderCodeArchive::CreateShader(int32 Index)
 	if (Shader)
 	{
 		INC_DWORD_STAT(STAT_Shaders_NumShadersCreated);
+
+#if (CSV_PROFILER && !UE_BUILD_SHIPPING) 
+		TCsvPersistentCustomStat<int>* CsvStatNumShadersCreated = FCsvProfiler::Get()->GetOrCreatePersistentCustomStatInt(TEXT("NumShadersCreated"), CSV_CATEGORY_INDEX(Shaders));
+		CsvStatNumShadersCreated->Add(1);
+#endif
+
 		Shader->SetHash(ShaderHash);
 	}
 
@@ -2256,6 +2262,12 @@ TRefCountPtr<FRHIShader> FIoStoreShaderCodeArchive::CreateShader(int32 ShaderInd
 	if (Shader)
 	{
 		INC_DWORD_STAT(STAT_Shaders_NumShadersCreated);
+
+#if (CSV_PROFILER && !UE_BUILD_SHIPPING) 
+		TCsvPersistentCustomStat<int>* CsvStatNumShadersCreated = FCsvProfiler::Get()->GetOrCreatePersistentCustomStatInt(TEXT("NumShadersCreated"), CSV_CATEGORY_INDEX(Shaders));
+		CsvStatNumShadersCreated->Add(1);
+#endif
+
 		Shader->SetHash(ShaderHash);
 	}
 
diff --git a/Engine/Source/Runtime/RenderCore/Private/ShaderCore.cpp b/Engine/Source/Runtime/RenderCore/Private/ShaderCore.cpp
--- a/Engine/Source/Runtime/RenderCore/Private/ShaderCore.cpp
+++ b/Engine/Source/Runtime/RenderCore/Private/ShaderCore.cpp
@@ -115,6 +115,8 @@ DEFINE_STAT(STAT_Shaders_ShaderMemory);
 DEFINE_STAT(STAT_Shaders_ShaderResourceMemory);
 DEFINE_STAT(STAT_Shaders_ShaderPreloadMemory);
 
+CSV_DEFINE_CATEGORY(Shaders, (!UE_BUILD_SHIPPING));
+
 PRAGMA_DISABLE_DEPRECATION_WARNINGS		// FShaderCompilerDefinitions will be made internal in the future, marked deprecated until then
 
 /**
diff --git a/Engine/Source/Runtime/RenderCore/Private/ShaderMap.cpp b/Engine/Source/Runtime/RenderCore/Private/ShaderMap.cpp
--- a/Engine/Source/Runtime/RenderCore/Private/ShaderMap.cpp
+++ b/Engine/Source/Runtime/RenderCore/Private/ShaderMap.cpp
@@ -14,6 +14,7 @@
 #include "Misc/ScopeLock.h"
 #include "UObject/RenderingObjectVersion.h"
 #include "DataDrivenShaderPlatformInfo.h"
+#include "ProfilingDebugging/CsvProfiler.h"
 
 static EShaderPermutationFlags GetCurrentShaderPermutationFlags()
 {
@@ -86,6 +87,14 @@ void FShaderMapBase::AssignCopy(const FShaderMapBase& Source)
 	INC_DWORD_STAT_BY(STAT_Shaders_ShaderMemory, Content.FrozenSize);
 	INC_DWORD_STAT_BY(STAT_Shaders_NumShadersLoaded, NumFrozenShaders);
 
+#if (CSV_PROFILER && !UE_BUILD_SHIPPING) 
+	TCsvPersistentCustomStat<float>* CsvStatShaderMemoryMB = FCsvProfiler::Get()->GetOrCreatePersistentCustomStatFloat(TEXT("ShaderMemoryMB"), CSV_CATEGORY_INDEX(Shaders));
+	TCsvPersistentCustomStat<int>* CsvStatNumShadersLoaded = FCsvProfiler::Get()->GetOrCreatePersistentCustomStatInt(TEXT("NumShadersLoaded"), CSV_CATEGORY_INDEX(Shaders));
+
+	CsvStatShaderMemoryMB->Add(float(Content.FrozenSize) / (1024.0f * 1024.f));
+	CsvStatNumShadersLoaded->Add(NumFrozenShaders);
+#endif
+
 	Code = new FShaderMapResourceCode(*Source.Code);
 	InitResource();
 }
@@ -109,6 +118,14 @@ void FShaderMapBase::FinalizeContent()
 		NumFrozenShaders = Content.Object->GetNumShaders();
 		INC_DWORD_STAT_BY(STAT_Shaders_ShaderMemory, Content.FrozenSize);
 		INC_DWORD_STAT_BY(STAT_Shaders_NumShadersLoaded, NumFrozenShaders);
+
+#if (CSV_PROFILER && !UE_BUILD_SHIPPING) 
+		TCsvPersistentCustomStat<float>* CsvStatShaderMemoryMB = FCsvProfiler::Get()->GetOrCreatePersistentCustomStatFloat(TEXT("ShaderMemoryMB"), CSV_CATEGORY_INDEX(Shaders));
+		TCsvPersistentCustomStat<int>* CsvStatNumShadersLoaded = FCsvProfiler::Get()->GetOrCreatePersistentCustomStatInt(TEXT("NumShadersLoaded"), CSV_CATEGORY_INDEX(Shaders));
+
+		CsvStatShaderMemoryMB->Add(float(Content.FrozenSize) / (1024.0f * 1024.f));
+		CsvStatNumShadersLoaded->Add(NumFrozenShaders);
+#endif
 	}
 	InitResource();
 }
@@ -117,6 +134,15 @@ void FShaderMapBase::UnfreezeContent()
 {
 	DEC_DWORD_STAT_BY(STAT_Shaders_ShaderMemory, Content.FrozenSize);
 	DEC_DWORD_STAT_BY(STAT_Shaders_NumShadersLoaded, NumFrozenShaders);
+
+#if (CSV_PROFILER && !UE_BUILD_SHIPPING) 
+	TCsvPersistentCustomStat<float>* CsvStatShaderMemoryMB = FCsvProfiler::Get()->GetOrCreatePersistentCustomStatFloat(TEXT("ShaderMemoryMB"), CSV_CATEGORY_INDEX(Shaders));
+	TCsvPersistentCustomStat<int>* CsvStatNumShadersLoaded = FCsvProfiler::Get()->GetOrCreatePersistentCustomStatInt(TEXT("NumShadersLoaded"), CSV_CATEGORY_INDEX(Shaders));
+
+	CsvStatShaderMemoryMB->Sub(float(Content.FrozenSize) / (1024.0f * 1024.f));
+	CsvStatNumShadersLoaded->Sub(NumFrozenShaders);
+#endif
+
 	Content.Unfreeze(PointerTable);
 	NumFrozenShaders = 0u;
 }
@@ -238,6 +264,14 @@ bool FShaderMapBase::Serialize(FArchive& Ar, bool bInlineShaderResources, bool b
 			INC_DWORD_STAT_BY(STAT_Shaders_ShaderMemory, Content.FrozenSize);
 			INC_DWORD_STAT_BY(STAT_Shaders_NumShadersLoaded, NumFrozenShaders);
 
+#if (CSV_PROFILER && !UE_BUILD_SHIPPING) 
+			TCsvPersistentCustomStat<float>* CsvStatShaderMemoryMB = FCsvProfiler::Get()->GetOrCreatePersistentCustomStatFloat(TEXT("ShaderMemoryMB"), CSV_CATEGORY_INDEX(Shaders));
+			TCsvPersistentCustomStat<int>* CsvStatNumShadersLoaded = FCsvProfiler::Get()->GetOrCreatePersistentCustomStatInt(TEXT("NumShadersLoaded"), CSV_CATEGORY_INDEX(Shaders));
+
+			CsvStatShaderMemoryMB->Add(float(Content.FrozenSize) / (1024.0f * 1024.f));
+			CsvStatNumShadersLoaded->Add(NumFrozenShaders);
+#endif
+
 			BeginInitResource(Resource);
 			INC_DWORD_STAT_BY(STAT_Shaders_ShaderResourceMemory, Resource->GetSizeBytes());
 		}
@@ -279,6 +313,15 @@ void FShaderMapBase::DestroyContent()
 {
 	DEC_DWORD_STAT_BY(STAT_Shaders_ShaderMemory, Content.FrozenSize);
 	DEC_DWORD_STAT_BY(STAT_Shaders_NumShadersLoaded, NumFrozenShaders);
+
+#if (CSV_PROFILER && !UE_BUILD_SHIPPING) 
+	TCsvPersistentCustomStat<float>* CsvStatShaderMemoryMB = FCsvProfiler::Get()->GetOrCreatePersistentCustomStatFloat(TEXT("ShaderMemoryMB"), CSV_CATEGORY_INDEX(Shaders));
+	TCsvPersistentCustomStat<int>* CsvStatNumShadersLoaded = FCsvProfiler::Get()->GetOrCreatePersistentCustomStatInt(TEXT("NumShadersLoaded"), CSV_CATEGORY_INDEX(Shaders));
+
+	CsvStatShaderMemoryMB->Sub(float(Content.FrozenSize) / (1024.0f * 1024.f));
+	CsvStatNumShadersLoaded->Sub(NumFrozenShaders);
+#endif
+
 	Content.Destroy(PointerTable);
 	NumFrozenShaders = 0u;
 }
@@ -401,9 +444,9 @@ void FShaderMapContent::RemoveShaderTypePermutaion(const FHashedName& TypeName,
 			DeleteObjectFromLayout(Shader);
 
 			// Replace the shader we're removing with the last shader in the list
-			Shaders.RemoveAtSwap(Index, EAllowShrinking::No);
-			ShaderTypes.RemoveAtSwap(Index, EAllowShrinking::No);
-			ShaderPermutations.RemoveAtSwap(Index, EAllowShrinking::No);
+			Shaders.RemoveAtSwap(Index, 1, EAllowShrinking::No);
+			ShaderTypes.RemoveAtSwap(Index, 1, EAllowShrinking::No);
+			ShaderPermutations.RemoveAtSwap(Index, 1, EAllowShrinking::No);
 			check(ShaderTypes.Num() == Shaders.Num());
 			check(ShaderPermutations.Num() == Shaders.Num());
 			ShaderHash.Remove(Hash, Index);
diff --git a/Engine/Source/Runtime/RenderCore/Private/ShaderResource.cpp b/Engine/Source/Runtime/RenderCore/Private/ShaderResource.cpp
--- a/Engine/Source/Runtime/RenderCore/Private/ShaderResource.cpp
+++ b/Engine/Source/Runtime/RenderCore/Private/ShaderResource.cpp
@@ -28,6 +28,9 @@
 
 DECLARE_LOG_CATEGORY_CLASS(LogShaderWarnings, Log, Log);
 
+#if (CSV_PROFILER && !UE_BUILD_SHIPPING) 
+TCsvPersistentCustomStat<int>* CsvStatNumShaderMapsUsedForRendering = nullptr;
+#endif
 
 static int32 GShaderCompilerEmitWarningsOnLoad = 0;
 static FAutoConsoleVariableRef CVarShaderCompilerEmitWarningsOnLoad(
@@ -469,19 +472,36 @@ void FShaderMapResource::ReleaseShaders()
 	{
 		FScopeLock ScopeLock(&RHIShadersCreationGuard);
 
+		int NumReleaseShaders = 0;
+
 		for (uint32 Idx = 0; Idx < NumRHIShaders; ++Idx)
 		{
 			if (FRHIShader* Shader = RHIShaders[Idx].load(std::memory_order_acquire))
 			{
 				Shader->Release();
+				NumReleaseShaders++;
 				DEC_DWORD_STAT(STAT_Shaders_NumShadersCreated);
 			}
 		}
+
+#if (CSV_PROFILER && !UE_BUILD_SHIPPING) 
+		TCsvPersistentCustomStat<int>* CsvStatNumShadersCreated = FCsvProfiler::Get()->GetOrCreatePersistentCustomStatInt(TEXT("NumShadersCreated"), CSV_CATEGORY_INDEX(Shaders));
+		CsvStatNumShadersCreated->Sub(NumReleaseShaders);
+#endif
+
 		RHIShaders = nullptr;
 		NumRHIShaders = 0;
 		if (bAtLeastOneRHIShaderCreated)
 		{
 			DEC_DWORD_STAT(STAT_Shaders_NumShaderMapsUsedForRendering);
+
+#if (CSV_PROFILER && !UE_BUILD_SHIPPING) 
+			if(CsvStatNumShaderMapsUsedForRendering == nullptr)
+			{
+				CsvStatNumShaderMapsUsedForRendering = FCsvProfiler::Get()->GetOrCreatePersistentCustomStatInt(TEXT("NumShaderMapsUsedForRendering"), CSV_CATEGORY_INDEX(Shaders));
+			}
+			CsvStatNumShaderMapsUsedForRendering->Sub(1);
+#endif
 		}
 		bAtLeastOneRHIShaderCreated = false;
 	}
@@ -563,6 +583,14 @@ FRHIShader* FShaderMapResource::CreateShaderOrCrash(int32 ShaderIndex, bool bReq
 			if (!bAtLeastOneRHIShaderCreated)
 			{
 				INC_DWORD_STAT(STAT_Shaders_NumShaderMapsUsedForRendering);
+
+#if (CSV_PROFILER && !UE_BUILD_SHIPPING) 
+				if (CsvStatNumShaderMapsUsedForRendering == nullptr)
+				{
+					CsvStatNumShaderMapsUsedForRendering = FCsvProfiler::Get()->GetOrCreatePersistentCustomStatInt(TEXT("NumShaderMapsUsedForRendering"), CSV_CATEGORY_INDEX(Shaders));
+				}
+				CsvStatNumShaderMapsUsedForRendering->Add(1);
+#endif
 				bAtLeastOneRHIShaderCreated = 1;
 			}
 
@@ -675,6 +703,12 @@ FRHIShader* FShaderMapResource_InlineCode::CreateRHIShaderOrCrash(int32 ShaderIn
 	}
 
 	INC_DWORD_STAT(STAT_Shaders_NumShadersCreated);
+
+#if (CSV_PROFILER && !UE_BUILD_SHIPPING) 
+	TCsvPersistentCustomStat<int>* CsvStatNumShadersCreated = FCsvProfiler::Get()->GetOrCreatePersistentCustomStatInt(TEXT("NumShadersCreated"), CSV_CATEGORY_INDEX(Shaders));
+	CsvStatNumShadersCreated->Add(1);
+#endif
+
 	RHIShader->SetHash(ShaderHash);
 
 	// contract of this function is to return a shader with an already held reference
diff --git a/Engine/Source/Runtime/RenderCore/Public/ShaderCore.h b/Engine/Source/Runtime/RenderCore/Public/ShaderCore.h
--- a/Engine/Source/Runtime/RenderCore/Public/ShaderCore.h
+++ b/Engine/Source/Runtime/RenderCore/Public/ShaderCore.h
@@ -59,6 +59,7 @@ class ITargetPlatform;
 using FShaderStatVariant = TVariant<bool, float, int32, uint32>;
 DECLARE_INTRINSIC_TYPE_LAYOUT(FShaderStatVariant);
 
+CSV_DECLARE_CATEGORY_MODULE_EXTERN(RENDERCORE_API, Shaders);
 /**
  * Controls whether shader related logs are visible.
  * Note: The runtime verbosity is driven by the console variable 'r.ShaderDevelopmentMode'
