[shaders] handle gracefully the case where we don't have remapping data for a particular line of shader code and just leave the error/warning unmodified. This _shouldn't_ happen in practice, generally is caused by regressions introduced in shaderformat code (where code is modified without preserving line numbers/whitespace). However due to another bug in 5.4 this is currently occurring in live (said bug has been fixed but is slightly involved to bring across to 5.4 so going to avoid this for now). Not remapping just means the filename/line number for an error or warning will be reported on the stripped code rather than the unstripped code, which is slightly less useful but not the end of the world.

diff --git a/Engine/Source/Runtime/RenderCore/Private/ShaderPreprocessTypes.cpp b/Engine/Source/Runtime/RenderCore/Private/ShaderPreprocessTypes.cpp
--- a/Engine/Source/Runtime/RenderCore/Private/ShaderPreprocessTypes.cpp
+++ b/Engine/Source/Runtime/RenderCore/Private/ShaderPreprocessTypes.cpp
@@ -80,14 +80,19 @@ void FShaderDiagnosticRemapper::Remap(FShaderCompilerError& Diagnostic) const
 		}
 
 		FRemapData RemapData = GetRemapData(StrippedLineNum);
+		// if the given line number doesn't exist in the unstripped source/remap data, then just don't bother remapping
+		// typically this only occurs if non-whitespace-preserving changes have been made in a shader format's compile implementation,
+		// which in general should be caught before check-in, but in the case it's not reporting a non-remapped error is better than nothing
+		if (RemapData.IsValid())
+		{
+			FString NewErrorLineString;
+			NewErrorLineString.Reserve(ErrorLineStr.Len());
+			NewErrorLineString.AppendInt(RemapData.LineNumber);
+			NewErrorLineString.Append(ErrorLineStr.GetCharArray().GetData() + LineNumberEnd, ErrorLineStr.Len() - LineNumberEnd);
 
-		FString NewErrorLineString;
-		NewErrorLineString.Reserve(ErrorLineStr.Len());
-		NewErrorLineString.AppendInt(RemapData.LineNumber);
-		NewErrorLineString.Append(ErrorLineStr.GetCharArray().GetData() + LineNumberEnd, ErrorLineStr.Len() - LineNumberEnd);
-
-		Diagnostic.ErrorLineString = MoveTemp(NewErrorLineString);
-		Diagnostic.ErrorVirtualFilePath = RemapData.Filename;
+			Diagnostic.ErrorLineString = MoveTemp(NewErrorLineString);
+			Diagnostic.ErrorVirtualFilePath = RemapData.Filename;
+		}
 	}
 
 	int32 FilenameIndex = Diagnostic.StrippedErrorMessage.Find(GPreprocessConstants.FilenameSentinel.GetData());
@@ -111,6 +116,13 @@ void FShaderDiagnosticRemapper::Remap(FShaderCompilerError& Diagnostic) const
 		}
 
 		FRemapData RemapData = GetRemapData(StrippedLineNum);
+		// if the given line number doesn't exist in the unstripped source/remap data, then just don't bother remapping
+		// typically this only occurs if non-whitespace-preserving changes have been made in a shader format's compile implementation,
+		// which in general should be caught before check-in, but in the case it's not reporting a non-remapped error is better than nothing
+		if (!RemapData.IsValid())
+		{
+			break;
+		}
 
 		// Assume line number is the same number of digits for simplicity in reserve allocation size;
 		// it's an upper bound but in most practical cases worst case we're allocating an extra byte
@@ -164,24 +176,34 @@ void FShaderDiagnosticRemapper::AddStrippedLine(int32 StrippedLineNum, int32 Off
 
 FShaderDiagnosticRemapper::FRemapData FShaderDiagnosticRemapper::GetRemapData(int32 StrippedLineNum) const
 {
-	// Find the index which should contain a block with starting line number greater than the one to which the message points; the preceding
-	// block will be the one containing the stripped line in the unstripped source (in the case UpperBoundBy returns Blocks.Num() this means
-	// the error must be in the last recorded block)
-	int32 FoundIndex = Algo::UpperBoundBy(Blocks, StrippedLineNum, [](const FSourceBlock& Block) { return Block.StrippedLineNum; }) - 1;
+	if (StrippedLineNum >= 1 && StrippedLineNum <= StrippedLineOffsets.Num())
+	{
+		// Find the index which should contain a block with starting line number greater than the one to which the message points; the preceding
+		// block will be the one containing the stripped line in the unstripped source (in the case UpperBoundBy returns Blocks.Num() this means
+		// the error must be in the last recorded block)
+		int32 FoundIndex = Algo::UpperBoundBy(Blocks, StrippedLineNum, [](const FSourceBlock& Block) { return Block.StrippedLineNum; }) - 1;
 
-	// Sanity check - we can't possibly have an error which occurs before the first recorded block unless the stripping didn't record blocks
-	// properly (it only skips recording a block if it's empty in the stripped source, i.e. all comments or whitespace)
-	check(FoundIndex >= 0);
+		// Sanity check - we can't possibly have an error which occurs before the first recorded block unless the stripping didn't record blocks
+		// properly (it only skips recording a block if it's empty in the stripped source, i.e. all comments or whitespace)
+		check(FoundIndex >= 0);
 
-	const FSourceBlock& FoundBlock = Blocks[FoundIndex];
+		const FSourceBlock& FoundBlock = Blocks[FoundIndex];
 
-	// Sanity check to validate the assumption made in AddSourceBlock - we should never be querying remap data for a line number
-	// that doesn't have an associated file path in the unstripped source; this implies a warning or error has been emitted in
-	// system-generated code and should be addressed by the developer prior to submission.
-	checkf(!FoundBlock.OriginalPath.IsEmpty(), TEXT("Unexpected compile error/warning found in system-generated code"));
+		// Sanity check to validate the assumption made in AddSourceBlock - we should never be querying remap data for a line number
+		// that doesn't have an associated file path in the unstripped source; this implies a warning or error has been emitted in
+		// system-generated code and should be addressed by the developer prior to submission.
+		checkf(!FoundBlock.OriginalPath.IsEmpty(), TEXT("Unexpected compile error/warning found in system-generated code"));
 
-	// -1 when indexing stripped line offsets since line numbers are 1-based
-	return FRemapData{ FoundBlock.OriginalPath, FoundBlock.OriginalLineNum + StrippedLineOffsets[StrippedLineNum - 1] };
+		// -1 when indexing stripped line offsets since line numbers are 1-based
+		return FRemapData{ FoundBlock.OriginalPath, FoundBlock.OriginalLineNum + StrippedLineOffsets[StrippedLineNum - 1] };
+	}
+	else
+	{
+		// something went wrong and we're asking for remap data for a line number that didn't exist in the unstripped source.
+		// return an invalid FRemapData which should be handled explicitly by calling code.
+		static FString InvalidPath;
+		return FRemapData{ InvalidPath };
+	}
 }
 
 inline bool IsEndOfTheLine(FShaderSource::CharType C)
diff --git a/Engine/Source/Runtime/RenderCore/Public/ShaderPreprocessTypes.h b/Engine/Source/Runtime/RenderCore/Public/ShaderPreprocessTypes.h
--- a/Engine/Source/Runtime/RenderCore/Public/ShaderPreprocessTypes.h
+++ b/Engine/Source/Runtime/RenderCore/Public/ShaderPreprocessTypes.h
@@ -36,7 +36,9 @@ private:
 	struct FRemapData
 	{
 		const FString& Filename;
-		int32 LineNumber;
+		int32 LineNumber = INDEX_NONE;
+
+		bool IsValid() { return LineNumber != INDEX_NONE; }
 	};
 	FRemapData GetRemapData(int32 StrippedLineNum) const;
 
