[shaders] fix yet another shader error reporting path (non-default materials also have their own path for this); fixed by not relying on calling code to merge the generated error prefix + actual error message

#test forced a shader compile error in both default and non-default materials and ensured output includes additional diagnostic information.

diff --git a/Engine/Source/Runtime/Engine/Private/ShaderCompiler/ShaderCompiler.cpp b/Engine/Source/Runtime/Engine/Private/ShaderCompiler/ShaderCompiler.cpp
--- a/Engine/Source/Runtime/Engine/Private/ShaderCompiler/ShaderCompiler.cpp
+++ b/Engine/Source/Runtime/Engine/Private/ShaderCompiler/ShaderCompiler.cpp
@@ -2812,11 +2812,9 @@ static void BuildErrorStringAndReport(const FShaderDiagnosticInfo& DiagInfo, FSt
 {
 	bool bReportedDebugInfo = false;
 
-	check(DiagInfo.UniqueErrors.Num() == DiagInfo.UniqueErrorPrefixes.Num());
-
 	for (int32 ErrorIndex = 0; ErrorIndex < DiagInfo.UniqueErrors.Num(); ErrorIndex++)
 	{
-		FString UniqueErrorString = DiagInfo.UniqueErrorPrefixes[ErrorIndex] + DiagInfo.UniqueErrors[ErrorIndex] + TEXT("\n");
+		FString UniqueErrorString = DiagInfo.UniqueErrors[ErrorIndex] + TEXT("\n");
 
 		if (FPlatformMisc::IsDebuggerPresent())
 		{
@@ -10452,9 +10450,8 @@ void ProcessCompiledGlobalShaders(const TArray<FShaderCommonCompileJobPtr>& Comp
 	// Report unique errors for global shaders.
 	for (int32 ErrorIndex = 0; ErrorIndex < ShaderDiagInfo.UniqueErrors.Num(); ++ErrorIndex)
 	{
-		FString ErrorString = ShaderDiagInfo.UniqueErrorPrefixes[ErrorIndex] + ShaderDiagInfo.UniqueErrors[ErrorIndex];
 		UE_LOGFMT_NSLOC(LogShaders, Error, "Shaders", "GlobalShaderCompileError", "{ErrorMessage}", 
-			("ErrorMessage", ErrorString));
+			("ErrorMessage", ShaderDiagInfo.UniqueErrors[ErrorIndex]));
 	}
 
 	for (const FString& WarningString : ShaderDiagInfo.UniqueWarnings)
diff --git a/Engine/Source/Runtime/RenderCore/Internal/ShaderDiagnostics.h b/Engine/Source/Runtime/RenderCore/Internal/ShaderDiagnostics.h
--- a/Engine/Source/Runtime/RenderCore/Internal/ShaderDiagnostics.h
+++ b/Engine/Source/Runtime/RenderCore/Internal/ShaderDiagnostics.h
@@ -10,7 +10,6 @@ struct FShaderDiagnosticInfo
 {
 	TArray<FShaderCommonCompileJob*> ErrorJobs;
 	TArray<FString> UniqueErrors;
-	TArray<FString> UniqueErrorPrefixes;
 	TArray<FString> UniqueWarnings;
 	TArray<EShaderPlatform> ErrorPlatforms;
 	FString TargetShaderPlatformString;
@@ -19,7 +18,10 @@ struct FShaderDiagnosticInfo
 
 private:
 	void AddAndProcessErrorsForJob(FShaderCommonCompileJob& Job);
+	int32 AddAndProcessErrorsForFailedJobFiltered(FShaderCompileJob& Job, const TCHAR* FilterMessage);
 	void AddWarningsForJob(const FShaderCommonCompileJob& Job);
+
+	TArray<uint32> UniqueErrorHashes;
 };
 
 RENDERCORE_API FString GetSingleJobCompilationDump(const FShaderCompileJob* SingleJob);
diff --git a/Engine/Source/Runtime/RenderCore/Private/ShaderDiagnostics.cpp b/Engine/Source/Runtime/RenderCore/Private/ShaderDiagnostics.cpp
--- a/Engine/Source/Runtime/RenderCore/Private/ShaderDiagnostics.cpp
+++ b/Engine/Source/Runtime/RenderCore/Private/ShaderDiagnostics.cpp
@@ -31,7 +31,7 @@ static FString ConvertToNativePlatformAbsolutePath(const TCHAR* InPath)
 	return Path;
 }
 
-int32 AddAndProcessErrorsForFailedJobFiltered(FShaderCompileJob& CurrentJob, FShaderDiagnosticInfo& OutShaderDiagInfo, const TCHAR* FilterMessage)
+int32 FShaderDiagnosticInfo::AddAndProcessErrorsForFailedJobFiltered(FShaderCompileJob& CurrentJob, const TCHAR* FilterMessage)
 {
 	int32 NumAddedErrors = 0;
 
@@ -76,7 +76,8 @@ int32 AddAndProcessErrorsForFailedJobFiltered(FShaderCompileJob& CurrentJob, FSh
 				CurrentError.ErrorVirtualFilePath = FString(CurrentJob.Key.ShaderType->GetShaderFilename());
 			}
 
-			if (OutShaderDiagInfo.UniqueErrors.Find(CurrentErrorString) == INDEX_NONE)
+			uint32 ErrorHash = GetTypeHash(CurrentErrorString);
+			if (!UniqueErrorHashes.Contains(ErrorHash))
 			{
 				// build up additional info in a "prefix" string; only do this once for each unique error
 				FString UniqueErrorPrefix;
@@ -108,9 +109,10 @@ int32 AddAndProcessErrorsForFailedJobFiltered(FShaderCompileJob& CurrentJob, FSh
 						*ShaderErrorLineString);
 				}
 
-				OutShaderDiagInfo.UniqueErrors.Add(CurrentErrorString);
-				OutShaderDiagInfo.UniqueErrorPrefixes.Add(UniqueErrorPrefix);
-				OutShaderDiagInfo.ErrorJobs.AddUnique(&CurrentJob);
+				UniqueErrorHashes.Add(ErrorHash);
+
+				UniqueErrors.Add(UniqueErrorPrefix + CurrentErrorString);
+				ErrorJobs.AddUnique(&CurrentJob);
 			}
 			++NumAddedErrors;
 		}
@@ -185,15 +187,16 @@ void FShaderDiagnosticInfo::AddAndProcessErrorsForJob(FShaderCommonCompileJob& J
 			if (Job.Output.Errors.Num() == 0)
 			{
 				// Job hard crashed
-				UniqueErrors.Add(FString::Printf(TEXT("Internal Error!\n\t%s"), *GetSingleJobCompilationDump(&Job)));
-				UniqueErrorPrefixes.Add(TEXT(""));
+				FString InternalErrorStr = FString::Printf(TEXT("Internal Error!\n\t%s"), *GetSingleJobCompilationDump(&Job));
+				UniqueErrors.Add(InternalErrorStr);
+				UniqueErrorHashes.Add(GetTypeHash(InternalErrorStr));
 			}
 
 			// If we filter all error messages because they are interpreted as warnings, we have to assume all error messages are in fact errors and not warnings.
 			// In that case, add jobs again without a filter; e.g. when the stripped message starts with "Internal exception".
-			if (AddAndProcessErrorsForFailedJobFiltered(Job, *this, TEXT("error")) == 0)
+			if (AddAndProcessErrorsForFailedJobFiltered(Job, TEXT("error")) == 0)
 			{
-				AddAndProcessErrorsForFailedJobFiltered(Job, *this, nullptr);
+				AddAndProcessErrorsForFailedJobFiltered(Job, nullptr);
 			}
 
 		});
