Updated IsRayTracingUsingReferenceBasedResidency() to use GetValueOnAnyThread()

diff --git a/Engine/Source/Runtime/RenderCore/Private/RenderUtils.cpp b/Engine/Source/Runtime/RenderCore/Private/RenderUtils.cpp
--- a/Engine/Source/Runtime/RenderCore/Private/RenderUtils.cpp
+++ b/Engine/Source/Runtime/RenderCore/Private/RenderUtils.cpp
@@ -1642,7 +1642,7 @@ bool IsRayTracingUsingReferenceBasedResidency()
 {
 	static const auto CVar = IConsoleManager::Get().FindTConsoleVariableDataBool(TEXT("r.RayTracing.UseReferenceBasedResidency"));
 
-	return GetRayTracingMode() == ERayTracingMode::Dynamic && IsRayTracingEnabled() && CVar && CVar->GetValueOnRenderThread();
+	return GetRayTracingMode() == ERayTracingMode::Dynamic && IsRayTracingEnabled() && CVar && CVar->GetValueOnAnyThread();
 }
 
 bool UseSplineMeshSceneResources(const FStaticShaderPlatform Platform)
