Groundwork for using FRHIBatchedShaderParameters instead of FRayTracingShaderBindings when binding ray generation shader resources

diff --git a/Engine/Source/Runtime/RHI/Private/RHICommandList.cpp b/Engine/Source/Runtime/RHI/Private/RHICommandList.cpp
--- a/Engine/Source/Runtime/RHI/Private/RHICommandList.cpp
+++ b/Engine/Source/Runtime/RHI/Private/RHICommandList.cpp
@@ -2171,3 +2171,99 @@ void FRHICommandList::GenerateMips(FRHITexture*)
 {
 	UE_LOG(LogRHI, Fatal, TEXT("GenerateMips on RHI command lists is deprecated and no longer functions. Use the FGenerateMips helper class from the RenderCore module to generate mips on textures."));
 }
+
+FRayTracingShaderBindings UE::RHI::ConvertRayTracingShaderBindings(const FRHIBatchedShaderParameters& BatchedParameters)
+{
+	FRayTracingShaderBindings Result;
+
+#if RHI_RAYTRACING
+
+	// Use array views for bounds checking
+	TArrayView<FRHITexture*> Textures = Result.Textures;
+	TArrayView<FRHIShaderResourceView*> SRVs = Result.SRVs;
+	TArrayView<FRHIUniformBuffer*> UniformBuffers = Result.UniformBuffers;
+	TArrayView<FRHISamplerState*> Samplers = Result.Samplers;
+	TArrayView<FRHIUnorderedAccessView*> UAVs = Result.UAVs;
+
+	checkf(BatchedParameters.Parameters.IsEmpty(), TEXT("FRHIShaderParameter is not supported by FRayTracingShaderBindings"));
+
+	// TODO: Handle FRHIBatchedShaderParameters::BindlessParameters once supported in FRayTracingShaderBindings
+
+	for (const FRHIShaderParameterResource& It : BatchedParameters.ResourceParameters)
+	{
+		using EType = FRHIShaderParameterResource::EType;
+		switch (It.Type)
+		{
+		case EType::Texture:
+			Textures[It.Index] = static_cast<FRHITexture*>(It.Resource);
+			break;
+		case EType::ResourceView:
+			SRVs[It.Index] = static_cast<FRHIShaderResourceView*>(It.Resource);
+			break;
+		case EType::UnorderedAccessView:
+			UAVs[It.Index] = static_cast<FRHIUnorderedAccessView*>(It.Resource);
+			break;
+		case EType::Sampler:
+			Samplers[It.Index] = static_cast<FRHISamplerState*>(It.Resource);
+			break;
+		case EType::UniformBuffer:
+			UniformBuffers[It.Index] = static_cast<FRHIUniformBuffer*>(It.Resource);
+			break;
+		case EType::ResourceCollection:
+			checkNoEntry(); // not supported
+			break;
+		default:
+			checkNoEntry();
+		}
+	}
+
+#else // // RHI_RAYTRACING
+
+	checkNoEntry();
+
+#endif // RHI_RAYTRACING
+
+	return Result;
+}
+
+void FRHICommandList::RayTraceDispatch(
+	FRayTracingPipelineState* Pipeline,
+	FRHIRayTracingShader* RayGenShader,
+	FRHIRayTracingScene* Scene,
+	FRHIShaderBindingTable* SBT,
+	const FRHIBatchedShaderParameters& GlobalResourceBindings,
+	uint32 Width, uint32 Height)
+{
+#if RHI_RAYTRACING
+
+	FRayTracingShaderBindings LegacyBindings = UE::RHI::ConvertRayTracingShaderBindings(GlobalResourceBindings);
+	RayTraceDispatch(Pipeline, RayGenShader, Scene, SBT, LegacyBindings, Width, Height);
+
+#else // RHI_RAYTRACING
+
+	checkNoEntry();
+
+#endif // RHI_RAYTRACING
+}
+
+void FRHICommandList::RayTraceDispatchIndirect(
+	FRayTracingPipelineState* Pipeline,
+	FRHIRayTracingShader* RayGenShader,
+	FRHIRayTracingScene* Scene,
+	FRHIShaderBindingTable* SBT,
+	const FRHIBatchedShaderParameters& GlobalResourceBindings,
+	FRHIBuffer* ArgumentBuffer, uint32 ArgumentOffset)
+{
+#if RHI_RAYTRACING
+
+	FRayTracingShaderBindings LegacyBindings = UE::RHI::ConvertRayTracingShaderBindings(GlobalResourceBindings);
+	RayTraceDispatchIndirect(Pipeline, RayGenShader, Scene, SBT, LegacyBindings, ArgumentBuffer, ArgumentOffset);
+
+#else // RHI_RAYTRACING
+
+	checkNoEntry();
+
+#endif // RHI_RAYTRACING
+}
+
+
diff --git a/Engine/Source/Runtime/RHI/Public/RHICommandList.h b/Engine/Source/Runtime/RHI/Public/RHICommandList.h
--- a/Engine/Source/Runtime/RHI/Public/RHICommandList.h
+++ b/Engine/Source/Runtime/RHI/Public/RHICommandList.h
@@ -4023,6 +4023,12 @@ public:
 		}
 	}
 
+	/*
+	* Compatibility adaptor that operates on the new FRHIBatchedShaderParameters instead of legacy FRayTracingShaderBindings (planned for deprecation).
+	* This will become the default native code path in a future UE version.
+	*/
+	RHI_API void RayTraceDispatch(FRayTracingPipelineState* Pipeline, FRHIRayTracingShader* RayGenShader, FRHIRayTracingScene* Scene, FRHIShaderBindingTable* SBT, const FRHIBatchedShaderParameters& GlobalResourceBindings, uint32 Width, uint32 Height);
+
 	/**
 	 * Trace rays using dimensions from a GPU buffer containing uint[3], interpreted as number of rays in X, Y and Z dimensions.
 	 * ArgumentBuffer must be in IndirectArgs|SRVCompute state.
@@ -4053,6 +4059,12 @@ public:
 		}
 	}
 
+	/*
+	* Compatibility adaptor that operates on the new FRHIBatchedShaderParameters instead of legacy FRayTracingShaderBindings (planned for deprecation).
+	* This will become the default native code path in a future UE version.
+	*/
+	RHI_API void RayTraceDispatchIndirect(FRayTracingPipelineState* Pipeline, FRHIRayTracingShader* RayGenShader, FRHIRayTracingScene* Scene, FRHIShaderBindingTable* SBT, const FRHIBatchedShaderParameters& GlobalResourceBindings, FRHIBuffer* ArgumentBuffer, uint32 ArgumentOffset);
+
 	UE_DEPRECATED(5.5, "Use FRHIShaderBindingTable instead.")
 	FORCEINLINE_DEBUGGABLE void SetRayTracingBindings(
 		FRHIRayTracingScene* Scene, FRayTracingPipelineState* Pipeline,
@@ -5544,6 +5556,10 @@ namespace UE::RHI
 	//
 	RHI_API void CopySharedMips_AssumeSRVMaskState(FRHICommandList& RHICmdList, FRHITexture* SrcTexture, FRHITexture* DstTexture);
 
+	// Backwards compatibility adaptor to convert new FRHIBatchedShaderParameters to legacy FRayTracingShaderBindings.
+	// This function will be deprecated in a future release, once legacy FRayTracingShaderBindings is removed.
+	RHI_API FRayTracingShaderBindings ConvertRayTracingShaderBindings(const FRHIBatchedShaderParameters& BatchedParameters);
+
 } //! UE::RHI
 
 #undef RHICOMMAND_CALLSTACK
diff --git a/Engine/Source/Runtime/RenderCore/Private/ShaderParameterStruct.cpp b/Engine/Source/Runtime/RenderCore/Private/ShaderParameterStruct.cpp
--- a/Engine/Source/Runtime/RenderCore/Private/ShaderParameterStruct.cpp
+++ b/Engine/Source/Runtime/RenderCore/Private/ShaderParameterStruct.cpp
@@ -678,7 +678,8 @@ static void ExtractShaderParameterResources(
 	TArray<FRHIShaderParameterResource>& OutResourceParameters,
 	TArray<FRHIShaderParameterResource>& OutBindlessParameters,
 	const FShaderParameterBindings& Bindings,
-	TConstArrayView<uint8> ParametersData)
+	TConstArrayView<uint8> ParametersData,
+	EUniformBufferBindingFlags UniformBufferBindingFlags)
 {
 	const FShaderParameterReader Reader(ParametersData);
 
@@ -710,7 +711,7 @@ static void ExtractShaderParameterResources(
 		for (const FShaderParameterBindings::FParameterStructReference& Parameter : Bindings.GraphUniformBuffers)
 		{
 			const FRDGUniformBufferBinding& UniformBufferBinding = Reader.Read<FRDGUniformBufferBinding>(Parameter);
-			if (UniformBufferBinding.IsShader())
+			if (EnumHasAnyFlags(UniformBufferBinding.GetBindingFlags(), UniformBufferBindingFlags))
 			{
 				UniformBufferBinding->MarkResourceAsUsed();
 
@@ -721,7 +722,7 @@ static void ExtractShaderParameterResources(
 		for (const FShaderParameterBindings::FParameterStructReference& Parameter : Bindings.ParameterReferences)
 		{
 			const FUniformBufferBinding& UniformBufferBinding = Reader.Read<FUniformBufferBinding>(Parameter);
-			if (UniformBufferBinding.IsShader())
+			if (EnumHasAnyFlags(UniformBufferBinding.GetBindingFlags(), UniformBufferBindingFlags))
 			{
 				OutResourceParameters.Emplace(UniformBufferBinding.GetUniformBuffer(), GetParameterIndex(Parameter));
 			}
@@ -732,7 +733,8 @@ static void ExtractShaderParameterResources(
 static void ExtractShaderParameterResources(
 	FRHIBatchedShaderParameters& BatchedParameters,
 	const FShaderParameterBindings& Bindings,
-	TConstArrayView<uint8> ParametersData)
+	TConstArrayView<uint8> ParametersData,
+	EUniformBufferBindingFlags UniformBufferBindingFlags)
 {
 	const FShaderParameterReader Reader(ParametersData);
 
@@ -760,7 +762,7 @@ static void ExtractShaderParameterResources(
 		for (const FShaderParameterBindings::FParameterStructReference& Parameter : Bindings.GraphUniformBuffers)
 		{
 			const FRDGUniformBufferBinding& UniformBufferBinding = Reader.Read<FRDGUniformBufferBinding>(Parameter);
-			if (UniformBufferBinding.IsShader())
+			if (EnumHasAnyFlags(UniformBufferBinding.GetBindingFlags(), UniformBufferBindingFlags))
 			{
 				UniformBufferBinding->MarkResourceAsUsed();
 
@@ -771,7 +773,7 @@ static void ExtractShaderParameterResources(
 		for (const FShaderParameterBindings::FParameterStructReference& Parameter : Bindings.ParameterReferences)
 		{
 			const FUniformBufferBinding& UniformBufferBinding = Reader.Read<FUniformBufferBinding>(Parameter);
-			if (UniformBufferBinding.IsShader())
+			if (EnumHasAnyFlags(UniformBufferBinding.GetBindingFlags(), UniformBufferBindingFlags))
 			{
 				BatchedParameters.AddResourceParameter(UniformBufferBinding.GetUniformBuffer(), GetParameterIndex(Parameter));
 			}
@@ -793,7 +795,7 @@ void SetShaderParameters(
 		BatchedParameters.SetShaderParameter(Parameter.BufferIndex, Parameter.BaseIndex, Parameter.ByteSize, FullParametersData.GetData() + Parameter.ByteOffset);
 	}
 
-	ExtractShaderParameterResources(BatchedParameters, Bindings, FullParametersData);
+	ExtractShaderParameterResources(BatchedParameters, Bindings, FullParametersData, EUniformBufferBindingFlags::Shader);
 }
 
 /** Set shader's parameters from its parameters struct. */
@@ -806,6 +808,8 @@ inline void SetShaderParametersInternal(
 	const void* InParametersData)
 {
 	checkf(Bindings.RootParameterBufferIndex == FShaderParameterBindings::kInvalidBufferIndex, TEXT("Can't use SetShaderParameters() for root parameter buffer index."));
+	checkf(!IsRayTracingShaderFrequency(ShaderRHI->GetFrequency()), TEXT("Can't use SetShaderParameters() with RHICmdList parameter for ray tracing shaders."));
+
 	FRHIBatchedShaderParameters& ShaderParameters = RHICmdList.GetScratchShaderParameters();
 	SetShaderParameters(ShaderParameters, Bindings, ParametersMetadata, InParametersData);
 	RHICmdList.SetBatchedShaderParameters(ShaderRHI, ShaderParameters);
@@ -939,4 +943,130 @@ void SetShaderParameters(
 		RTBindingsWriter.SetUniformBuffer(Bindings.RootParameterBufferIndex, RTBindingsWriter.RootUniformBuffer);
 	}
 }
+
+FRayTracingShaderBindings ConvertRayTracingShaderBindings(const FRHIBatchedShaderParameters& BatchedParameters)
+{
+	FRayTracingShaderBindings Result;
+
+	// Use array views for bounds checking
+	TArrayView<FRHITexture*> Textures = Result.Textures;
+	TArrayView<FRHIShaderResourceView*> SRVs = Result.SRVs;
+	TArrayView<FRHIUniformBuffer*> UniformBuffers = Result.UniformBuffers;
+	TArrayView<FRHISamplerState*> Samplers = Result.Samplers;
+	TArrayView<FRHIUnorderedAccessView*> UAVs = Result.UAVs;
+
+	checkf(BatchedParameters.Parameters.IsEmpty(), TEXT("FRHIShaderParameter is not supported by FRayTracingShaderBindings"));
+
+	// TODO: Handle FRHIBatchedShaderParameters::BindlessParameters once supported in FRayTracingShaderBindings
+
+	for (const FRHIShaderParameterResource& It : BatchedParameters.ResourceParameters)
+	{
+		using EType = FRHIShaderParameterResource::EType;
+		switch (It.Type)
+		{
+			case EType::Texture:
+				Textures[It.Index] = static_cast<FRHITexture*>(It.Resource);
+				break;
+			case EType::ResourceView:
+				SRVs[It.Index] = static_cast<FRHIShaderResourceView*>(It.Resource);
+				break;
+			case EType::UnorderedAccessView:
+				UAVs[It.Index] = static_cast<FRHIUnorderedAccessView*>(It.Resource);
+				break;
+			case EType::Sampler:
+				Samplers[It.Index] = static_cast<FRHISamplerState*>(It.Resource);
+				break;
+			case EType::UniformBuffer:
+				UniformBuffers[It.Index] = static_cast<FRHIUniformBuffer*>(It.Resource);
+				break;
+			case EType::ResourceCollection:
+				checkNoEntry(); // not supported
+				break;
+			default:
+				checkNoEntry();
+		}
+	}
+
+	return Result;
+}
+
+void SetRayTracingShaderParameters(
+	FRHIBatchedShaderParameters& BatchedParameters,
+	const FShaderParameterBindings& Bindings,
+	const FShaderParametersMetadata* ParametersMetadata,
+	const void* InParametersData)
+{
+	TConstArrayView<uint8> FullParametersData((const uint8*)InParametersData, ParametersMetadata->GetSize());
+
+#if 0 // Loose parameter binding could be supported if the need arises, but we currently don't support it.
+
+	for (const FShaderParameterBindings::FParameter& Parameter : Bindings.Parameters)
+	{
+		BatchedParameters.SetShaderParameter(Parameter.BufferIndex, Parameter.BaseIndex, Parameter.ByteSize, FullParametersData.GetData() + Parameter.ByteOffset);
+	}
+
+#else
+
+	checkf(Bindings.Parameters.Num() == 0, TEXT("Ray tracing shader should use SHADER_USE_ROOT_PARAMETER_STRUCT() to passdown the cbuffer layout to the shader compiler."));
+
+#endif
+
+	ExtractShaderParameterResources(BatchedParameters, Bindings, FullParametersData, EUniformBufferBindingFlags::StaticAndShader);
+
+	// Root uniform buffer.
+	if (Bindings.RootParameterBufferIndex != FShaderParameterBindings::kInvalidBufferIndex)
+	{
+		// Do not do any validation at some resources may have been removed from the structure because known to not be used by the shader.
+		EUniformBufferValidation Validation = EUniformBufferValidation::None;
+
+		FUniformBufferRHIRef RootUniformBuffer = RHICreateUniformBuffer(InParametersData, ParametersMetadata->GetLayoutPtr(), UniformBuffer_SingleDraw, Validation);
+		BatchedParameters.AddResourceParameter(RootUniformBuffer, Bindings.RootParameterBufferIndex);
+	}
+
+#if DO_CHECK
+	{
+		FRayTracingShaderBindingsWriter RTBindingsWriter;
+		SetShaderParameters(RTBindingsWriter, Bindings, ParametersMetadata, InParametersData);
+
+		FRayTracingShaderBindings ConvertedBindings = ConvertRayTracingShaderBindings(BatchedParameters);
+
+		{
+			for (int32 i = 0; i < UE_ARRAY_COUNT(ConvertedBindings.Textures); ++i)
+			{
+				check(RTBindingsWriter.Textures[i] == ConvertedBindings.Textures[i]);
+			}
+
+			for (int32 i = 0; i < UE_ARRAY_COUNT(ConvertedBindings.SRVs); ++i)
+			{
+				check(RTBindingsWriter.SRVs[i] == ConvertedBindings.SRVs[i]);
+			}
+
+			for (int32 i = 0; i < UE_ARRAY_COUNT(ConvertedBindings.UniformBuffers); ++i)
+			{
+				if (RTBindingsWriter.UniformBuffers[i])
+				{
+					check(ConvertedBindings.UniformBuffers[i]);
+					check(RTBindingsWriter.UniformBuffers[i]->GetLayout() == ConvertedBindings.UniformBuffers[i]->GetLayout());
+				}
+				else
+				{
+					check(ConvertedBindings.UniformBuffers[i] == nullptr);
+				}
+			}
+
+			for (int32 i = 0; i < UE_ARRAY_COUNT(ConvertedBindings.Samplers); ++i)
+			{
+				check(RTBindingsWriter.Samplers[i] == ConvertedBindings.Samplers[i]);
+			}
+
+			for (int32 i = 0; i < UE_ARRAY_COUNT(ConvertedBindings.UAVs); ++i)
+			{
+				check(RTBindingsWriter.UAVs[i] == ConvertedBindings.UAVs[i]);
+			}
+		}
+	}
+#endif // DO_CHECK
+}
+
+
 #endif // RHI_RAYTRACING
diff --git a/Engine/Source/Runtime/RenderCore/Public/ShaderParameterStruct.h b/Engine/Source/Runtime/RenderCore/Public/ShaderParameterStruct.h
--- a/Engine/Source/Runtime/RenderCore/Public/ShaderParameterStruct.h
+++ b/Engine/Source/Runtime/RenderCore/Public/ShaderParameterStruct.h
@@ -239,7 +239,17 @@ inline void SetShaderParameters(
 	const typename TShaderClass::FParameters& Parameters)
 {
 	ValidateShaderParameters(Shader, ParametersMetadata, &Parameters);
-	SetShaderParameters(BatchedParameters, Shader->Bindings, ParametersMetadata, &Parameters);
+
+#if RHI_RAYTRACING
+	if (IsRayTracingShaderFrequency(Shader->GetFrequency()))
+	{
+		SetRayTracingShaderParameters(BatchedParameters, Shader->Bindings, ParametersMetadata, &Parameters);
+	}
+	else
+#endif // RHI_RAYTRACING
+	{
+		SetShaderParameters(BatchedParameters, Shader->Bindings, ParametersMetadata, &Parameters);
+	}
 }
 
 template<typename TShaderClass>
@@ -294,6 +304,17 @@ inline void SetShaderParameters(TRHICmdList& RHICmdList, const TShaderRef<TShade
 
 #if RHI_RAYTRACING
 
+/**
+* Similar to SetShaderParameters(), but also binds static uniform buffers that are skipped otherwise.
+* This helper function exists in UE 5.5 to aid with compatibility between FRHIBatchedShaderParameters and legacy FRayTracingShaderBindings[Writer].
+* It will be deprecated in a future release, once legacy FRayTracingShaderBindings[Writer] is removed.
+*/
+RENDERCORE_API void SetRayTracingShaderParameters(
+	FRHIBatchedShaderParameters& BatchedParameters,
+	const FShaderParameterBindings& Bindings,
+	const FShaderParametersMetadata* ParametersMetadata,
+	const void* ParametersData);
+
 RENDERCORE_API void SetShaderParameters(
 	FRayTracingShaderBindingsWriter& RTBindingsWriter,
 	const FShaderParameterBindings& Bindings,
diff --git a/Engine/Source/Runtime/Renderer/Private/RayTracing/RayTracingBarycentrics.cpp b/Engine/Source/Runtime/Renderer/Private/RayTracing/RayTracingBarycentrics.cpp
--- a/Engine/Source/Runtime/Renderer/Private/RayTracing/RayTracingBarycentrics.cpp
+++ b/Engine/Source/Runtime/Renderer/Private/RayTracing/RayTracingBarycentrics.cpp
@@ -180,13 +180,13 @@ void RenderRayTracingBarycentricsRGS(FRDGBuilder& GraphBuilder, const FScene& Sc
 		ERDGPassFlags::Compute,
 		[RayGenParameters, RayGenShader, &View, SBT, Pipeline, ViewRect](FRHICommandList& RHICmdList)
 	{
-		FRayTracingShaderBindingsWriter GlobalResources;
-		SetShaderParameters(GlobalResources, RayGenShader, *RayGenParameters);
+		FRHIBatchedShaderParameters& ShaderParameters = RHICmdList.GetScratchShaderParameters();
+		SetShaderParameters(ShaderParameters, RayGenShader, *RayGenParameters);
 
 		// Dispatch rays using default shader binding table
 		RHICmdList.SetRayTracingMissShader(SBT, 0, Pipeline, 0 /* ShaderIndexInPipeline */, 0, nullptr, 0);
 		RHICmdList.CommitShaderBindingTable(SBT);
-		RHICmdList.RayTraceDispatch(Pipeline, RayGenShader.GetRayTracingShader(), View.GetRayTracingSceneChecked(), SBT, GlobalResources, ViewRect.Size().X, ViewRect.Size().Y);
+		RHICmdList.RayTraceDispatch(Pipeline, RayGenShader.GetRayTracingShader(), View.GetRayTracingSceneChecked(), SBT, ShaderParameters, ViewRect.Size().X, ViewRect.Size().Y);
 	});
 }
 
diff --git a/Engine/Source/Runtime/Renderer/Private/RayTracing/RayTracingPrimaryRays.cpp b/Engine/Source/Runtime/Renderer/Private/RayTracing/RayTracingPrimaryRays.cpp
--- a/Engine/Source/Runtime/Renderer/Private/RayTracing/RayTracingPrimaryRays.cpp
+++ b/Engine/Source/Runtime/Renderer/Private/RayTracing/RayTracingPrimaryRays.cpp
@@ -186,11 +186,11 @@ void FDeferredShadingSceneRenderer::RenderRayTracingPrimaryRaysView(
 		FRHIShaderBindingTable* SBT = View.RayTracingSBT;
 		FRayTracingPipelineState* Pipeline = View.RayTracingMaterialPipeline;
 
-		FRayTracingShaderBindingsWriter GlobalResources;
-		SetShaderParameters(GlobalResources, RayGenShader, *PassParameters);
+		FRHIBatchedShaderParameters& ShaderParameters = RHICmdList.GetScratchShaderParameters();
+		SetShaderParameters(ShaderParameters, RayGenShader, *PassParameters);
 
 		FRHIRayTracingScene* RayTracingSceneRHI = View.GetRayTracingSceneChecked();
-		RHICmdList.RayTraceDispatch(Pipeline, RayGenShader.GetRayTracingShader(), RayTracingSceneRHI, SBT, GlobalResources, RayTracingResolution.X, RayTracingResolution.Y);
+		RHICmdList.RayTraceDispatch(Pipeline, RayGenShader.GetRayTracingShader(), RayTracingSceneRHI, SBT, ShaderParameters, RayTracingResolution.X, RayTracingResolution.Y);
 	});
 }
 
