[RHI] Added shader map and shader library listing commands that let you analyse the memory footprint of those structures.

#tests local

diff --git a/Engine/Plugins/Compositing/OpenColorIO/Source/OpenColorIO/Private/OpenColorIOShaderMap.cpp b/Engine/Plugins/Compositing/OpenColorIO/Source/OpenColorIO/Private/OpenColorIOShaderMap.cpp
--- a/Engine/Plugins/Compositing/OpenColorIO/Source/OpenColorIO/Private/OpenColorIOShaderMap.cpp
+++ b/Engine/Plugins/Compositing/OpenColorIO/Source/OpenColorIO/Private/OpenColorIOShaderMap.cpp
@@ -551,6 +551,16 @@ void FOpenColorIOShaderMap::GetShaderList(TMap<FShaderId, TShaderRef<FShader>>&
 	GetContent()->GetShaderList(*this, FSHAHash(), OutShaders);
 }
 
+void FOpenColorIOShaderMap::GetShaderList(TMap<FHashedName, TShaderRef<FShader>>& OutShaders) const
+{
+	GetContent()->GetShaderList(*this, OutShaders);
+}
+
+void FOpenColorIOShaderMap::GetShaderPipelineList(TArray<FShaderPipelineRef>& OutShaderPipelines) const
+{
+	GetContent()->GetShaderPipelineList(*this, OutShaderPipelines, FShaderPipeline::EAll);
+}
+
 /**
  * Registers a OpenColorIO shader map in the global map.
  */
diff --git a/Engine/Plugins/Compositing/OpenColorIO/Source/OpenColorIO/Public/OpenColorIOShared.h b/Engine/Plugins/Compositing/OpenColorIO/Source/OpenColorIO/Public/OpenColorIOShared.h
--- a/Engine/Plugins/Compositing/OpenColorIO/Source/OpenColorIO/Public/OpenColorIOShared.h
+++ b/Engine/Plugins/Compositing/OpenColorIO/Source/OpenColorIO/Public/OpenColorIOShared.h
@@ -213,7 +213,9 @@ public:
 #endif // WITH_EDITOR
 
 	/** Builds a list of the shaders in a shader map. */
-	OPENCOLORIO_API  void GetShaderList(TMap<FShaderId, TShaderRef<FShader>>& OutShaders) const;
+	OPENCOLORIO_API void GetShaderList(TMap<FShaderId, TShaderRef<FShader>>& OutShaders) const;
+	OPENCOLORIO_API virtual void GetShaderList(TMap<FHashedName, TShaderRef<FShader>>& OutShaders) const override;
+	OPENCOLORIO_API virtual void GetShaderPipelineList(TArray<FShaderPipelineRef>& OutShaderPipelines) const override;
 
 	/** Registers a OpenColorIO shader map in the global map so it can be used by OpenColorIO ColorTransform. */
 	void Register(EShaderPlatform InShaderPlatform);
diff --git a/Engine/Plugins/FX/Niagara/Source/NiagaraShader/Private/NiagaraShader.cpp b/Engine/Plugins/FX/Niagara/Source/NiagaraShader/Private/NiagaraShader.cpp
--- a/Engine/Plugins/FX/Niagara/Source/NiagaraShader/Private/NiagaraShader.cpp
+++ b/Engine/Plugins/FX/Niagara/Source/NiagaraShader/Private/NiagaraShader.cpp
@@ -70,10 +70,6 @@ namespace NiagaraShaderCookStats
 //
 FCriticalSection GIdToNiagaraShaderMapCS;
 TMap<FNiagaraShaderMapId, FNiagaraShaderMap*> FNiagaraShaderMap::GIdToNiagaraShaderMap[SP_NumPlatforms];
-#if ALLOW_SHADERMAP_DEBUG_DATA
-TArray<FNiagaraShaderMap*> FNiagaraShaderMap::AllNiagaraShaderMaps;
-FCriticalSection FNiagaraShaderMap::AllNiagaraShaderMapsGuard;
-#endif
 
 #if WITH_EDITOR
 TMap<FNiagaraShaderMapRef, TArray<FNiagaraShaderScript*>> FNiagaraShaderMap::NiagaraShaderMapsBeingCompiled;
@@ -962,6 +958,16 @@ void FNiagaraShaderMap::GetShaderList(TMap<FShaderId, TShaderRef<FShader>>& OutS
 	GetContent()->GetShaderList(*this, FSHAHash(), OutShaders);
 }
 
+void FNiagaraShaderMap::GetShaderList(TMap<FHashedName, TShaderRef<FShader>>& OutShaders) const
+{
+	GetContent()->GetShaderList(*this, OutShaders);
+}
+
+void FNiagaraShaderMap::GetShaderPipelineList(TArray<FShaderPipelineRef>& OutShaderPipelines) const
+{
+	GetContent()->GetShaderPipelineList(*this, OutShaderPipelines, FShaderPipeline::EAll);
+}
+
 /**
  * Registers a Niagara shader map in the global map so it can be used by scripts.
  */
@@ -1028,10 +1034,6 @@ FNiagaraShaderMap::FNiagaraShaderMap() :
 	bIsPersistent(true) 
 {
 	checkSlow(IsInGameThread() || IsAsyncLoading());
-#if ALLOW_SHADERMAP_DEBUG_DATA
-	FScopeLock AllSMAccess(&AllNiagaraShaderMapsGuard);
-	AllNiagaraShaderMaps.Add(this);
-#endif
 }
 
 #if WITH_EDITOR
@@ -1044,10 +1046,6 @@ FNiagaraShaderMap::FNiagaraShaderMap(EWorkerThread)
 	, bCompiledSuccessfully(true)
 	, bIsPersistent(true) 
 {
-#if ALLOW_SHADERMAP_DEBUG_DATA
-	FScopeLock AllSMAccess(&AllNiagaraShaderMapsGuard);
-	AllNiagaraShaderMaps.Add(this);
-#endif
 }
 #endif // WITH_EDITOR
 
@@ -1056,10 +1054,6 @@ FNiagaraShaderMap::~FNiagaraShaderMap()
 	checkSlow(IsInGameThread() || IsAsyncLoading());
 	check(bDeletedThroughDeferredCleanup);
 	check(!bRegistered);
-#if ALLOW_SHADERMAP_DEBUG_DATA
-	FScopeLock AllSMAccess(&AllNiagaraShaderMapsGuard);
-	AllNiagaraShaderMaps.RemoveSwap(this);
-#endif
 }
 
 bool FNiagaraShaderMap::Serialize(FArchive& Ar, bool bInlineShaderResources, bool bLoadedByCookedMaterial)
diff --git a/Engine/Plugins/FX/Niagara/Source/NiagaraShader/Public/NiagaraShared.h b/Engine/Plugins/FX/Niagara/Source/NiagaraShader/Public/NiagaraShared.h
--- a/Engine/Plugins/FX/Niagara/Source/NiagaraShader/Public/NiagaraShared.h
+++ b/Engine/Plugins/FX/Niagara/Source/NiagaraShader/Public/NiagaraShared.h
@@ -576,10 +576,9 @@ public:
 #endif // WITH_EDITOR
 
 	/** Builds a list of the shaders in a shader map. */
-	NIAGARASHADER_API  void GetShaderList(TMap<FShaderId, TShaderRef<FShader>>& OutShaders) const;
-
-	/** Builds a list of the shader pipelines in a shader map. */
-	//ENGINE_API void GetShaderPipelineList(TArray<FShaderPipeline*>& OutShaderPipelines) const;
+	NIAGARASHADER_API void GetShaderList(TMap<FShaderId, TShaderRef<FShader>>& OutShaders) const;
+	NIAGARASHADER_API virtual void GetShaderList(TMap<FHashedName, TShaderRef<FShader>>& OutShaders) const override;
+	NIAGARASHADER_API virtual void GetShaderPipelineList(TArray<FShaderPipelineRef>& OutShaderPipelines) const override;
 
 	/** Registers a niagara shader map in the global map so it can be used by Niagara scripts. */
 	void Register(EShaderPlatform InShaderPlatform);
@@ -644,17 +643,6 @@ private:
 	*/
 	static TMap<FNiagaraShaderMapId, FNiagaraShaderMap*> GIdToNiagaraShaderMap[SP_NumPlatforms];
 
-#if ALLOW_SHADERMAP_DEBUG_DATA
-	/**
-	* All script shader maps in memory.
-	* No ref counting needed as these are removed on destruction of the shader map.
-	*/
-	static TArray<FNiagaraShaderMap*> AllNiagaraShaderMaps;
-
-	/** Guards access to AllNiagaraShaderMaps, which can be written to from an async loading thread. */
-	static FCriticalSection AllNiagaraShaderMapsGuard;
-#endif
-
 #if WITH_EDITOR
 	/** Tracks resources and their shader maps that need to be compiled but whose compilation is being deferred. */
 	static TMap<FNiagaraShaderMapRef, TArray<FNiagaraShaderScript*> > NiagaraShaderMapsBeingCompiled;
diff --git a/Engine/Plugins/Runtime/ComputeFramework/Source/ComputeFramework/Private/ComputeFramework/ComputeKernelShared.h b/Engine/Plugins/Runtime/ComputeFramework/Source/ComputeFramework/Private/ComputeFramework/ComputeKernelShared.h
--- a/Engine/Plugins/Runtime/ComputeFramework/Source/ComputeFramework/Private/ComputeFramework/ComputeKernelShared.h
+++ b/Engine/Plugins/Runtime/ComputeFramework/Source/ComputeFramework/Private/ComputeFramework/ComputeKernelShared.h
@@ -194,7 +194,9 @@ public:
 #endif // WITH_EDITOR
 
 	/** Builds a list of the shaders in a shader map. */
-	 void GetShaderList(TMap<FShaderId, TShaderRef<FShader>>& OutShaders) const;
+	void GetShaderList(TMap<FShaderId, TShaderRef<FShader>>& OutShaders) const;
+	virtual void GetShaderList(TMap<FHashedName, TShaderRef<FShader>>& OutShaders) const override;
+	virtual void GetShaderPipelineList(TArray<FShaderPipelineRef>& OutShaderPipelines) const override;
 
 	/** Registers a ComputeKernel shader map in the global map so it can be used. */
 	void Register(EShaderPlatform InShaderPlatform);
diff --git a/Engine/Plugins/Runtime/ComputeFramework/Source/ComputeFramework/Private/ComputeKernelShaderMap.cpp b/Engine/Plugins/Runtime/ComputeFramework/Source/ComputeFramework/Private/ComputeKernelShaderMap.cpp
--- a/Engine/Plugins/Runtime/ComputeFramework/Source/ComputeFramework/Private/ComputeKernelShaderMap.cpp
+++ b/Engine/Plugins/Runtime/ComputeFramework/Source/ComputeFramework/Private/ComputeKernelShaderMap.cpp
@@ -596,6 +596,16 @@ void FComputeKernelShaderMap::GetShaderList(TMap<FShaderId, TShaderRef<FShader>>
 	GetContent()->GetShaderList(*this, FSHAHash(), OutShaders);
 }
 
+void FComputeKernelShaderMap::GetShaderList(TMap<FHashedName, TShaderRef<FShader>>& OutShaders) const
+{
+	GetContent()->GetShaderList(*this, OutShaders);
+}
+
+void FComputeKernelShaderMap::GetShaderPipelineList(TArray<FShaderPipelineRef>& OutShaderPipelines) const
+{
+	GetContent()->GetShaderPipelineList(*this, OutShaderPipelines, FShaderPipeline::EAll);
+}
+
 /**
  * Registers a ComputeKernel shader map in the global map.
  */
diff --git a/Engine/Source/Programs/UnrealBuildTool/Configuration/ReadOnlyTargetRules.cs b/Engine/Source/Programs/UnrealBuildTool/Configuration/ReadOnlyTargetRules.cs
--- a/Engine/Source/Programs/UnrealBuildTool/Configuration/ReadOnlyTargetRules.cs
+++ b/Engine/Source/Programs/UnrealBuildTool/Configuration/ReadOnlyTargetRules.cs
@@ -172,6 +172,8 @@ namespace UnrealBuildTool
 
 		public bool bUseIris => Inner.bUseIris;
 
+		public bool bTrackRHIResourceInfo => Inner.bTrackRHIResourceInfo;
+
 		[Obsolete("Deprecated in UE5.4 - No longer used.")]
 		public bool bCompileIntelMetricsDiscovery => Inner.bCompileIntelMetricsDiscovery;
 
diff --git a/Engine/Source/Programs/UnrealBuildTool/Configuration/TargetRules.cs b/Engine/Source/Programs/UnrealBuildTool/Configuration/TargetRules.cs
--- a/Engine/Source/Programs/UnrealBuildTool/Configuration/TargetRules.cs
+++ b/Engine/Source/Programs/UnrealBuildTool/Configuration/TargetRules.cs
@@ -1015,6 +1015,13 @@ namespace UnrealBuildTool
 		[CommandLine("-UseIris", Value = "true")]
 		public bool bUseIris { get; set; } = true;
 
+		/// <summary>
+		/// Whether to track owner (asset name) of RHI resource for non Shipping configurations.
+		/// Useful for ListShaderMaps and ListShaderLibraries commands.
+		/// </summary>
+		[RequiresUniqueBuildEnvironment]
+		public bool bTrackRHIResourceInfo { get; set; } = false;
+
 		/// <summary>
 		/// Whether we are compiling editor code or not. Prefer the more explicit bCompileAgainstEditor instead.
 		/// </summary>
diff --git a/Engine/Source/Runtime/Apple/MetalRHI/Private/Shaders/MetalShaderLibrary.cpp b/Engine/Source/Runtime/Apple/MetalRHI/Private/Shaders/MetalShaderLibrary.cpp
--- a/Engine/Source/Runtime/Apple/MetalRHI/Private/Shaders/MetalShaderLibrary.cpp
+++ b/Engine/Source/Runtime/Apple/MetalRHI/Private/Shaders/MetalShaderLibrary.cpp
@@ -96,6 +96,11 @@ int32 FMetalShaderLibrary::GetNumShaderMaps() const
 	return SerializedShaders.ShaderMapEntries.Num();
 }
 
+uint32 FMetalShaderLibrary::GetSizeBytes() const
+{
+	return SerializedShaders.GetAllocatedSize() + ShaderCode.GetAllocatedSize();
+}
+
 int32 FMetalShaderLibrary::GetNumShadersForShaderMap(int32 ShaderMapIndex) const
 {
 	return SerializedShaders.ShaderMapEntries[ShaderMapIndex].NumShaders;
diff --git a/Engine/Source/Runtime/Apple/MetalRHI/Private/Shaders/MetalShaderLibrary.h b/Engine/Source/Runtime/Apple/MetalRHI/Private/Shaders/MetalShaderLibrary.h
--- a/Engine/Source/Runtime/Apple/MetalRHI/Private/Shaders/MetalShaderLibrary.h
+++ b/Engine/Source/Runtime/Apple/MetalRHI/Private/Shaders/MetalShaderLibrary.h
@@ -31,6 +31,7 @@ public:
 
 	virtual int32 GetNumShaders() const override;
 	virtual int32 GetNumShaderMaps() const override;
+	virtual uint32 GetSizeBytes() const override;
 	virtual int32 GetNumShadersForShaderMap(int32 ShaderMapIndex) const override;
 	virtual int32 GetShaderIndex(int32 ShaderMapIndex, int32 i) const override;
 
diff --git a/Engine/Source/Runtime/Engine/Public/MaterialShared.h b/Engine/Source/Runtime/Engine/Public/MaterialShared.h
--- a/Engine/Source/Runtime/Engine/Public/MaterialShared.h
+++ b/Engine/Source/Runtime/Engine/Public/MaterialShared.h
@@ -1457,10 +1457,10 @@ public:
 	ENGINE_API void GetShaderList(TMap<FShaderId, TShaderRef<FShader>>& OutShaders) const;
 
 	/** Builds a list of the shaders in a shader map. Key is FShaderType::TypeName */
-	ENGINE_API void GetShaderList(TMap<FHashedName, TShaderRef<FShader>>& OutShaders) const;
+	ENGINE_API virtual void GetShaderList(TMap<FHashedName, TShaderRef<FShader>>& OutShaders) const override;
 
 	/** Builds a list of the shader pipelines in a shader map. */
-	ENGINE_API void GetShaderPipelineList(TArray<FShaderPipelineRef>& OutShaderPipelines) const;
+	ENGINE_API virtual void GetShaderPipelineList(TArray<FShaderPipelineRef>& OutShaderPipelines) const override;
 
 
 	/** Number of Shaders in Shadermap */
diff --git a/Engine/Source/Runtime/RHI/Public/RHIShaderLibrary.h b/Engine/Source/Runtime/RHI/Public/RHIShaderLibrary.h
--- a/Engine/Source/Runtime/RHI/Public/RHIShaderLibrary.h
+++ b/Engine/Source/Runtime/RHI/Public/RHIShaderLibrary.h
@@ -24,6 +24,7 @@ public:
 	virtual int32 GetNumShaders() const = 0;
 	virtual int32 GetNumShadersForShaderMap(int32 ShaderMapIndex) const = 0;
 	virtual int32 GetShaderIndex(int32 ShaderMapIndex, int32 i) const = 0;
+	virtual uint32 GetSizeBytes() const = 0;
 	virtual FSHAHash GetShaderHash(int32 ShaderMapIndex, int32 ShaderIndex) = 0;
 	virtual int32 FindShaderMapIndex(const FSHAHash& Hash) = 0;
 	virtual int32 FindShaderIndex(const FSHAHash& Hash) = 0;
diff --git a/Engine/Source/Runtime/RHI/RHI.Build.cs b/Engine/Source/Runtime/RHI/RHI.Build.cs
--- a/Engine/Source/Runtime/RHI/RHI.Build.cs
+++ b/Engine/Source/Runtime/RHI/RHI.Build.cs
@@ -20,10 +20,11 @@ public class RHI : ModuleRules
 
 			if (Target.Type != TargetRules.TargetType.Server)   // Dedicated servers should skip loading everything but NullDrv
 			{
-				if (Target.Configuration != UnrealTargetConfiguration.Shipping && Target.Configuration != UnrealTargetConfiguration.Test)
-                {
+				// Always disable for Shipping, always enable for Development, let the user decide to use for Test and DebugGame
+				if (Target.Configuration != UnrealTargetConfiguration.Shipping && Target.bTrackRHIResourceInfo || Target.Configuration == UnrealTargetConfiguration.Development)
+				{
 					PublicDefinitions.Add("RHI_WANT_RESOURCE_INFO=1");
-                }
+				}
 
 				// UEBuildAndroid.cs adds VulkanRHI for Android builds if it is enabled
 				if (Target.Platform.IsInGroup(UnrealPlatformGroup.Windows))
diff --git a/Engine/Source/Runtime/RenderCore/Private/GlobalShader.cpp b/Engine/Source/Runtime/RenderCore/Private/GlobalShader.cpp
--- a/Engine/Source/Runtime/RenderCore/Private/GlobalShader.cpp
+++ b/Engine/Source/Runtime/RenderCore/Private/GlobalShader.cpp
@@ -563,6 +563,16 @@ FShaderPipelineRef FGlobalShaderMapSection::GetShaderPipeline(const FShaderPipel
 	return Pipeline ? FShaderPipelineRef(Pipeline, *this) : FShaderPipelineRef();
 }
 
+void FGlobalShaderMapSection::GetShaderList(TMap<FHashedName, TShaderRef<FShader>>& OutShaders) const
+{
+	GetContent()->GetShaderList(*this, OutShaders);
+}
+
+void FGlobalShaderMapSection::GetShaderPipelineList(TArray<FShaderPipelineRef>& OutShaderPipelines) const
+{
+	GetContent()->GetShaderPipelineList(*this, OutShaderPipelines, FShaderPipeline::EAll);
+}
+
 FGlobalShaderMap::FGlobalShaderMap(EShaderPlatform InPlatform)
 	: Platform(InPlatform)
 {
diff --git a/Engine/Source/Runtime/RenderCore/Private/ShaderCodeLibrary.cpp b/Engine/Source/Runtime/RenderCore/Private/ShaderCodeLibrary.cpp
--- a/Engine/Source/Runtime/RenderCore/Private/ShaderCodeLibrary.cpp
+++ b/Engine/Source/Runtime/RenderCore/Private/ShaderCodeLibrary.cpp
@@ -283,6 +283,7 @@ public:
 	virtual void ReleasePreloadedShaderCode(int32 ShaderIndex) override;
 	virtual bool TryRelease() override;
 	virtual uint32 GetSizeBytes() const override { return sizeof(*this) + GetAllocatedSize(); }
+	virtual FString GetFriendlyName() const override;
 
 	class FShaderLibraryInstance* LibraryInstance;
 	int32 ShaderMapIndex;
@@ -977,6 +978,20 @@ public:
 		return sizeof(*this) + ShaderBucketsSize + Resources.GetAllocatedSize();
 	}
 
+	uint32 GetShaderMapsSizeBytes()
+	{
+		uint32 ShaderMapsSize = 0;
+		FRWScopeLock Locker(ResourceLock, SLT_ReadOnly);
+		for (FShaderMapResource_SharedCode* Resource : Resources)
+		{
+			if (Resource)
+			{
+				ShaderMapsSize += Resource->GetSizeBytes();
+			}
+		}
+		return ShaderMapsSize;
+	}
+
 	const int32 GetNumShadersForShaderMap(int32 ShaderMapIndex) const
 	{
 		return Library->GetNumShadersForShaderMap(ShaderMapIndex);
@@ -1286,6 +1301,11 @@ bool FShaderMapResource_SharedCode::TryRelease()
 	return false;
 }
 
+FString FShaderMapResource_SharedCode::GetFriendlyName() const
+{
+	return LibraryInstance->Library->GetName();
+}
+
 #if WITH_EDITOR
 struct FShaderCodeStats
 {
@@ -3428,6 +3448,15 @@ public:
 		}
 	}
 
+	template<typename F>
+	void IterateNamedShaderLibrariesSafe(F&& Func)
+	{
+		FRWScopeLock NamedReadLock(NamedLibrariesMutex, SLT_ReadOnly);
+		for (auto& [LogicalName, NamedShaderLibrary] : NamedLibrariesStack)
+		{
+			Invoke(Forward<F>(Func), LogicalName, *NamedShaderLibrary);
+		}
+	}
 };
 
 static FSharedShaderCodeRequest OnSharedShaderCodeRequest;
@@ -4264,3 +4293,48 @@ void UE::ShaderLibrary::Private::FNamedShaderLibrary::DumpLibraryContents(const
 	}
 }
 #endif
+
+FAutoConsoleCommandWithArgsAndOutputDevice GListShaderLibrariesCmd(
+	TEXT("ListShaderLibraries"),
+	TEXT("Spits out a csv table containing stats of all shader libraries"),
+	FConsoleCommandWithArgsAndOutputDeviceDelegate::CreateStatic(
+		[](const TArray<FString>& Params, FOutputDevice& Out)
+		{
+			auto Iter = [&](const FString& LogicalName, UE::ShaderLibrary::Private::FNamedShaderLibrary& NamedShaderLibrary)
+			{
+				FRWScopeLock ComponentReadLock(NamedShaderLibrary.ComponentsMutex, SLT_ReadOnly);
+				for (const TUniquePtr<FShaderLibraryInstance>& ShaderLibrary : NamedShaderLibrary.Components)
+				{
+					const FString& Name = ShaderLibrary->Library->GetName();
+					const FString OwnerName = ShaderLibrary->Library->GetOwnerName().ToString();
+					uint32 Id = ShaderLibrary->Library->GetId();
+					int32 NumShaders = ShaderLibrary->GetNumShaders();
+					int32 NumShaderMaps = ShaderLibrary->GetNumResources();
+					uint32 LibrarySize = ShaderLibrary->GetSizeBytes();
+					uint32 RHILibrarySize = ShaderLibrary->Library->GetSizeBytes();
+					uint32 MapsSize = ShaderLibrary->GetShaderMapsSizeBytes();
+
+					Out.Logf(TEXT("%s,%s,%s,%x,%d,%d,%.3f,%.3f"),
+						*Name,
+						*LogicalName,
+						*OwnerName,
+						Id,
+						NumShaders,
+						NumShaderMaps,
+						LibrarySize / 1024.f,
+						RHILibrarySize / 1024.f,
+						MapsSize / 1024.f
+					);
+				}
+			};
+
+			if (FShaderLibrariesCollection* Collection = FShaderLibrariesCollection::Impl)
+			{
+				Out.Logf(TEXT("ShaderLibraryName,LogicalName,OwnerName,Id,NumShaders,NumShaderMaps,LibrarySizeKb,RHILibrarySizeKb,ShaderMapsSizeKb"));
+				Collection->IterateNamedShaderLibrariesSafe(Iter);
+			}
+			else
+			{
+				UE_LOG(LogShaderLibrary, Warning, TEXT("ShaderLibrariesCollection is not available."));
+			}
+		}));
diff --git a/Engine/Source/Runtime/RenderCore/Private/ShaderMap.cpp b/Engine/Source/Runtime/RenderCore/Private/ShaderMap.cpp
--- a/Engine/Source/Runtime/RenderCore/Private/ShaderMap.cpp
+++ b/Engine/Source/Runtime/RenderCore/Private/ShaderMap.cpp
@@ -13,6 +13,7 @@
 #include "ShaderCore.h"
 #include "Misc/ScopeLock.h"
 #include "UObject/RenderingObjectVersion.h"
+#include "HAL/LowLevelMemTracker.h"
 #include "DataDrivenShaderPlatformInfo.h"
 #include "ProfilingDebugging/CsvProfiler.h"
 
@@ -25,11 +26,88 @@ static EShaderPermutationFlags GetCurrentShaderPermutationFlags()
 	return Result;
 }
 
+#ifndef ALLOW_SHADERMAP_TRACKING
+#define ALLOW_SHADERMAP_TRACKING 0
+#endif
+
+#if ALLOW_SHADERMAP_TRACKING
+static TArray<FShaderMapBase*> GAllShaderMaps;
+static FCriticalSection GAllShaderMapsGuard;
+
+TAutoConsoleVariable<bool> CVarEnableShaderMapTracking(
+	TEXT("r.TrackShaderMaps"),
+	WITH_EDITOR,
+	TEXT("Enables the tracking of every shadermap instantiated. Required to run ListShaderMaps command."),
+	ECVF_ReadOnly);
+
+FAutoConsoleCommandWithArgsAndOutputDevice GListShaderMapsCmd(
+	TEXT("ListShaderMaps"),
+	TEXT("Spits out a csv table containing stats of all shadermaps"),
+	FConsoleCommandWithArgsAndOutputDeviceDelegate::CreateStatic(
+		[](const TArray<FString>& Params, FOutputDevice& Out)
+		{
+			if (!CVarEnableShaderMapTracking.GetValueOnGameThread())
+			{
+				UE_LOG(LogShaders, Warning, TEXT("Enable r.TrackShaderMaps in ini config to enable the functionality"));
+				return;
+			}
+			Out.Logf(TEXT("ShaderMapName,OwnerName,NumShaders,NumPipelines,SizeKb,bUsedForRendering"));
+			FScopeLock SMAccess(&GAllShaderMapsGuard);
+			for (FShaderMapBase* ShaderMap : GAllShaderMaps)
+			{
+				FString FriendlyName = TEXT("Unknown");
+				FString OwnerName = TEXT("Unknown");
+				uint32 CodeSize = ShaderMap->GetFrozenContentSize();
+				bool bUseForRendering = false;
+
+				if (const FShaderMapResource* Resource = ShaderMap->GetResource())
+				{
+					FriendlyName = Resource->GetFriendlyName();
+					OwnerName = Resource->GetOwnerName().ToString();
+					CodeSize += Resource->GetSizeBytes();
+					bUseForRendering = Resource->ContainsAtLeastOneRHIShaderCreated();
+				}
+
+				TMap<FHashedName, TShaderRef<FShader>> Shaders;
+				ShaderMap->GetShaderList(Shaders);
+				TArray<FShaderPipelineRef> Pipelines;
+				ShaderMap->GetShaderPipelineList(Pipelines);
+
+				// Editor doesn't have the size baked, so grab it from shaders themselves.
+				if (CodeSize == 0)
+				{
+					for (auto& [Hash, Shader] : Shaders)
+					{
+						CodeSize += Shader->GetCodeSize();
+					}
+				}
+
+				Out.Logf(TEXT("%s,%s,%d,%d,%.3f,%s"),
+					*FriendlyName,
+					*OwnerName,
+					Shaders.Num(),
+					Pipelines.Num(),
+					CodeSize / 1024.f,
+					bUseForRendering ? TEXT("YES") : TEXT("NO")
+				);
+			}
+		}));
+#endif // ALLOW_SHADERMAP_TRACKING
+
 FShaderMapBase::FShaderMapBase()
 	: PointerTable(nullptr)
 	, NumFrozenShaders(0u)
 {
 	PermutationFlags = GetCurrentShaderPermutationFlags();
+
+#if ALLOW_SHADERMAP_TRACKING
+	if (CVarEnableShaderMapTracking.GetValueOnAnyThread())
+	{
+		LLM_SCOPE_BYNAME(TEXT("Debug/ShaderMapsTracking"));
+		FScopeLock SMAccess(&GAllShaderMapsGuard);
+		GAllShaderMaps.Add(this);
+	}
+#endif // ALLOW_SHADERMAP_TRACKING
 }
 
 FShaderMapBase::~FShaderMapBase()
@@ -39,6 +117,14 @@ FShaderMapBase::~FShaderMapBase()
 	{
 		delete PointerTable;
 	}
+
+#if ALLOW_SHADERMAP_TRACKING
+	if (CVarEnableShaderMapTracking.GetValueOnAnyThread())
+	{
+		FScopeLock SMAccess(&GAllShaderMapsGuard);
+		GAllShaderMaps.RemoveSingleSwap(this, EAllowShrinking::No);
+	}
+#endif // ALLOW_SHADERMAP_TRACKING
 }
 
 FShaderMapResourceCode* FShaderMapBase::GetResourceCode()
diff --git a/Engine/Source/Runtime/RenderCore/Public/GlobalShader.h b/Engine/Source/Runtime/RenderCore/Public/GlobalShader.h
--- a/Engine/Source/Runtime/RenderCore/Public/GlobalShader.h
+++ b/Engine/Source/Runtime/RenderCore/Public/GlobalShader.h
@@ -181,6 +181,8 @@ public:
 	RENDERCORE_API static FGlobalShaderMapSection* CreateFromArchive(FArchive& Ar);
 
 	RENDERCORE_API bool Serialize(FArchive& Ar);
+	RENDERCORE_API virtual void GetShaderList(TMap<FHashedName, TShaderRef<FShader>>& OutShaders) const override;
+	RENDERCORE_API virtual void GetShaderPipelineList(TArray<FShaderPipelineRef>& OutShaderPipelines) const override;
 private:
 	inline FGlobalShaderMapSection() 
 	{ 
diff --git a/Engine/Source/Runtime/RenderCore/Public/Shader.h b/Engine/Source/Runtime/RenderCore/Public/Shader.h
--- a/Engine/Source/Runtime/RenderCore/Public/Shader.h
+++ b/Engine/Source/Runtime/RenderCore/Public/Shader.h
@@ -558,6 +558,7 @@ public:
 	// FShaderMapResource interface
 	RENDERCORE_API virtual FSHAHash GetShaderHash(int32 ShaderIndex) override;
 	RENDERCORE_API virtual FRHIShader* CreateRHIShaderOrCrash(int32 ShaderIndex, bool bRequired) override;
+	RENDERCORE_API virtual FString GetFriendlyName() const override { return TEXT("InlineCode"); }
 	virtual uint32 GetSizeBytes() const override;
 
 	TRefCountPtr<FShaderMapResourceCode> Code;
@@ -2370,6 +2371,9 @@ public:
 		UnfreezeContent();
 		return Content.Object;
 	}
+	/** Builds a list of the shaders in a shader map. Key is FShaderType::TypeName */
+	RENDERCORE_API virtual void GetShaderList(TMap<FHashedName, TShaderRef<FShader>>& OutShaders) const = 0;
+	RENDERCORE_API virtual void GetShaderPipelineList(TArray<FShaderPipelineRef>& OutShaderPipelines) const = 0;
 
 	inline EShaderPlatform GetShaderPlatform() const { return Content.Object ? Content.Object->GetShaderPlatform() : SP_NumPlatforms; }
 	inline uint32 GetFrozenContentSize() const { return Content.FrozenSize; }
diff --git a/Engine/Source/Runtime/RenderCore/Public/ShaderCodeArchive.h b/Engine/Source/Runtime/RenderCore/Public/ShaderCodeArchive.h
--- a/Engine/Source/Runtime/RenderCore/Public/ShaderCodeArchive.h
+++ b/Engine/Source/Runtime/RenderCore/Public/ShaderCodeArchive.h
@@ -320,7 +320,7 @@ public:
 
 	virtual bool IsNativeLibrary() const override { return false; }
 
-	int64 GetSizeBytes() const
+	virtual uint32 GetSizeBytes() const override
 	{
 		return sizeof(*this) +
 			SerializedShaders.GetAllocatedSize() +
@@ -580,7 +580,7 @@ public:
 
 	virtual bool IsNativeLibrary() const override { return false; }
 
-	uint64 GetSizeBytes() const
+	virtual uint32 GetSizeBytes() const override
 	{
 		return sizeof(*this) +
 			Header.GetAllocatedSize() +
diff --git a/Engine/Source/Runtime/RenderCore/RenderCore.Build.cs b/Engine/Source/Runtime/RenderCore/RenderCore.Build.cs
--- a/Engine/Source/Runtime/RenderCore/RenderCore.Build.cs
+++ b/Engine/Source/Runtime/RenderCore/RenderCore.Build.cs
@@ -53,5 +53,10 @@ public class RenderCore : ModuleRules
 			// For IOpenGLDynamicRHI::RHIGenerateMips
 			PublicIncludePathModuleNames.Add("OpenGLDrv");
 		}
+
+		if (Target.Configuration != UnrealTargetConfiguration.Shipping)
+		{
+			PrivateDefinitions.Add("ALLOW_SHADERMAP_TRACKING=1");
+		}
     }
 }
