- Allow other modules / plugins to use shared headers between code / shaders

diff --git a/Engine/Source/Runtime/RenderCore/Private/ShaderCore.cpp b/Engine/Source/Runtime/RenderCore/Private/ShaderCore.cpp
--- a/Engine/Source/Runtime/RenderCore/Private/ShaderCore.cpp
+++ b/Engine/Source/Runtime/RenderCore/Private/ShaderCore.cpp
@@ -355,6 +355,7 @@ FShaderHashCache GShaderHashCache;
 
 /** Global map of virtual file path to physical file paths */
 static TMap<FString, FString> GShaderSourceDirectoryMappings;
+static TArray<FString> GShaderSourceSharedVirtualDirectories = { TEXT("/Engine/Shared/") };
 
 static TAutoConsoleVariable<int32> CVarForceDebugViewModes(
 	TEXT("r.ForceDebugViewModes"),
@@ -1363,8 +1364,9 @@ bool CheckVirtualShaderFilePath(FStringView VirtualFilePath, TArray<FShaderCompi
 		bSuccess = false;
 	}
 
-	FStringView Extension = FPathViews::GetExtension(VirtualFilePath);
-	if (VirtualFilePath.StartsWith(TEXT("/Engine/Shared/")))
+	const FStringView Extension = FPathViews::GetExtension(VirtualFilePath);
+	const bool bIsSharedDirectory = GShaderSourceSharedVirtualDirectories.ContainsByPredicate([&VirtualFilePath](const FString& SharedDirectory) { return VirtualFilePath.StartsWith(SharedDirectory); });
+	if (bIsSharedDirectory)
 	{
 		if ((Extension != TEXTVIEW("h")))
 		{
@@ -3645,6 +3647,30 @@ void AddShaderSourceDirectoryMapping(const FString& VirtualShaderDirectory, cons
 	GShaderSourceDirectoryMappings.Add(VirtualShaderDirectory, RealShaderDirectory);
 }
 
+void AddShaderSourceSharedVirtualDirectory(const FString& VirtualShaderDirectory)
+{
+	check(IsInGameThread());
+	if (FPlatformProperties::RequiresCookedData() || !AllowShaderCompiling())
+	{
+		return;
+	}
+
+	// Do sanity checks of the virtual shader directory to map.
+	checkf(
+		VirtualShaderDirectory.StartsWith(TEXT("/")) &&
+		VirtualShaderDirectory.EndsWith(TEXT("/")) &&
+		!VirtualShaderDirectory.Contains(FString(TEXT("."))),
+		TEXT("Shared VirtualShaderDirectory = \"%s\" must start and end with '/' and contain no '.' characters."),
+		*VirtualShaderDirectory
+	);
+
+	// Detect collisions with any other mappings.
+	check(!GShaderSourceSharedVirtualDirectories.Contains(VirtualShaderDirectory));
+
+	// Add to the list of shared directories
+	GShaderSourceSharedVirtualDirectories.Add(VirtualShaderDirectory);
+}
+
 void FShaderCode::Compress(FName ShaderCompressionFormat, FOodleDataCompression::ECompressor InOodleCompressor, FOodleDataCompression::ECompressionLevel InOodleLevel)
 {
 	checkf(OptionalDataSize == -1, TEXT("FShaderCode::Compress() was called before calling FShaderCode::FinalizeShaderCode()"));
diff --git a/Engine/Source/Runtime/RenderCore/Public/ShaderCore.h b/Engine/Source/Runtime/RenderCore/Public/ShaderCore.h
--- a/Engine/Source/Runtime/RenderCore/Public/ShaderCore.h
+++ b/Engine/Source/Runtime/RenderCore/Public/ShaderCore.h
@@ -1360,6 +1360,11 @@ extern RENDERCORE_API void ResetAllShaderSourceDirectoryMappings();
  * @param RealShaderDirectory FPlatformProcess::BaseDir() relative path of the directory map.
  */
 extern RENDERCORE_API void AddShaderSourceDirectoryMapping(const FString& VirtualShaderDirectory, const FString& RealShaderDirectory);
+/**
+ * Specifies that the virtual shader directory and all subdirectories should contain only .h files that are shared between C++ / shader
+ * @param VirtualShaderDirectory Unique absolute path of the virtual shader directory (ex: /Project/Shared/).
+ */
+extern RENDERCORE_API void AddShaderSourceSharedVirtualDirectory(const FString& VirtualShaderDirectory);
 
 extern RENDERCORE_API void AddShaderSourceFileEntry(TArray<FString>& OutVirtualFilePaths, FString VirtualFilePath, EShaderPlatform ShaderPlatform, const FName* ShaderPlatformName = nullptr);
 extern RENDERCORE_API void GetAllVirtualShaderSourcePaths(TArray<FString>& OutVirtualFilePaths, EShaderPlatform ShaderPlatform, const FName* ShaderPlatformName = nullptr);
