Invalidate shaders when r.Mobile.Forward.RenderRectLightsAsSpotLights is toggled

diff --git a/Engine/Source/Runtime/RenderCore/Private/Shader.cpp b/Engine/Source/Runtime/RenderCore/Private/Shader.cpp
--- a/Engine/Source/Runtime/RenderCore/Private/Shader.cpp
+++ b/Engine/Source/Runtime/RenderCore/Private/Shader.cpp
@@ -1992,6 +1992,16 @@ void ShaderMapAppendKey(EShaderPlatform Platform, FShaderKeyGenerator& KeyGen)
 			KeyGen.AppendSeparator();
 			KeyGen.Append(TEXT("MobFDP"));
 		}
+
+		{
+			static const auto CVarRenderRectLightAsSpotLight = IConsoleManager::Get().FindTConsoleVariableDataInt(TEXT("r.Mobile.Forward.RenderRectLightsAsSpotLights"));
+			const bool bRenderAsSpotLight = (CVarRenderRectLightAsSpotLight && CVarRenderRectLightAsSpotLight->GetValueOnAnyThread() != 0);
+			if (bRenderAsSpotLight)
+			{
+				KeyGen.AppendSeparator();
+				KeyGen.Append(TEXT("R2S"));
+			}
+		}
 	}
 	else
 	{
