[UE] Fixed ListShaderLibraries didn't output shader maps size
#rnx
#tests cooked build

diff --git a/Engine/Source/Runtime/RenderCore/Private/ShaderCodeLibrary.cpp b/Engine/Source/Runtime/RenderCore/Private/ShaderCodeLibrary.cpp
--- a/Engine/Source/Runtime/RenderCore/Private/ShaderCodeLibrary.cpp
+++ b/Engine/Source/Runtime/RenderCore/Private/ShaderCodeLibrary.cpp
@@ -4314,7 +4314,7 @@ FAutoConsoleCommandWithArgsAndOutputDevice GListShaderLibrariesCmd(
 					uint32 RHILibrarySize = ShaderLibrary->Library->GetSizeBytes();
 					uint32 MapsSize = ShaderLibrary->GetShaderMapsSizeBytes();
 
-					Out.Logf(TEXT("%s,%s,%s,%x,%d,%d,%.3f,%.3f"),
+					Out.Logf(TEXT("%s,%s,%s,%x,%d,%d,%.3f,%.3f,%.3f"),
 						*Name,
 						*LogicalName,
 						*OwnerName,
