FDispatchShaderBundleCS and FDispatchShaderBundleWorkGraph need to always be "bindful" since they manage their own descriptor heaps.

diff --git a/Engine/Source/Runtime/RenderCore/Private/ShaderBundles.cpp b/Engine/Source/Runtime/RenderCore/Private/ShaderBundles.cpp
--- a/Engine/Source/Runtime/RenderCore/Private/ShaderBundles.cpp
+++ b/Engine/Source/Runtime/RenderCore/Private/ShaderBundles.cpp
@@ -21,11 +21,7 @@ bool FDispatchShaderBundleCS::ShouldCompilePermutation(const FGlobalShaderPermut
 
 void FDispatchShaderBundleCS::ModifyCompilationEnvironment(const FGlobalShaderPermutationParameters& Parameters, FShaderCompilerEnvironment& OutEnvironment)
 {
-	if (FDataDrivenShaderPlatformInfo::GetRequiresBindfulUtilityShaders(Parameters.Platform))
-	{
-		OutEnvironment.CompilerFlags.Add(CFLAG_ForceBindful);
-	}
-
+	OutEnvironment.CompilerFlags.Add(CFLAG_ForceBindful);
 	OutEnvironment.CompilerFlags.Add(CFLAG_RootConstants);
 
 	OutEnvironment.SetDefine(TEXT("THREADGROUP_SIZEX"), ThreadGroupSizeX);
@@ -46,10 +42,7 @@ bool FDispatchShaderBundleWorkGraph::ShouldCompilePermutation(const FGlobalShade
 
 void FDispatchShaderBundleWorkGraph::ModifyCompilationEnvironment(const FGlobalShaderPermutationParameters& Parameters, FShaderCompilerEnvironment& OutEnvironment)
 {
-	if (FDataDrivenShaderPlatformInfo::GetRequiresBindfulUtilityShaders(Parameters.Platform))
-	{
-		OutEnvironment.CompilerFlags.Add(CFLAG_ForceBindful);
-	}
+	OutEnvironment.CompilerFlags.Add(CFLAG_ForceBindful);
 
 	OutEnvironment.SetDefine(TEXT("THREADGROUP_SIZEX"), ThreadGroupSizeX);
 	OutEnvironment.SetDefine(TEXT("MAX_DISPATCHGRID_SIZEX"), (CVarShaderBundleMaxSize.GetValueOnAnyThread() + ThreadGroupSizeX - 1) / ThreadGroupSizeX);
