Fix validation code in SetRayTracingShaderParameters() that previously did not account for padding in FRHIShaderParameterResource

diff --git a/Engine/Source/Runtime/RHI/Public/RHIShaderParameters.h b/Engine/Source/Runtime/RHI/Public/RHIShaderParameters.h
--- a/Engine/Source/Runtime/RHI/Public/RHIShaderParameters.h
+++ b/Engine/Source/Runtime/RHI/Public/RHIShaderParameters.h
@@ -72,6 +72,13 @@ struct FRHIShaderParameterResource
 	{
 	}
 
+	bool operator == (const FRHIShaderParameterResource& Other) const
+	{
+		return Resource == Other.Resource
+			&& Index == Other.Index
+			&& Type == Other.Type;
+	}
+
 	FRHIResource* Resource = nullptr;
 	uint16        Index = 0;
 	EType         Type = EType::Texture;
diff --git a/Engine/Source/Runtime/RenderCore/Private/ShaderParameterStruct.cpp b/Engine/Source/Runtime/RenderCore/Private/ShaderParameterStruct.cpp
--- a/Engine/Source/Runtime/RenderCore/Private/ShaderParameterStruct.cpp
+++ b/Engine/Source/Runtime/RenderCore/Private/ShaderParameterStruct.cpp
@@ -1082,9 +1082,7 @@ void SetRayTracingShaderParameters(
 			check(RTBindingsWriter.BindlessParameters.Num() == ConvertedBindings.BindlessParameters.Num());
 			for (int32 i = 0; i < ConvertedBindings.BindlessParameters.Num(); ++i)
 			{
-				const FRHIShaderParameterResource& A = RTBindingsWriter.BindlessParameters[i];
-				const FRHIShaderParameterResource& B = ConvertedBindings.BindlessParameters[i];
-				check(FMemory::Memcmp(&A, &B, sizeof(A)) == 0);
+				check(RTBindingsWriter.BindlessParameters[i] == ConvertedBindings.BindlessParameters[i]);
 			}
 		}
 	}
