Fix global shader non-determinism due to un-initialized data in the memory images for shader maps due to the TMaps and Variants used by the generic stat system.

- Memory Image serialization is problematic due to padding in structures and uninitialized data.  This can be handled by marking up each class and member with macros, but this can be invasive for common types like TMaps and TVariants.
- To fix the issue I've completely removed the shader stats from the memory image serialization and this avoids the issue, and most of the code can stay the same.
- Generic stats aren't store on the FShader anymore, and the stats have been moved to an array on the FShaderEditorOnlyDataEntry data object owned by the FShaderMapResourceCode owned by the base shader map class.  This data is natively serializated we don't have the issues that the memory images do.
- This required fixup on the interfaces that query for the generic shader stats to get this data from a different location.
- Wrap all the code in WITH_EDITORONLY_DATA
- Turn the stats into a struct, FGenericShaderStat, and we just store an array of them.  We don't really need fast lookup.

diff --git a/Engine/Plugins/FX/Niagara/Source/NiagaraEditor/Private/ViewModels/NiagaraScriptStatsViewModel.cpp b/Engine/Plugins/FX/Niagara/Source/NiagaraEditor/Private/ViewModels/NiagaraScriptStatsViewModel.cpp
--- a/Engine/Plugins/FX/Niagara/Source/NiagaraEditor/Private/ViewModels/NiagaraScriptStatsViewModel.cpp
+++ b/Engine/Plugins/FX/Niagara/Source/NiagaraEditor/Private/ViewModels/NiagaraScriptStatsViewModel.cpp
@@ -210,15 +210,16 @@ namespace NiagaraScriptStatsLocal
 									{
 										ResultsString.Appendf(TEXT("- Num Instructions = %d\n"), Shader->GetNumInstructions());
 									}
-
+#if WITH_EDITORONLY_DATA
 									TStringBuilder<2048> ShaderStatsBuilder;
-									const FShader::FShaderStatisticMap ShaderStats = Shader->GetShaderStatistics();
+									const TArray<FGenericShaderStat> ShaderStats = ShaderScript->GetGameThreadShaderMap()->GetShaderStatistics(Shader->GetType(ShaderScript->GetGameThreadShaderMap()->GetPointerTable()));
 									for (const auto& ShaderStat : ShaderStats)
 									{
-										ShaderStatsBuilder << TEXT("- ") << ShaderStat.Key << TEXT(" = ");
+										ShaderStatsBuilder << TEXT("- ") << ShaderStat.StatName << TEXT(" = ");
 										Visit([&ShaderStatsBuilder](auto& StoredValue) { ShaderStatsBuilder << StoredValue << "\n"; }, ShaderStat.Value);
 									}
 									ResultsString.Append(ShaderStatsBuilder.ToString());
+#endif // WITH_EDITORONLY_DATA
 								}
 							}
 						}
diff --git a/Engine/Shaders/Public/ShaderVersion.ush b/Engine/Shaders/Public/ShaderVersion.ush
--- a/Engine/Shaders/Public/ShaderVersion.ush
+++ b/Engine/Shaders/Public/ShaderVersion.ush
@@ -3,4 +3,4 @@
 // Each time the console command is executed it generates a new GUID. As this file is included
 // in Platform.ush (which should be included in any shader) it allows to invalidate the shader DDC.
 // 
-#pragma message("UESHADERMETADATA_VERSION EE2C3AC2-33D9-412B-BCEC-245641AB9666")
+#pragma message("UESHADERMETADATA_VERSION 2b7d999b-f44f-496d-9bdf-50dcd21236bb")
diff --git a/Engine/Source/Editor/MaterialEditor/Private/MaterialStatsCommon.cpp b/Engine/Source/Editor/MaterialEditor/Private/MaterialStatsCommon.cpp
--- a/Engine/Source/Editor/MaterialEditor/Private/MaterialStatsCommon.cpp
+++ b/Engine/Source/Editor/MaterialEditor/Private/MaterialStatsCommon.cpp
@@ -527,20 +527,22 @@ void FMaterialStatsUtils::GetRepresentativeShaderTypesAndDescriptions(TMap<FName
 	}
 }
 
-static FString GetShaderString(const FShader::FShaderStatisticMap& Statistics)
+#if WITH_EDITORONLY_DATA
+static FString GetShaderString(const TArray<FGenericShaderStat>& Statistics)
 {
 	TStringBuilder<2048> StatisticsStrBuilder;
 	for (const auto& Stat : Statistics)
 	{
-		StatisticsStrBuilder << Stat.Key << ": ";
+		StatisticsStrBuilder << Stat.StatName << ": ";
 		Visit([&StatisticsStrBuilder](auto& StoredValue)
-		{
+ 		{
 			StatisticsStrBuilder << StoredValue << "\n";
 		}, Stat.Value);
 	}
 
 	return StatisticsStrBuilder.ToString();
 }
+#endif // WITH_EDITORONLY_DATA
 
 /**
 * Gets instruction counts that best represent the likely usage of this material based on shading model and other factors.
@@ -578,7 +580,7 @@ void FMaterialStatsUtils::GetRepresentativeInstructionCounts(TArray<FShaderInstr
 						Info.ShaderType = ShaderInfo.ShaderType;
 						Info.ShaderDescription = ShaderInfo.ShaderDescription;
 						Info.InstructionCount = NumInstructions;
-						Info.ShaderStatisticsString = GetShaderString(MaterialShaderMap->GetShaderStatisticsMapForShader(ShaderType));
+						Info.ShaderStatisticsString = GetShaderString(MaterialShaderMap->GetShaderStatistics(ShaderType));
 						if (Info.ShaderStatisticsString.Len() == 0)
 						{
 							Info.ShaderStatisticsString = TEXT("n/a");
@@ -620,7 +622,9 @@ void FMaterialStatsUtils::GetRepresentativeInstructionCounts(TArray<FShaderInstr
 									Info.ShaderType = ShaderInfo.ShaderType;
 									Info.ShaderDescription = ShaderInfo.ShaderDescription;
 									Info.InstructionCount = NumInstructions;
-									Info.ShaderStatisticsString = GetShaderString(MeshShaderMap->GetShaderStatisticsMapForShader(*MaterialShaderMap, ShaderType));
+#if WITH_EDITORONLY_DATA
+									Info.ShaderStatisticsString = GetShaderString(MeshShaderMap->GetShaderStatistics(*MaterialShaderMap, ShaderType));
+#endif // WITH_EDITORONLY_DATA
 									if (Info.ShaderStatisticsString.Len() == 0)
 									{
 										Info.ShaderStatisticsString = TEXT("n/a");
diff --git a/Engine/Source/Runtime/Engine/Public/MaterialShared.h b/Engine/Source/Runtime/Engine/Public/MaterialShared.h
--- a/Engine/Source/Runtime/Engine/Public/MaterialShared.h
+++ b/Engine/Source/Runtime/Engine/Public/MaterialShared.h
@@ -1456,7 +1456,6 @@ public:
 	uint32 GetMaxNumInstructionsForShader(FShaderType* ShaderType) const { return GetContent()->GetMaxNumInstructionsForShader(*this, ShaderType); }
 
 #if WITH_EDITOR
-	FShader::FShaderStatisticMap GetShaderStatisticsMapForShader(FShaderType* ShaderType) const { return GetContent()->GetShaderStatisticsMapForShader(*this, ShaderType); }
 
 	/** Submits compile jobs for this shadermap, returns number of jobs submitted. */
 	int32 SubmitCompileJobs(uint32 CompilingShaderMapId,
diff --git a/Engine/Source/Runtime/RenderCore/Private/Shader.cpp b/Engine/Source/Runtime/RenderCore/Private/Shader.cpp
--- a/Engine/Source/Runtime/RenderCore/Private/Shader.cpp
+++ b/Engine/Source/Runtime/RenderCore/Private/Shader.cpp
@@ -731,11 +731,6 @@ FShader::FShader(const CompiledShaderInitializerType& Initializer)
 		// Store off the VF source hash that this shader was compiled with
 		VFSourceHash = Initializer.VertexFactoryType->GetSourceHash(Initializer.Target.GetPlatform());
 	}
-
-	for (const TPair<FString, FShaderStatVariant>& Pair : Initializer.ShaderStatistics)
-	{
-		ShaderStatistics.Add(FMemoryImageName(FName(Pair.Key)), Pair.Value);
-	}
 #endif // WITH_EDITORONLY_DATA
 
 	BuildParameterMapInfo(Initializer.ParameterMap.GetParameterMap());
diff --git a/Engine/Source/Runtime/RenderCore/Private/ShaderCore.cpp b/Engine/Source/Runtime/RenderCore/Private/ShaderCore.cpp
--- a/Engine/Source/Runtime/RenderCore/Private/ShaderCore.cpp
+++ b/Engine/Source/Runtime/RenderCore/Private/ShaderCore.cpp
@@ -51,6 +51,54 @@
 #include "Windows/HideWindowsPlatformTypes.h"
 #endif
 
+bool operator==(const FShaderStatVariant LHS, const FShaderStatVariant RHS)
+{
+	if (LHS.IsType<bool>() && RHS.IsType<bool>())
+	{
+		return LHS.Get<bool>() == RHS.Get<bool>();
+	}
+	else if (LHS.IsType<float>() && RHS.IsType<float>())
+	{
+		return LHS.Get<float>() == RHS.Get<float>();
+	}
+	else if (LHS.IsType<int32>() && RHS.IsType<int32>())
+	{
+		return LHS.Get<int32>() == RHS.Get<int32>();
+	}
+	else if (LHS.IsType<uint32>() && RHS.IsType<uint32>())
+	{
+		return LHS.Get<uint32>() == RHS.Get<uint32>();
+	}
+
+	return false;
+}
+
+FArchive& operator<<(FArchive& Ar, FGenericShaderStat& Stat)
+{
+	if (Ar.IsSaving())
+	{
+		FString StatNameString = Stat.StatName.ToString();
+		Ar << StatNameString;
+	}
+	else if (Ar.IsLoading())
+	{
+		FString StatNameString;
+		Ar << StatNameString;
+		Stat.StatName = FName(*StatNameString);
+	}
+	else
+	{
+		Ar << Stat.StatName;
+	}
+
+	Ar << Stat.Value;
+	return Ar;
+}
+
+bool FGenericShaderStat::operator==(const FGenericShaderStat& RHS) const
+{
+	return (StatName == RHS.StatName) && (Value == RHS.Value);
+}
 
 static TAutoConsoleVariable<bool> CVarDumpDebugInfoForCacheHits(
 	TEXT("r.ShaderCompiler.DumpDebugInfoForCacheHits"),
diff --git a/Engine/Source/Runtime/RenderCore/Private/ShaderMap.cpp b/Engine/Source/Runtime/RenderCore/Private/ShaderMap.cpp
--- a/Engine/Source/Runtime/RenderCore/Private/ShaderMap.cpp
+++ b/Engine/Source/Runtime/RenderCore/Private/ShaderMap.cpp
@@ -373,6 +373,35 @@ bool FShaderMapBase::Serialize(FArchive& Ar, bool bInlineShaderResources, bool b
 	return (bool)Content.Object;
 }
 
+#if WITH_EDITORONLY_DATA
+TArray<FGenericShaderStat> FShaderMapBase::GetShaderStatistics(FShaderType* ShaderType) const
+{
+	TArray<FGenericShaderStat> ShaderStatistics;
+	FShader* Shader = GetContent()->GetShader(ShaderType);
+	if (Shader)
+	{
+		ShaderStatistics = GetShaderStatistics(*Shader);
+	}
+	return ShaderStatistics;
+}
+
+TArray<FGenericShaderStat> FShaderMapBase::GetShaderStatistics(FShader& Shader) const
+{
+	TArray<FGenericShaderStat> ShaderStatistics;
+
+	if (Code)
+	{
+		const int32 ShaderIndex = Code->FindShaderIndex(Shader.GetOutputHash());
+		if (Code->ShaderEditorOnlyDataEntries.IsValidIndex(ShaderIndex))
+		{
+			ShaderStatistics = Code->ShaderEditorOnlyDataEntries[ShaderIndex].ShaderStatistics;
+		}
+	}
+
+	return ShaderStatistics;
+}
+#endif // WITH_EDITORONLY_DATA
+
 FString FShaderMapBase::ToString() const
 {
 	TStringBuilder<32000> String;
@@ -770,20 +799,20 @@ uint32 FShaderMapContent::GetMaxNumInstructionsForShader(const FShaderMapBase& I
 	return MaxNumInstructions;
 }
 
-#if WITH_EDITOR
-const FShader::FShaderStatisticMap FShaderMapContent::GetShaderStatisticsMapForShader(const FShaderMapBase& InShaderMap, FShaderType* ShaderType) const
+#if WITH_EDITORONLY_DATA
+TArray<FGenericShaderStat> FShaderMapContent::GetShaderStatistics(const FShaderMapBase& InShaderMap, FShaderType* ShaderType) const
 {
-	FShader::FShaderStatisticMap Statistics;
+	TArray<FGenericShaderStat> ShaderStatistics;
 
 	FShader* Shader = GetShader(ShaderType);
 	if (Shader)
 	{
-		Statistics = Shader->GetShaderStatistics();
+		ShaderStatistics = InShaderMap.GetShaderStatistics(*Shader);
 	}
 
-	return Statistics;
+	return ShaderStatistics;
 }
-#endif // WITH_EDITOR
+#endif // WITH_EDITORONLY_DATA
 
 struct FSortedShaderEntry
 {
diff --git a/Engine/Source/Runtime/RenderCore/Private/ShaderResource.cpp b/Engine/Source/Runtime/RenderCore/Private/ShaderResource.cpp
--- a/Engine/Source/Runtime/RenderCore/Private/ShaderResource.cpp
+++ b/Engine/Source/Runtime/RenderCore/Private/ShaderResource.cpp
@@ -253,7 +253,7 @@ void FShaderMapResourceCode::AddShaderCompilerOutput(const FShaderCompilerOutput
 
 #if WITH_EDITORONLY_DATA
 		// Output.Errors contains warnings in the case any exist (no errors since if there were the job would have failed)
-		AddEditorOnlyData(Index, DebugName, Output.PlatformDebugData, Output.Errors);
+		AddEditorOnlyData(Index, DebugName, Output.PlatformDebugData, Output.Errors, Output.ShaderStatistics);
 #endif
 
 		FShaderEntry& Entry = ShaderEntries.InsertDefaulted_GetRef(Index);
@@ -324,16 +324,21 @@ void FShaderMapResourceCode::AddShaderCompilerOutput(const FShaderCompilerOutput
 		// We append the warnings for any additional jobs which resulted in the same bytecode for the sake of determinism in the
 		// results saved to DDC. 
 		AppendWarningsToEditorOnlyData(Index, DebugName, Output.Errors);
+		ValidateShaderStatisticsEditorOnlyData(Index, Output.ShaderStatistics);
 	}
 #endif
 }
 
 #if WITH_EDITORONLY_DATA
-void FShaderMapResourceCode::AddEditorOnlyData(int32 Index, const FString& DebugName, TConstArrayView<uint8> InPlatformDebugData, TConstArrayView<FShaderCompilerError> InCompilerWarnings)
+void FShaderMapResourceCode::AddEditorOnlyData(int32 Index, const FString& DebugName, TConstArrayView<uint8> InPlatformDebugData, TConstArrayView<FShaderCompilerError> InCompilerWarnings, const TArray<FGenericShaderStat>& ShaderStatistics)
 {
 	FShaderEditorOnlyDataEntry& Entry = ShaderEditorOnlyDataEntries.InsertDefaulted_GetRef(Index);
 	Entry.PlatformDebugData = InPlatformDebugData;
 
+	// This should be a newly created shader entry.
+	check(Entry.ShaderStatistics.Num() == 0);
+	Entry.ShaderStatistics = ShaderStatistics;
+
 	AppendWarningsToEditorOnlyData(Index, DebugName, InCompilerWarnings);
 }
 
@@ -352,6 +357,29 @@ void FShaderMapResourceCode::AppendWarningsToEditorOnlyData(int32 Index, const F
 	}
 }
 
+void FShaderMapResourceCode::ValidateShaderStatisticsEditorOnlyData(int32 Index, const TArray<FGenericShaderStat>& ShaderStatistics)
+{
+	check(ShaderEditorOnlyDataEntries.IsValidIndex(Index));
+	FShaderEditorOnlyDataEntry& Entry = ShaderEditorOnlyDataEntries[Index];
+
+	if (Entry.ShaderStatistics.Num() != ShaderStatistics.Num())
+	{
+		UE_LOG(LogShaders, Warning, TEXT("Non-determinism detected in shader statistics.  Multiple duplicate shaders have the same shader statistics."));
+		return;
+	}
+
+	for (int i = 0; i < ShaderStatistics.Num(); ++i)
+	{
+		const FGenericShaderStat& StatA = Entry.ShaderStatistics[i];
+		const FGenericShaderStat& StatB = ShaderStatistics[i];
+		if (!(StatA == StatB))
+		{
+			UE_LOG(LogShaders, Warning, TEXT("Non-determinism detected in shader statistics.  Multiple duplicate shaders have the same shader statistics."));
+			return;
+		}
+	}
+}
+
 void FShaderMapResourceCode::LogShaderCompilerWarnings()
 {
 	if (ShaderEditorOnlyDataEntries.Num() > 0 && GShaderCompilerEmitWarningsOnLoad != 0)
diff --git a/Engine/Source/Runtime/RenderCore/Public/Shader.h b/Engine/Source/Runtime/RenderCore/Public/Shader.h
--- a/Engine/Source/Runtime/RenderCore/Public/Shader.h
+++ b/Engine/Source/Runtime/RenderCore/Public/Shader.h
@@ -504,9 +504,12 @@ public:
 		 *  Does not contain errors since if there were any errors, this object wouldn't exist. */
 		TArray<FString> CompilerWarnings;
 
+		/** Generic, data-driven key/value pairs of statistics. */
+		TArray<FGenericShaderStat> ShaderStatistics;
+
 		friend FArchive& operator<<(FArchive& Ar, FShaderEditorOnlyDataEntry& Entry)
 		{
-			return Ar << Entry.PlatformDebugData << Entry.CompilerWarnings;
+			return Ar << Entry.PlatformDebugData << Entry.CompilerWarnings << Entry.ShaderStatistics;
 		}
 	};
 #endif // WITH_EDITORONLY_DATA
@@ -529,7 +532,7 @@ public:
 	int32 FindShaderIndex(const FSHAHash& InHash) const;
 
 #if WITH_EDITORONLY_DATA
-	void AddEditorOnlyData(int32 Index, const FString& DebugName, TConstArrayView<uint8> InPlatformDebugData, TConstArrayView<FShaderCompilerError> InCompilerWarnings);
+	void AddEditorOnlyData(int32 Index, const FString& DebugName, TConstArrayView<uint8> InPlatformDebugData, TConstArrayView<FShaderCompilerError> InCompilerWarnings, const TArray<FGenericShaderStat>& ShaderStatistics);
 	void AppendWarningsToEditorOnlyData(int32 Index, const FString& DebugName, TConstArrayView<FShaderCompilerError> InCompilerWarnings);
 	RENDERCORE_API void LogShaderCompilerWarnings();
 #endif
@@ -544,6 +547,11 @@ public:
 	// Optional array of editor-only data indexed in the same order as ShaderEntries (sorted by the shader hash)
 	// Empty in the cases where the editor-only data is not serialized.
 	TArray<FShaderEditorOnlyDataEntry> ShaderEditorOnlyDataEntries;
+
+private:
+	// Given an existing shader entry and a new set of shader statistics determine if the shader stats match.
+	// This is to prevent non-determinism issues.
+	void ValidateShaderStatisticsEditorOnlyData(int32 Index, const TArray<FGenericShaderStat>& ShaderStatistics);
 #endif // WITH_EDITORONLY_DATA
 };
 	
@@ -928,7 +936,12 @@ public:
 	inline uint32 GetNumTextureSamplers() const { return NumTextureSamplers; }
 	inline uint32 GetCodeSize() const { return CodeSize; }
 	inline void SetNumInstructions(uint32 Value) { NumInstructions = Value; }
-	inline const FShaderStatisticMap& GetShaderStatistics() const { return ShaderStatistics; }
+	UE_DEPRECATED(5.5, "GetShaderStatistics which returns a reference to the stat map is deprecated. Shader Stats can be accessed on the editor only data on the FShaderMapResourceCode object.")
+	inline const FShaderStatisticMap& GetShaderStatistics() const
+	{
+		static FShaderStatisticMap DummyStatMap;
+		return DummyStatMap;
+	}
 #else
 	inline uint32 GetNumTextureSamplers() const { return 0u; }
 	inline uint32 GetCodeSize() const { return 0u; }
@@ -1038,9 +1051,6 @@ private:
 
 	/** Size of shader's compiled code */
 	LAYOUT_FIELD_EDITORONLY(uint32, CodeSize);
-
-	/** Generic, data-driven key/value pairs of statistics. */
-	LAYOUT_FIELD_EDITORONLY(FShaderStatisticMap, ShaderStatistics);
 };
 
 RENDERCORE_API const FTypeLayoutDesc& GetTypeLayoutDesc(const FPointerTableBase* PtrTable, const FShader& Shader);
@@ -1593,7 +1603,7 @@ struct FShaderCompiledShaderInitializerType
 	uint32 NumTextureSamplers;
 	uint32 CodeSize;
 	int32 PermutationId;
-	TMap<FString, FShaderStatVariant> ShaderStatistics;
+	TArray<FGenericShaderStat> ShaderStatistics;
 
 	RENDERCORE_API FShaderCompiledShaderInitializerType(
 		const FShaderType* InType,
@@ -2300,9 +2310,18 @@ public:
 
 	RENDERCORE_API void SaveShaderStableKeys(const FShaderMapBase& InShaderMap, EShaderPlatform TargetShaderPlatform, const struct FStableShaderKeyAndValue& SaveKeyVal);
 
-	RENDERCORE_API const FShader::FShaderStatisticMap GetShaderStatisticsMapForShader(const FShaderMapBase& InShaderMap, FShaderType* ShaderType) const;
+	UE_DEPRECATED(5.5, "GetShaderStatistics which returns a reference to the stat map is deprecated. Please use the function that returns a non-ref.")
+	RENDERCORE_API const FShader::FShaderStatisticMap GetShaderStatisticsMapForShader(const FShaderMapBase& InShaderMap, FShaderType* ShaderType) const
+	{
+		static FShader::FShaderStatisticMap DummyStatMap;
+		return DummyStatMap;
+	}
 #endif // WITH_EDITOR
 
+#if WITH_EDITORONLY_DATA
+	RENDERCORE_API TArray<FGenericShaderStat> GetShaderStatistics(const FShaderMapBase& InShaderMap, FShaderType* ShaderType) const;
+#endif // WITH_EDITORONLY_DATA
+
 	/** @return true if the map is empty */
 	inline bool IsEmpty() const
 	{
@@ -2391,6 +2410,11 @@ public:
 
 	RENDERCORE_API FString ToString() const;
 
+#if WITH_EDITORONLY_DATA
+	RENDERCORE_API TArray<FGenericShaderStat> GetShaderStatistics(FShaderType* ShaderType) const;
+	RENDERCORE_API TArray<FGenericShaderStat> GetShaderStatistics(FShader& Shader) const;
+#endif
+
 #if WITH_EDITOR
 	inline void GetOutdatedTypes(TArray<const FShaderType*>& OutdatedShaderTypes, TArray<const FShaderPipelineType*>& OutdatedShaderPipelineTypes, TArray<const FVertexFactoryType*>& OutdatedFactoryTypes) const
 	{
diff --git a/Engine/Source/Runtime/RenderCore/Public/ShaderCompilerCore.h b/Engine/Source/Runtime/RenderCore/Public/ShaderCompilerCore.h
--- a/Engine/Source/Runtime/RenderCore/Public/ShaderCompilerCore.h
+++ b/Engine/Source/Runtime/RenderCore/Public/ShaderCompilerCore.h
@@ -33,7 +33,7 @@ typedef TSharedPtr<TArray<ANSICHAR>, ESPMode::ThreadSafe> FShaderSharedAnsiStrin
 // this is for the protocol, not the data, bump if FShaderCompilerInput/FShaderPreprocessOutput serialization, SerializeWorkerInput or ProcessInputFromArchive changes.
 inline const int32 ShaderCompileWorkerInputVersion = 28;
 // this is for the protocol, not the data, bump if FShaderCompilerOutput or WriteToOutputArchive changes.
-inline const int32 ShaderCompileWorkerOutputVersion = 20;
+inline const int32 ShaderCompileWorkerOutputVersion = 21;
 // this is for the protocol, not the data.
 inline const int32 ShaderCompileWorkerSingleJobHeader = 'S';
 // this is for the protocol, not the data.
@@ -597,7 +597,7 @@ struct FShaderCompilerOutput
 
 	TArray<uint8> PlatformDebugData;
 
-	TMap<FString, FShaderStatVariant> ShaderStatistics;
+	TArray<FGenericShaderStat> ShaderStatistics;
 
 	/** Generates OutputHash from the compiler output. */
 	RENDERCORE_API void GenerateOutputHash();
@@ -614,7 +614,9 @@ struct FShaderCompilerOutput
 	template<typename TValue>
 	void AddStatistic(const TCHAR* Name, TValue Value)
 	{
-		ShaderStatistics.Emplace(FString(Name), FShaderStatVariant(TInPlaceType<TValue>(), Value));
+		FGenericShaderStat& Stat = ShaderStatistics.AddZeroed_GetRef();
+		Stat.StatName = FName(Name);
+		Stat.Value = FShaderStatVariant(TInPlaceType<TValue>(), Value);
 	}
 
 	// Bump ShaderCompileWorkerOutputVersion if FShaderCompilerOutput changes
diff --git a/Engine/Source/Runtime/RenderCore/Public/ShaderCore.h b/Engine/Source/Runtime/RenderCore/Public/ShaderCore.h
--- a/Engine/Source/Runtime/RenderCore/Public/ShaderCore.h
+++ b/Engine/Source/Runtime/RenderCore/Public/ShaderCore.h
@@ -57,7 +57,20 @@ class FSHA1;
 class ITargetPlatform;
 
 using FShaderStatVariant = TVariant<bool, float, int32, uint32>;
-DECLARE_INTRINSIC_TYPE_LAYOUT(FShaderStatVariant);
+
+bool operator==(const FShaderStatVariant LHS, const FShaderStatVariant RHS);
+
+struct FGenericShaderStat
+{
+public:
+	FName StatName;
+	FShaderStatVariant Value;
+
+	FGenericShaderStat() = default;
+
+	friend FArchive& operator<<(FArchive& Ar, FGenericShaderStat& Stat);
+	bool operator==(const FGenericShaderStat& RHS) const;
+};
 
 CSV_DECLARE_CATEGORY_MODULE_EXTERN(RENDERCORE_API, Shaders);
 /**
