* Fix D3D ES31 not having device profile setup.
* Add Substrate variable for D3D ES31 profile

diff --git a/Engine/Config/BaseDeviceProfiles.ini b/Engine/Config/BaseDeviceProfiles.ini
--- a/Engine/Config/BaseDeviceProfiles.ini
+++ b/Engine/Config/BaseDeviceProfiles.ini
@@ -8,6 +8,7 @@
 +DeviceProfileNameAndTypes=Windows,Windows
 +DeviceProfileNameAndTypes=WindowsServer,Windows
 +DeviceProfileNameAndTypes=WindowsClient,Windows
++DeviceProfileNameAndTypes=Windows_Preview_ES31,Windows
 +DeviceProfileNameAndTypes=IOS,IOS
 +DeviceProfileNameAndTypes=iPadAir2,IOS
 +DeviceProfileNameAndTypes=iPadAir3,IOS
@@ -1240,4 +1241,15 @@ BaseProfileName=Android_Mid
 +CVars=xr.DisableOpenXROnAndroidWithoutOculus=0
 +CVars=r.Android.DisableVulkanSupport=0
 +CVars=r.Vulkan.SupportsBCTextureFormats=1
-+CVars=r.Vulkan.MaxBarriersPerBatch=128
\ No newline at end of file
++CVars=r.Vulkan.MaxBarriersPerBatch=128
+
+[Windows_Preview_ES31 DeviceProfile]
+DeviceType=Windows
+BaseProfileName=Windows
++CVars=r.Substrate.RoughDiffuse=0
++CVars=r.Substrate.ShadingQuality=2
++CVars=r.Substrate.TileCoord8bits=1
++CVars=r.Substrate.SheenQuality=2
++CVars=r.Substrate.Glints=0
++CVars=r.Substrate.SpecularProfile=0
++CVars=r.Substrate.ClosuresPerPixel=1
diff --git a/Engine/Config/Windows/DataDrivenPlatformInfo.ini b/Engine/Config/Windows/DataDrivenPlatformInfo.ini
--- a/Engine/Config/Windows/DataDrivenPlatformInfo.ini
+++ b/Engine/Config/Windows/DataDrivenPlatformInfo.ini
@@ -173,6 +173,8 @@ PlatformName=PCD3D_ES3_1
 ShaderFormat=PCD3D_ES31
 ShaderPlatform=PCD3D_ES3_1
 MenuTooltip=LOCTEXT("PreviewMenuTooltip_PCD3D_Mobile", "Windows using Mobile profile.")
+DeviceProfileName=Windows_Preview_ES31
+PreviewDeviceProfileName=Windows_Preview_ES31
 
 [ShaderPlatform OPENGL_PCES3_1]
 Language=OpenGL
diff --git a/Engine/Source/Runtime/RenderCore/Private/RenderUtils.cpp b/Engine/Source/Runtime/RenderCore/Private/RenderUtils.cpp
--- a/Engine/Source/Runtime/RenderCore/Private/RenderUtils.cpp
+++ b/Engine/Source/Runtime/RenderCore/Private/RenderUtils.cpp
@@ -1777,13 +1777,13 @@ static TAutoConsoleVariable<int32> CVarSubstrateBytesPerPixel(
 	TEXT("r.Substrate.BytesPerPixel"),
 	80,
 	TEXT("Substrate allocated byte per pixel to store materials data. Higher value means more complex material can be represented."),
-	ECVF_ReadOnly | ECVF_RenderThreadSafe);
+	ECVF_ReadOnly | ECVF_RenderThreadSafe | ECVF_Preview);
 
 static TAutoConsoleVariable<int32> CVarSubstrateClosuresPerPixel(
 	TEXT("r.Substrate.ClosuresPerPixel"),
 	8,	// Similar to SUBSTRATE_MAX_CLOSURE_COUNT
 	TEXT("Substrate closure count per pixel can be constrained. That is usefull to reduce the number of byte written durting the base pass, but also to limit the number of closures in Forward. Higher value means more complex material can be represented."),
-	ECVF_ReadOnly | ECVF_RenderThreadSafe);
+	ECVF_ReadOnly | ECVF_RenderThreadSafe | ECVF_Preview);
 
 static TAutoConsoleVariable<int32> CVarSubstrateBackCompatibility(
 	TEXT("r.SubstrateBackCompatibility"),
@@ -1801,13 +1801,13 @@ static TAutoConsoleVariable<int32> CVarSubstrateSheenQuality(
 	TEXT("r.Substrate.SheenQuality"),
 	1,
 	TEXT("Define Substrate sheen quality (1: Disney-based sheen, 2: Charlie-based sheen, ). r.Substrate.ShadingQuality=2 forces SheenQuality to 2. This variable is read-only."),
-	ECVF_ReadOnly | ECVF_RenderThreadSafe);
+	ECVF_ReadOnly | ECVF_RenderThreadSafe | ECVF_Preview);
 
 static TAutoConsoleVariable<int32> CVarSubstrateShadingQuality(
 	TEXT("r.Substrate.ShadingQuality"),
 	1,
 	TEXT("Define Substrate shading quality (1: accurate lighting, 2: approximate lighting). This variable is read-only."),
-	ECVF_ReadOnly | ECVF_RenderThreadSafe);
+	ECVF_ReadOnly | ECVF_RenderThreadSafe | ECVF_Preview);
 
 static TAutoConsoleVariable<int32> CVarSubstrateDBufferPass(
 	TEXT("r.Substrate.DBufferPass"),
@@ -1820,7 +1820,7 @@ static TAutoConsoleVariable<int32> CVarMaterialRoughDiffuse(
 	TEXT("r.Material.RoughDiffuse"),
 	0,
 	TEXT("Enable rough diffuse material."),
-	ECVF_ReadOnly | ECVF_RenderThreadSafe);
+	ECVF_ReadOnly | ECVF_RenderThreadSafe | ECVF_Preview);
 
 static TAutoConsoleVariable<int32> CVarSubstrateRoughDiffuse(
 	TEXT("r.Substrate.RoughDiffuse"),
@@ -1832,7 +1832,7 @@ static TAutoConsoleVariable<int32> CVarSubstrateGlints(
 	TEXT("r.Substrate.Glints"),
 	1,
 	TEXT("Enable Glint support for Substrate slabs. If changed, shaders needs to be recompiled."),
-	ECVF_ReadOnly | ECVF_RenderThreadSafe);
+	ECVF_ReadOnly | ECVF_RenderThreadSafe | ECVF_Preview);
 
 static TAutoConsoleVariable<int32> CVarSubstrateGlintsLUT(
 	TEXT("r.Substrate.Glints.LUT"),
@@ -1856,7 +1856,7 @@ static TAutoConsoleVariable<int32> CVarSubstrateSpecularProfile(
 	TEXT("r.Substrate.SpecularProfile"),
 	1,
 	TEXT("Enable Specular Profile support for Substrate slabs. If changed, shaders needs to be recompiled."),
-	ECVF_ReadOnly | ECVF_RenderThreadSafe);
+	ECVF_ReadOnly | ECVF_RenderThreadSafe | ECVF_Preview | ECVF_MobileShaderChange);
 
 static TAutoConsoleVariable<int32> CVarSubstrateDebugAdvancedVisualizationShaders(
 	TEXT("r.Substrate.Debug.AdvancedVisualizationShaders"),
