Refactored management of evicted ray tracing geometry.

Context:
- Evicted FRayTracingGeometry maintained the StreamedIn flag, which was later used in MakeResident(...) to recreate the geometry as-if it was streamed-in.
- While this was valid on PC since ray tracing geometry for streamed-in meshes is simply built at runtime, on platforms that support offline built BLAS this code path didn't handle loading the offline data (instead the BLAS ended up being built at runtime like on PC).

Changes:
- Refactored some branches out of FRayTracingGeometry::InitRHIForStreaming(...)/ReleaseRHIForStreaming(...) and moved them higher level where more context is available to avoid unexpected state simply being ignored.
- Evicting a streamed-in ray tracing geometry clears the StreamedIn and Valid flags, and sets Initializer.Type back to StreamingDestination.
- MakeResident() will simply create the RHI object without special handling for streamed-in ray tracing geometries.
- High level logic that calls FRayTracingGeometry::MakeResident(...) (currently FRayTracingGeometryManager) is responsible for handling the streaming-in of FRayTracingGeometry that requires it (ie: building BLAS at runtime, or streaming offline data as necessary).
- When ray tracing is toggled on, make all registered geometries resident and kick-off streaming requests as necessary.

diff --git a/Engine/Source/Runtime/Engine/Internal/Streaming/SkeletalMeshUpdate.cpp b/Engine/Source/Runtime/Engine/Internal/Streaming/SkeletalMeshUpdate.cpp
--- a/Engine/Source/Runtime/Engine/Internal/Streaming/SkeletalMeshUpdate.cpp
+++ b/Engine/Source/Runtime/Engine/Internal/Streaming/SkeletalMeshUpdate.cpp
@@ -111,10 +111,12 @@ void FSkeletalMeshStreamIn::FIntermediateRayTracingGeometry::SafeRelease()
 
 void FSkeletalMeshStreamIn::FIntermediateRayTracingGeometry::TransferRayTracingGeometry(FRayTracingGeometry& RayTracingGeometry, FRHIResourceReplaceBatcher& Batcher)
 {
-	RayTracingGeometry.InitRHIForStreaming(RayTracingGeometryRHI, Batcher);
-	RayTracingGeometry.SetRequiresBuild(bRequiresBuild);
-
-	SafeRelease();
+	if (ensureMsgf(RayTracingGeometryRHI.IsValid(),
+		TEXT("FIntermediateRayTracingGeometry should have a valid RHI object. Was r.RayTracing.Enable toggled between FStaticMeshStreamIn::CreateBuffers(...) and FStaticMeshStreamIn::DoFinishUpdate(...)?")))
+	{
+		RayTracingGeometry.InitRHIForStreaming(RayTracingGeometryRHI, Batcher);
+		RayTracingGeometry.SetRequiresBuild(bRequiresBuild);
+	}
 }
 
 #endif
@@ -209,10 +211,12 @@ void FSkeletalMeshStreamIn::DoFinishUpdate(const FContext& Context)
 			{
 				FSkeletalMeshLODRenderData& LODResource = *Context.LODResourcesView[LODIdx];
 
-				if (LODResource.GetNumVertices() > 0 && LODResource.bReferencedByStaticSkeletalMeshObjects_RenderThread)
+				if (IsRayTracingEnabled() && LODResource.GetNumVertices() > 0 && LODResource.bReferencedByStaticSkeletalMeshObjects_RenderThread && !LODResource.StaticRayTracingGeometry.IsEvicted())
 				{
 					IntermediateRayTracingGeometry[LODIdx].TransferRayTracingGeometry(LODResource.StaticRayTracingGeometry, Batcher);
 				}
+
+				IntermediateRayTracingGeometry[LODIdx].SafeRelease();
 			}
 		}
 
@@ -382,12 +386,9 @@ void FSkeletalMeshStreamOut::ReleaseBuffers(const FContext& Context)
 				}
 
 #if RHI_RAYTRACING
-				if (IsRayTracingAllowed())
+				if (IsRayTracingAllowed() && RenderData->LODRenderData[LODIndex].bReferencedByStaticSkeletalMeshObjects_RenderThread && !LODResource.StaticRayTracingGeometry.IsEvicted())
 				{
-					if (RenderData->LODRenderData[LODIndex].bReferencedByStaticSkeletalMeshObjects_RenderThread)
-					{
-						LODResource.StaticRayTracingGeometry.ReleaseRHIForStreaming(Batcher);
-					}
+					LODResource.StaticRayTracingGeometry.ReleaseRHIForStreaming(Batcher);
 				}
 #endif
 			}
diff --git a/Engine/Source/Runtime/Engine/Private/Rendering/RayTracingGeometryManager.cpp b/Engine/Source/Runtime/Engine/Private/Rendering/RayTracingGeometryManager.cpp
--- a/Engine/Source/Runtime/Engine/Private/Rendering/RayTracingGeometryManager.cpp
+++ b/Engine/Source/Runtime/Engine/Private/Rendering/RayTracingGeometryManager.cpp
@@ -286,7 +286,7 @@ void FRayTracingGeometryManager::SetRayTracingGeometryGroupCurrentFirstLODIndex(
 		{
 			// some LODs might be stripped during cook
 			// skeletal meshes only create static LOD when rendering as static
-			if (Group.Geometries[LODIdx] != nullptr)
+			if (Group.Geometries[LODIdx] && !Group.Geometries[LODIdx]->IsEvicted())
 			{
 				Group.Geometries[LODIdx]->ReleaseRHIForStreaming(Batcher);
 			}
@@ -382,6 +382,8 @@ void FRayTracingGeometryManager::Tick(FRHICommandList& RHICmdList)
 					RegisteredGeometry.Geometry->Evict();
 				}
 			}
+
+			PendingStreamingRequests.Empty();
 		}
 		else
 		{
@@ -399,6 +401,8 @@ void FRayTracingGeometryManager::Tick(FRHICommandList& RHICmdList)
 			TEXT("There's likely some issue tracking resident geometries or not all geometries have been evicted."),
 			TotalResidentSize
 		);
+
+		check(PendingStreamingRequests.IsEmpty());
 	}
 	else if (IsRayTracingUsingReferenceBasedResidency())
 	{
@@ -496,6 +500,11 @@ void FRayTracingGeometryManager::Tick(FRHICommandList& RHICmdList)
 				{
 					RegisteredGeometry.Geometry->MakeResident(RHICmdList);
 				}
+
+				if (!RequestRayTracingGeometryStreamIn(RHICmdList, RegisteredGeometry.Geometry))
+				{
+					PendingStreamingRequests.Add(RegisteredGeometry.Geometry);
+				}
 			}
 		}
 		else
@@ -507,6 +516,18 @@ void FRayTracingGeometryManager::Tick(FRHICommandList& RHICmdList)
 				checkf(!RegisteredGeometry.Geometry->IsEvicted(), TEXT("Ray tracing geometry should not be evicted when ray tracing is enabled."));
 			}
 #endif
+			
+			TArray<FRayTracingGeometry*> CurrentPendingStreamingRequests;
+			Swap(CurrentPendingStreamingRequests, PendingStreamingRequests);
+			PendingStreamingRequests.Reserve(CurrentPendingStreamingRequests.Num());
+
+			for (FRayTracingGeometry* Geometry : CurrentPendingStreamingRequests)
+			{
+				if (!RequestRayTracingGeometryStreamIn(RHICmdList, Geometry))
+				{
+					PendingStreamingRequests.Add(Geometry);
+				}
+			}
 		}
 	}
 
@@ -518,7 +539,7 @@ void FRayTracingGeometryManager::Tick(FRHICommandList& RHICmdList)
 	SET_MEMORY_STAT(STAT_RayTracingGeometryResidentMemory, TotalResidentSize);
 }
 
-void FRayTracingGeometryManager::RequestRayTracingGeometryStreamIn(FRHICommandList& RHICmdList, FRayTracingGeometry* Geometry)
+bool FRayTracingGeometryManager::RequestRayTracingGeometryStreamIn(FRHICommandList& RHICmdList, FRayTracingGeometry* Geometry)
 {
 	FRegisteredGeometry& RegisteredGeometry = RegisteredGeometries[Geometry->RayTracingGeometryHandle];
 
@@ -526,7 +547,18 @@ void FRayTracingGeometryManager::RequestRayTracingGeometryStreamIn(FRHICommandLi
 		|| RegisteredGeometry.Status == FRegisteredGeometry::FStatus::Streaming)
 	{
 		// no streaming required or streaming request already in-flight
-		return;
+		return true;
+	}
+
+	if (Geometry->GroupHandle != INDEX_NONE)
+	{
+		const FRayTracingGeometryGroup& Group = RegisteredGroups[Geometry->GroupHandle];
+
+		if (Geometry->LODIndex < Group.CurrentFirstLODIdx)
+		{
+			// streaming request no longer necessary
+			return true;
+		}
 	}
 
 	FByteBulkData* StreamableData = RegisteredGeometry.StreamableData;
@@ -570,7 +602,7 @@ void FRayTracingGeometryManager::RequestRayTracingGeometryStreamIn(FRHICommandLi
 
 		if (NumStreamingRequests >= GRayTracingStreamingMaxPendingRequests)
 		{
-			return;
+			return false;
 		}
 
 		FStreamingRequest& StreamingRequest = StreamingRequests[NextStreamingRequestIndex];
@@ -609,6 +641,8 @@ void FRayTracingGeometryManager::RequestRayTracingGeometryStreamIn(FRHICommandLi
 
 		Geometry->RequestBuildIfNeeded(RHICmdList, ERTAccelerationStructureBuildPriority::Normal);
 	}
+
+	return true;
 }
 
 void FRayTracingGeometryManager::ProcessCompletedStreamingRequests(FRHICommandList& RHICmdList)
diff --git a/Engine/Source/Runtime/Engine/Private/Streaming/StaticMeshUpdate.cpp b/Engine/Source/Runtime/Engine/Private/Streaming/StaticMeshUpdate.cpp
--- a/Engine/Source/Runtime/Engine/Private/Streaming/StaticMeshUpdate.cpp
+++ b/Engine/Source/Runtime/Engine/Private/Streaming/StaticMeshUpdate.cpp
@@ -124,10 +124,12 @@ void FStaticMeshStreamIn::FIntermediateRayTracingGeometry::SafeRelease()
 
 void FStaticMeshStreamIn::FIntermediateRayTracingGeometry::TransferRayTracingGeometry(FRayTracingGeometry& RayTracingGeometry, FRHIResourceReplaceBatcher& Batcher)
 {
-	RayTracingGeometry.InitRHIForStreaming(RayTracingGeometryRHI, Batcher);
-	RayTracingGeometry.SetRequiresBuild(bRequiresBuild);
-
-	SafeRelease();
+	if (ensureMsgf(RayTracingGeometryRHI.IsValid(),
+		TEXT("FIntermediateRayTracingGeometry should have a valid RHI object. Was r.RayTracing.Enable toggled between FStaticMeshStreamIn::CreateBuffers(...) and FStaticMeshStreamIn::DoFinishUpdate(...)?")))
+	{
+		RayTracingGeometry.InitRHIForStreaming(RayTracingGeometryRHI, Batcher);
+		RayTracingGeometry.SetRequiresBuild(bRequiresBuild);
+	}
 }
 
 #if DO_CHECK
@@ -176,7 +178,7 @@ void FStaticMeshStreamIn::CreateBuffers(const FContext& Context)
 			IntermediateBuffersArray[LODIdx].CreateFromCPUData(*StreamingRHICmdList, LODResource);
 
 #if RHI_RAYTRACING
-			if (IsRayTracingEnabled() && LODResource.RayTracingGeometry != nullptr && LODResource.VertexBuffers.StaticMeshVertexBuffer.GetNumVertices() > 0)
+			if (IsRayTracingEnabled() && LODResource.RayTracingGeometry != nullptr && LODResource.GetNumVertices() > 0)
 			{
 #if DO_CHECK
 				CheckRayTracingGeometryInitializer(LODResource, ERayTracingGeometryInitializerType::StreamingDestination, LODResource.RayTracingGeometry->Initializer);
@@ -234,10 +236,12 @@ void FStaticMeshStreamIn::DoFinishUpdate(const FContext& Context)
 			for (int32 LODIdx = PendingFirstLODIdx; LODIdx < CurrentFirstLODIdx; ++LODIdx)
 			{
 				FStaticMeshLODResources& LODResource = *Context.LODResourcesView[LODIdx];
-				if (LODResource.RayTracingGeometry != nullptr && LODResource.VertexBuffers.StaticMeshVertexBuffer.GetNumVertices() > 0)
+				if (IsRayTracingEnabled() && LODResource.GetNumVertices() > 0 && LODResource.RayTracingGeometry != nullptr && !LODResource.RayTracingGeometry->IsEvicted())
 				{
 					IntermediateRayTracingGeometry[LODIdx].TransferRayTracingGeometry(*LODResource.RayTracingGeometry, Batcher);
 				}
+
+				IntermediateRayTracingGeometry[LODIdx].SafeRelease();
 			}
 		}
 
@@ -246,8 +250,8 @@ void FStaticMeshStreamIn::DoFinishUpdate(const FContext& Context)
 		{
 			FStaticMeshLODResources& LODResource = *Context.LODResourcesView[LODIndex];
 
-			// Skip LODs that have their render data stripped
-			if (LODResource.RayTracingGeometry != nullptr && LODResource.VertexBuffers.StaticMeshVertexBuffer.GetNumVertices() > 0)
+			// Skip LODs that have their render data stripped or are evicted
+			if (LODResource.GetNumVertices() > 0 && LODResource.RayTracingGeometry != nullptr && !LODResource.RayTracingGeometry->IsEvicted())
 			{
 #if DO_CHECK
 				// Streaming LODs in/out shouldn't affect the ray tracing geometry initializer
@@ -258,7 +262,7 @@ void FStaticMeshStreamIn::DoFinishUpdate(const FContext& Context)
 #endif
 
 				// Under very rare circumstances that we switch ray tracing on/off right in the middle of streaming RayTracingGeometryRHI might not be valid.
-				if (IsRayTracingEnabled() && ensure(LODResource.RayTracingGeometry->IsValid() && !LODResource.RayTracingGeometry->IsEvicted()))
+				if (IsRayTracingEnabled() && ensure(LODResource.RayTracingGeometry->IsValid()))
 				{
 					LODResource.RayTracingGeometry->RequestBuildIfNeeded(FRHICommandListImmediate::Get(), ERTAccelerationStructureBuildPriority::Normal);
 				}
@@ -373,7 +377,7 @@ void FStaticMeshStreamOut::ReleaseRHIBuffers(const FContext& Context)
 			LODResource.ReleaseRHIForStreaming(Batcher);
 			
 #if RHI_RAYTRACING
-			if (LODResource.RayTracingGeometry != nullptr)
+			if (LODResource.RayTracingGeometry != nullptr && !LODResource.RayTracingGeometry->IsEvicted())
 			{
 				check(IsRayTracingAllowed());
 				LODResource.RayTracingGeometry->ReleaseRHIForStreaming(Batcher);
diff --git a/Engine/Source/Runtime/Engine/Public/Rendering/RayTracingGeometryManager.h b/Engine/Source/Runtime/Engine/Public/Rendering/RayTracingGeometryManager.h
--- a/Engine/Source/Runtime/Engine/Public/Rendering/RayTracingGeometryManager.h
+++ b/Engine/Source/Runtime/Engine/Public/Rendering/RayTracingGeometryManager.h
@@ -71,7 +71,7 @@ private:
 
 	void ReleaseRayTracingGeometryGroupReference(RayTracing::GeometryGroupHandle Handle);
 
-	void RequestRayTracingGeometryStreamIn(FRHICommandList& RHICmdList, FRayTracingGeometry* Geometry);
+	bool RequestRayTracingGeometryStreamIn(FRHICommandList& RHICmdList, FRayTracingGeometry* Geometry);
 	void ProcessCompletedStreamingRequests(FRHICommandList& RHICmdList);
 
 	FCriticalSection RequestCS;
@@ -140,6 +140,8 @@ private:
 	TSet<RayTracingGeometryHandle> ReferencedGeometryHandles;
 	TSet<RayTracing::GeometryGroupHandle> ReferencedGeometryGroups;
 
+	TArray<FRayTracingGeometry*> PendingStreamingRequests;
+
 	TArray<FStreamingRequest> StreamingRequests;
 	int32 NumStreamingRequests = 0;
 	int32 NextStreamingRequestIndex = 0;
diff --git a/Engine/Source/Runtime/RenderCore/Private/RayTracingGeometry.cpp b/Engine/Source/Runtime/RenderCore/Private/RayTracingGeometry.cpp
--- a/Engine/Source/Runtime/RenderCore/Private/RayTracingGeometry.cpp
+++ b/Engine/Source/Runtime/RenderCore/Private/RayTracingGeometry.cpp
@@ -26,41 +26,35 @@ PRAGMA_DISABLE_DEPRECATION_WARNINGS // Remove this in UE 5.6 once RayTracingGeom
 
 void FRayTracingGeometry::InitRHIForStreaming(FRHIRayTracingGeometry* IntermediateGeometry, FRHIResourceReplaceBatcher& Batcher)
 {
-	ensureMsgf(RayTracingGeometryRHI || !IsRayTracingEnabled(),
-		TEXT("RayTracingGeometryRHI should be valid when ray tracing is enabled.\n")
-		TEXT("This check failing points to a race condition between FRayTracingGeometryManager::Tick(...) and FStaticMeshStreamIn processing.\n")
+	checkf(IntermediateGeometry,
+		TEXT("IntermediateGeometry should be valid when streaming-in ray tracing geometry.\n")
+		TEXT("This will result in FRayTracingGeometry not being correctly initialized.\n")
 	);
 
+	checkf(RayTracingGeometryRHI, TEXT("RayTracingGeometryRHI must be valid when InitRHIForStreaming is called.\n"));
+
 	Initializer.Type = ERayTracingGeometryInitializerType::Rendering;
 
-	EnumAddFlags(GeometryState, EGeometryStateFlags::StreamedIn);
+	Batcher.EnqueueReplace(RayTracingGeometryRHI, IntermediateGeometry);
 
-	if (RayTracingGeometryRHI && IntermediateGeometry)
-	{
-		Batcher.EnqueueReplace(RayTracingGeometryRHI, IntermediateGeometry);
-		EnumAddFlags(GeometryState, EGeometryStateFlags::Valid);
+	EnumAddFlags(GeometryState, EGeometryStateFlags::Valid);
+	EnumAddFlags(GeometryState, EGeometryStateFlags::StreamedIn);
 
-		GRayTracingGeometryManager->RefreshRegisteredGeometry(RayTracingGeometryHandle);
-	}
-	else
-	{
-		check(GetRayTracingMode() == ERayTracingMode::Dynamic);
-	}
+	GRayTracingGeometryManager->RefreshRegisteredGeometry(RayTracingGeometryHandle);
 }
 
 void FRayTracingGeometry::ReleaseRHIForStreaming(FRHIResourceReplaceBatcher& Batcher)
 {
 	RemoveBuildRequest();
 
+	checkf(RayTracingGeometryRHI, TEXT("RayTracingGeometryRHI must be valid when ReleaseRHIForStreaming is called.\n"));
+
 	EnumRemoveFlags(GeometryState, EGeometryStateFlags::StreamedIn);
 	EnumRemoveFlags(GeometryState, EGeometryStateFlags::Valid);
 
-	Initializer.Type = ERayTracingGeometryInitializerType::StreamingDestination;
+	Batcher.EnqueueReplace(RayTracingGeometryRHI, nullptr);
 
-	if (RayTracingGeometryRHI)
-	{
-		Batcher.EnqueueReplace(RayTracingGeometryRHI, nullptr);
-	}
+	Initializer.Type = ERayTracingGeometryInitializerType::StreamingDestination;
 
 	GRayTracingGeometryManager->RefreshRegisteredGeometry(RayTracingGeometryHandle);
 }
@@ -79,6 +73,8 @@ void FRayTracingGeometry::RequestBuildIfNeeded(FRHICommandListBase& RHICmdList,
 void FRayTracingGeometry::MakeResident(FRHICommandList& RHICmdList)
 {
 	check(EnumHasAllFlags(GeometryState, EGeometryStateFlags::Evicted) && RayTracingGeometryRHI == nullptr);
+	checkf(!EnumHasAllFlags(GeometryState, EGeometryStateFlags::StreamedIn),
+		TEXT("Evicted FRayTracingGeometry shouldn't have StreamedIn flag set."));
 
 	if (!ensureMsgf(DynamicGeometrySharedBufferGenerationID == NonSharedVertexBuffers,
 		TEXT("Cannot call MakeResident(...) on FRayTracingGeometry using shared vertex buffers.\n")
@@ -91,62 +87,23 @@ void FRayTracingGeometry::MakeResident(FRHICommandList& RHICmdList)
 
 	EnumRemoveFlags(GeometryState, EGeometryStateFlags::Evicted);
 
-	// Streaming BLAS needs special handling to not get their "streaming" type wiped out as it will cause issues down the line.	
-	// We only have to do this if the geometry was marked to be streamed in.
-	// In that case we will recreate the geometry as-if it was streamed in.
-	if (EnumHasAnyFlags(GeometryState, FRayTracingGeometry::EGeometryStateFlags::StreamedIn))
-	{
-		FResourceArrayUploadInterface* OfflineData = Initializer.OfflineData;
-
-		Initializer.OfflineData = nullptr;
-
-		// When a mesh is streamed in (FStaticMeshStreamIn::DoFinishUpdate) we update the geometry initializer using just streamed in VB/IB.
-		// That initializer sets a Rendering type but RHI object was created as StreamingDestination and we have a mismatch between geometry initializer and RHI initializer.
-		// It's not an issue unless we try to initialize the geometry again using the geometry's initializer.
-		// We need the current geometry and RHI object to be StreamingDestination so the streaming continues to work.
-		Initializer.Type = ERayTracingGeometryInitializerType::StreamingDestination;
-
-		// Creating RHI with StreamingDestination type will only initialize RHI object but will not created the underlying BLAS buffers.
-		InitRHI(RHICmdList);
-
-		// Here we simulate geometry streaming: create geometry with StreamingSource type to allocate BLAS buffers (1) and swap it with the current geometry (2).
-		// Follows the same pattern as: (1) FStaticMeshStreamIn::CreateBuffers_* (2) FStaticMeshStreamIn::DoFinishUpdate
-		// There is no other way to initialize BLAS buffers for the geometry that has a StreamingDestination type.
-		{
-			FRHIResourceReplaceBatcher Batcher(RHICmdList, 1);
-			FRayTracingGeometryInitializer IntermediateInitializer = Initializer;
-			IntermediateInitializer.Type = ERayTracingGeometryInitializerType::StreamingSource;
-			IntermediateInitializer.OfflineData = OfflineData;
-
-			FRayTracingGeometryRHIRef IntermediateRayTracingGeometry = RHICmdList.CreateRayTracingGeometry(IntermediateInitializer);
-
-			SetRequiresBuild(IntermediateInitializer.OfflineData == nullptr || IntermediateRayTracingGeometry->IsCompressed());
-
-			InitRHIForStreaming(IntermediateRayTracingGeometry, Batcher);
-
-			// When Batcher goes out of scope it will add commands to copy the BLAS buffers on RHI thread.
-			// We need to do it before we build the current geometry (also on RHI thread).
-		}
-
-		RequestBuildIfNeeded(RHICmdList, ERTAccelerationStructureBuildPriority::Normal);
-	}
-	else
-	{
-		InitRHI(RHICmdList);
-	}
+	InitRHI(RHICmdList);
 }
 
 void FRayTracingGeometry::Evict()
 {
 	check(!EnumHasAllFlags(GeometryState, EGeometryStateFlags::Evicted) && RayTracingGeometryRHI != nullptr);
+
 	RemoveBuildRequest();
 	RayTracingGeometryRHI.SafeRelease();
 	EnumAddFlags(GeometryState, EGeometryStateFlags::Evicted);
 
-	if (IsRayTracingUsingReferenceBasedResidency() && EnumHasAllFlags(GeometryState, EGeometryStateFlags::StreamedIn))
+	if (EnumHasAllFlags(GeometryState, EGeometryStateFlags::StreamedIn))
 	{
-		FRHIResourceReplaceBatcher Batcher(FRHICommandListImmediate::Get(), 1);
-		ReleaseRHIForStreaming(Batcher);
+		EnumRemoveFlags(GeometryState, EGeometryStateFlags::StreamedIn);
+		EnumRemoveFlags(GeometryState, EGeometryStateFlags::Valid);
+
+		Initializer.Type = ERayTracingGeometryInitializerType::StreamingDestination;
 	}
 	
 	GRayTracingGeometryManager->RefreshRegisteredGeometry(RayTracingGeometryHandle);
