Metal doesn't support ushaderbytecode chunks, so use metallib chunks instead

diff --git a/Engine/Source/Runtime/RenderCore/Private/ShaderCodeLibrary.cpp b/Engine/Source/Runtime/RenderCore/Private/ShaderCodeLibrary.cpp
--- a/Engine/Source/Runtime/RenderCore/Private/ShaderCodeLibrary.cpp
+++ b/Engine/Source/Runtime/RenderCore/Private/ShaderCodeLibrary.cpp
@@ -2681,7 +2681,12 @@ public:
 						UE_LOG(LogShaderLibrary, Display, TEXT("Running without a pakfile/IoStore and did not find a monolithic library '%s' - attempting disk search for its chunks"), *Name);
 
 						TArray<FString> UshaderbytecodeFiles;
+#if PLATFORM_APPLE
+						// Metal doesn't support ushaderbytecode
+						FString SearchMask = Directory / FString::Printf(TEXT("%s*.metallib"), *Name);
+#else
 						FString SearchMask = Directory / FString::Printf(TEXT("ShaderArchive-*%s*.ushaderbytecode"), *Name);
+#endif
 						IFileManager::Get().FindFiles(UshaderbytecodeFiles, *SearchMask, true, false);
 
 						if (UshaderbytecodeFiles.Num() > 0)
