ShaderParameters: fix incorrect prefix being used for UBMT_RDG_BUFFER_UAV

#rnx

diff --git a/Engine/Source/Runtime/RenderCore/Private/ShaderParameters.cpp b/Engine/Source/Runtime/RenderCore/Private/ShaderParameters.cpp
--- a/Engine/Source/Runtime/RenderCore/Private/ShaderParameters.cpp
+++ b/Engine/Source/Runtime/RenderCore/Private/ShaderParameters.cpp
@@ -244,7 +244,7 @@ static void CreateHLSLUniformBufferStructMembersDeclaration(
 						return FShaderParameterParser::kBindlessSamplerPrefix;
 					}
 
-					if (InBaseType == UBMT_UAV || InBaseType == UBMT_RDG_TEXTURE_UAV)
+					if (InBaseType == UBMT_UAV || InBaseType == UBMT_RDG_TEXTURE_UAV || InBaseType == UBMT_RDG_BUFFER_UAV)
 					{
 						return FShaderParameterParser::kBindlessUAVPrefix;
 					}
@@ -361,7 +361,7 @@ static void CreateHLSLUniformBufferStructMembersDeclaration(
 				Decl.ResourceMembers.Appendf(TEXT("UB_RESOURCE_MEMBER_SAMPLER(%s, %s, %s);\n"), Member.GetShaderType(), *UniformBufferName, *ParameterName);
 				AddStructMember(Member, true);
 			}
-			else if (Member.GetBaseType() == UBMT_UAV || Member.GetBaseType() == UBMT_RDG_TEXTURE_UAV)
+			else if (Member.GetBaseType() == UBMT_UAV || Member.GetBaseType() == UBMT_RDG_TEXTURE_UAV || Member.GetBaseType() == UBMT_RDG_BUFFER_UAV)
 			{
 				Decl.ResourceMembers.Appendf(TEXT("UB_RESOURCE_MEMBER_UAV(%s, %s, %s);\n"), Member.GetShaderType(), *UniformBufferName, *ParameterName);
 				AddStructMember(Member, true);
