Re-enable GPU profiler stats in the editor, by backing out the workaround from CL 31055877 (UE-205211)

Remove "r.GPUStatsMaxQueriesPerFrame" cvar and associated logic in GPU profiler
 - The "overflow" event logic is broken since the latest RHI refactor as the location the EndFrame / Clear happens has changed.
 - Since the underlying bugs / memory inefficiencies that prompted the addition of this cvar have been solved, we don't need to keep any of this now. See below for details.

Query workarounds:

    Remove workaround for memory leak in blueprint debugger
     - Workaround was added in UE 4.26 CL 14541591, but has been superceeded by another workaround added in UE 5.0 CL 20674492.
     - It is no longer necessary to change GPU profiler cvar values before/after blueprint debugging to prevent a memory leak.

    Remove D3D12 excessive memory usage workaround in D3D12RHI.cpp
     - r.GPUStatsMaxQueriesPerFrame was set to 8192 at RHI startup. Previous implementation of queries in D3D12 was very wasteful on memory.
     - The modern implementation packs individual queries into heaps, which greatly reduces the overhead. It is no longer necessary to limit the number of queries in a frame.

Misc changes:

    Remove FDynamicRHI::RHICreateRenderQueryPool function
     - None of the platform RHIs provided an implementation that didn't just return a new FDefaultRHIRenderQueryPool.

    Remove NumQueries from FDefaultRHIRenderQueryPool
     - The value is never used by platform RHIs, and it only causes ensures in cases where the num is too small for what is actually used.

diff --git a/Engine/Config/BaseEditor.ini b/Engine/Config/BaseEditor.ini
--- a/Engine/Config/BaseEditor.ini
+++ b/Engine/Config/BaseEditor.ini
@@ -322,10 +322,6 @@ MaxNumOfAccessViolation=16
 HighSpawnRateOrBurstThreshold=35
 FarLODDistanceTheshold=3000
 
-[ConsoleVariables]
-; Workaround for crashes when opening large levels (UE-205211)
-r.GPUStatsEnabled=0 
-
 [GatherTextFromAssets]
 MinFreeMemory=1024
 MaxUsedMemory=16384
diff --git a/Engine/Source/Runtime/D3D12RHI/Private/D3D12RHI.cpp b/Engine/Source/Runtime/D3D12RHI/Private/D3D12RHI.cpp
--- a/Engine/Source/Runtime/D3D12RHI/Private/D3D12RHI.cpp
+++ b/Engine/Source/Runtime/D3D12RHI/Private/D3D12RHI.cpp
@@ -222,28 +222,6 @@ FD3D12DynamicRHI::FD3D12DynamicRHI(const TArray<TSharedPtr<FD3D12Adapter>>& Chos
 	GSupportsTimestampRenderQueries = true;
 	GSupportsParallelOcclusionQueries = true;
 
-	{
-		// Workaround for 4.14. Limit the number of GPU stats on D3D12 due to an issue with high memory overhead with render queries (Jira UE-38139)
-		// @TODO: Remove this when render query issues are fixed.
-		// 
-		// I think the high memory issues were fixed in 5.0 or some time before.  In the current implementation, a single 64K padded readback buffer is
-		// used for the query pool, with 8 bytes per entry.  So an increase from 1K to 8K should take no extra graphics memory, and this was confirmed
-		// by zero observed change to the D3D12RHI::Memory statistics.  Technically, this is the number of queries for four frames
-		// (NumGPUProfilerBufferedFrames == 4), and a real world test case with multiple views was hitting 880 queries per frame, so the number needed
-		// to be at least 3600 to not run out...  There will be some memory increase from FD3D12RenderQuery structures on the CPU, but it's fairly
-		// trivial (under 300K).
-		// 
-		// I considered removing the limit entirely, but that causes a different code path to run, which hasn't run in a long time on PC (over 5 years
-		// it has been in place), and I wasn't sure how to fully unit test the change, so I figured it was safer to simply increase the existing limit
-		// a bit.
-		//
-		static IConsoleVariable* GPUStatsEnabledCVar = IConsoleManager::Get().FindConsoleVariable(TEXT("r.GPUStatsMaxQueriesPerFrame"));
-		if (GPUStatsEnabledCVar)
-		{
-			GPUStatsEnabledCVar->Set(8192);
-		}
-	}
-
 	// Manually enable Async BVH build for D3D12 RHI
 	GRHISupportsRayTracingAsyncBuildAccelerationStructure = true;
 
diff --git a/Engine/Source/Runtime/OpenGLDrv/Private/OpenGLDrv.h b/Engine/Source/Runtime/OpenGLDrv/Private/OpenGLDrv.h
--- a/Engine/Source/Runtime/OpenGLDrv/Private/OpenGLDrv.h
+++ b/Engine/Source/Runtime/OpenGLDrv/Private/OpenGLDrv.h
@@ -480,8 +480,6 @@ public:
 	virtual void RHIMapStagingSurface_RenderThread(class FRHICommandListImmediate& RHICmdList, FRHITexture* Texture, uint32 GPUIndex, FRHIGPUFence* Fence, void*& OutData, int32& OutWidth, int32& OutHeight) final override;
 	virtual void RHIUnmapStagingSurface_RenderThread(class FRHICommandListImmediate& RHICmdList, FRHITexture* Texture, uint32 GPUIndex) final override;
 
-
-	virtual FRenderQueryPoolRHIRef RHICreateRenderQueryPool(ERenderQueryType QueryType, uint32 NumQueries = UINT32_MAX) final override;
 	virtual FStagingBufferRHIRef RHICreateStagingBuffer() final override;
 	virtual FGPUFenceRHIRef RHICreateGPUFence(const FName &Name) final override;
 
diff --git a/Engine/Source/Runtime/OpenGLDrv/Private/OpenGLQuery.cpp b/Engine/Source/Runtime/OpenGLDrv/Private/OpenGLQuery.cpp
--- a/Engine/Source/Runtime/OpenGLDrv/Private/OpenGLQuery.cpp
+++ b/Engine/Source/Runtime/OpenGLDrv/Private/OpenGLQuery.cpp
@@ -625,14 +625,6 @@ bool FOpenGLDynamicRHI::RHIGetRenderQueryResult(FRHIRenderQuery* QueryRHI, uint6
 	return false;
 }
 
-
-
-FRenderQueryPoolRHIRef FOpenGLDynamicRHI::RHICreateRenderQueryPool(ERenderQueryType QueryType, uint32 /*NumQueries*/)
-{
-	// Workaround to fix UE-77873: pass <QueryType> to primary implementation and ignore <NumQueries> for pre-allocated queries
-	return FDynamicRHI::RHICreateRenderQueryPool(QueryType);
-}
-
 extern void OnQueryCreation( FOpenGLRenderQuery* Query );
 extern void OnQueryDeletion( FOpenGLRenderQuery* Query );
 
diff --git a/Engine/Source/Runtime/RHI/Private/DynamicRHI.cpp b/Engine/Source/Runtime/RHI/Private/DynamicRHI.cpp
--- a/Engine/Source/Runtime/RHI/Private/DynamicRHI.cpp
+++ b/Engine/Source/Runtime/RHI/Private/DynamicRHI.cpp
@@ -759,23 +759,6 @@ bool FDynamicRHI::RHIMatchPrecachePSOInitializers(const FGraphicsPipelineStateIn
 	return true;
 }
 
-FDefaultRHIRenderQueryPool::FDefaultRHIRenderQueryPool(ERenderQueryType InQueryType, FDynamicRHI* InDynamicRHI, uint32 InNumQueries)
-	: DynamicRHI(InDynamicRHI)
-	, QueryType(InQueryType)
-	, NumQueries(InNumQueries)
-{
-	if (NumQueries != UINT32_MAX && (GSupportsTimestampRenderQueries || InQueryType != RQT_AbsoluteTime))
-	{
-		Queries.Reserve(NumQueries);
-		for (uint32 i = 0; i < NumQueries; i++)
-		{
-			Queries.Push(DynamicRHI->RHICreateRenderQuery(QueryType));
-			check(Queries.Last().IsValid());
-			++AllocatedQueries;
-		}
-	}
-}
-
 FDefaultRHIRenderQueryPool::~FDefaultRHIRenderQueryPool()
 {
 	check(IsInRHIThread() || IsInRenderingThread());
@@ -791,37 +774,28 @@ FRHIPooledRenderQuery FDefaultRHIRenderQueryPool::AllocateQuery()
 	}
 	else
 	{
-		FRHIPooledRenderQuery Query = FRHIPooledRenderQuery(this, DynamicRHI->RHICreateRenderQuery(QueryType));
+		FRHIPooledRenderQuery Query = FRHIPooledRenderQuery(this, RHICreateRenderQuery(QueryType));
 		if (Query.IsValid())
 		{
 			++AllocatedQueries;
 		}
-		ensure(AllocatedQueries <= NumQueries);
+
 		return Query;
 	}
 }
 
 void FDefaultRHIRenderQueryPool::ReleaseQuery(TRefCountPtr<FRHIRenderQuery>&& Query)
 {
-	if (QueryType == ERenderQueryType::RQT_Occlusion)
-	{
-		static int dbg = 0;
-		dbg++;
-	}
 	check(IsInParallelRenderingThread());
-	//Hard to validate because of Resource resurrection, better to remove GetQueryRef entirely
-	//checkf(Query.IsValid() && Query.GetRefCount() <= 2, TEXT("Query has been released but reference still held: use FRHIPooledRenderQuery::GetQueryRef() with extreme caution"));
-	
 	checkf(Query.IsValid(), TEXT("Only release valid queries"));
-	checkf((uint32)Queries.Num() < NumQueries, TEXT("Pool contains more queries than it started with, double release somewhere?"));
 
 	Queries.Push(MoveTemp(Query));
 	check(!Query.IsValid());
 }
 
-FRenderQueryPoolRHIRef RHICreateRenderQueryPool(ERenderQueryType QueryType, uint32 NumQueries)
+FRenderQueryPoolRHIRef RHICreateRenderQueryPool(ERenderQueryType QueryType, uint32 /* unused NumQueries */)
 {
-	return GDynamicRHI->RHICreateRenderQueryPool(QueryType, NumQueries);
+	return new FDefaultRHIRenderQueryPool(QueryType);
 }
 
 EColorSpaceAndEOTF FDynamicRHI::RHIGetColorSpace(FRHIViewport* Viewport)
diff --git a/Engine/Source/Runtime/RHI/Public/DynamicRHI.h b/Engine/Source/Runtime/RHI/Public/DynamicRHI.h
--- a/Engine/Source/Runtime/RHI/Public/DynamicRHI.h
+++ b/Engine/Source/Runtime/RHI/Public/DynamicRHI.h
@@ -93,16 +93,17 @@ class FDynamicRHI;
 class FDefaultRHIRenderQueryPool final : public FRHIRenderQueryPool
 {
 public:
-	RHI_API FDefaultRHIRenderQueryPool(ERenderQueryType InQueryType, FDynamicRHI* InDynamicRHI, uint32 InNumQueries);
-	RHI_API ~FDefaultRHIRenderQueryPool() override;
+	FDefaultRHIRenderQueryPool(ERenderQueryType InQueryType)
+		: QueryType(InQueryType)
+	{}
+
+	RHI_API virtual ~FDefaultRHIRenderQueryPool();
 
 private:
 	RHI_API virtual FRHIPooledRenderQuery AllocateQuery() override;
 	RHI_API virtual void ReleaseQuery(TRefCountPtr<FRHIRenderQuery>&& Query) override;
 
-	FDynamicRHI* DynamicRHI = nullptr;
-	ERenderQueryType QueryType;
-	uint32 NumQueries = 0;
+	const ERenderQueryType QueryType;
 	uint32 AllocatedQueries = 0;
 	TArray<TRefCountPtr<FRHIRenderQuery>> Queries;
 };
@@ -209,16 +210,6 @@ public:
 	{
 		return nullptr;
 	}
-	/**
-	* Creates a pool for querys like timers or occlusion queries.
-	* @param QueryType The ype of the queries provided by this pool like RQT_Occlusion or RQT_AbsoluteTime.
-	* @return the Querypool.
-	*/
-	// FlushType: Must be Thread-Safe.
-	virtual FRenderQueryPoolRHIRef RHICreateRenderQueryPool(ERenderQueryType QueryType, uint32 NumQueries = UINT32_MAX)
-	{
-		return new FDefaultRHIRenderQueryPool(QueryType, this, NumQueries);
-	}
 
 	virtual FGPUFenceRHIRef RHICreateGPUFence(const FName &Name)
 	{
diff --git a/Engine/Source/Runtime/RHI/Public/RHIValidation.h b/Engine/Source/Runtime/RHI/Public/RHIValidation.h
--- a/Engine/Source/Runtime/RHI/Public/RHIValidation.h
+++ b/Engine/Source/Runtime/RHI/Public/RHIValidation.h
@@ -235,11 +235,6 @@ public:
 		return RHI->RHICreateShaderLibrary(Platform, FilePath, Name);
 	}
 
-	virtual FRenderQueryPoolRHIRef RHICreateRenderQueryPool(ERenderQueryType QueryType, uint32 NumQueries = UINT32_MAX) override final
-	{
-		return RHI->RHICreateRenderQueryPool(QueryType, NumQueries);
-	}
-
 	virtual FGPUFenceRHIRef RHICreateGPUFence(const FName &Name) override final
 	{
 		return RHI->RHICreateGPUFence(Name);
diff --git a/Engine/Source/Runtime/RenderCore/Private/ProfilingDebugging/RealtimeGPUProfiler.cpp b/Engine/Source/Runtime/RenderCore/Private/ProfilingDebugging/RealtimeGPUProfiler.cpp
--- a/Engine/Source/Runtime/RenderCore/Private/ProfilingDebugging/RealtimeGPUProfiler.cpp
+++ b/Engine/Source/Runtime/RenderCore/Private/ProfilingDebugging/RealtimeGPUProfiler.cpp
@@ -12,15 +12,10 @@
 
 // @todo - new gpu profiler
 
-#else
+#elif HAS_GPU_STATS
 
 typedef TArray<TCHAR, TInlineAllocator<4096u>> FDescriptionStringBuffer;
 
-// Only exposed for debugging. Disabling this carries a severe performance penalty
-#define RENDER_QUERY_POOLING_ENABLED 1
-
-#if HAS_GPU_STATS 
-
 CSV_DEFINE_CATEGORY_MODULE(RENDERCORE_API, GPU, true);
 
 static TAutoConsoleVariable<int> CVarGPUStatsEnabled(
@@ -28,14 +23,6 @@ static TAutoConsoleVariable<int> CVarGPUStatsEnabled(
 	1,
 	TEXT("Enables or disables GPU stat recording"));
 
-
-static TAutoConsoleVariable<int> CVarGPUStatsMaxQueriesPerFrame(
-	TEXT("r.GPUStatsMaxQueriesPerFrame"),
-	-1,
-	TEXT("Limits the number of timestamps allocated per frame. -1 = no limit"), 
-	ECVF_RenderThreadSafe);
-
-
 static TAutoConsoleVariable<int> CVarGPUCsvStatsEnabled(
 	TEXT("r.GPUCsvStatsEnabled"),
 	0,
@@ -56,9 +43,6 @@ static TAutoConsoleVariable<int> CVarGPUStatsChildTimesIncluded(
 	TEXT("to the total GPU time, so we probably want this disabled.\n")
 );
 
-#endif //HAS_GPU_STATS
-
-#if HAS_GPU_STATS
 static const int32 NumGPUProfilerBufferedFrames = 4;
 static const uint64 InvalidQueryResult = 0xFFFFFFFFFFFFFFFFull;
 
@@ -396,7 +380,6 @@ class FRealtimeGPUProfilerFrame
 public:
 	FRealtimeGPUProfilerFrame(FRenderQueryPoolRHIRef InRenderQueryPool, uint32& InQueryCount)
 		: NextEventIdx(1)
-		, OverflowEventCount(0)
 		, NextResultPendingEventIdx(1)
 		, QueryCount(InQueryCount)
 		, RenderQueryPool(InRenderQueryPool)
@@ -428,8 +411,6 @@ public:
 
 	void Clear(void* Dummy)
 	{
-		check(!OverflowEventCount);
-
 		NextEventIdx = 1;
 		NextResultPendingEventIdx = 1;
 
@@ -449,18 +430,8 @@ public:
 	{
 		if (NextEventIdx >= GpuProfilerEvents.Num())
 		{
-			const int32 MaxNumQueries = CVarGPUStatsMaxQueriesPerFrame.GetValueOnRenderThread();
-
-			if (MaxNumQueries < 0 || QueryCount < (uint32)MaxNumQueries)
-			{
-				GpuProfilerEvents.Add(new FRealtimeGPUProfilerEvent(*RenderQueryPool));
-				QueryCount += FRealtimeGPUProfilerEvent::GetNumRHIQueriesPerEvent();
-			}
-			else
-			{
-				++OverflowEventCount;
-				return {};
-			}
+			GpuProfilerEvents.Add(new FRealtimeGPUProfilerEvent(*RenderQueryPool));
+			QueryCount += FRealtimeGPUProfilerEvent::GetNumRHIQueriesPerEvent();
 		}
 
 		const int32 EventIdx = NextEventIdx++;
@@ -480,12 +451,6 @@ public:
 
 	FRealtimeGPUProfilerQuery PopEvent()
 	{
-		if (OverflowEventCount)
-		{
-			--OverflowEventCount;
-			return {};
-		}
-
 		const int32 EventIdx = EventStack.Pop(EAllowShrinking::No);
 
 		return GpuProfilerEvents[EventIdx].End();
@@ -853,7 +818,6 @@ private:
 	};
 
 	int32 NextEventIdx;
-	int32 OverflowEventCount;
 	int32 NextResultPendingEventIdx;
 
 	uint32& QueryCount;
@@ -898,8 +862,7 @@ FRealtimeGPUProfiler::FRealtimeGPUProfiler()
 {
 	if (GSupportsTimestampRenderQueries)
 	{
-		const int MaxGPUQueries = CVarGPUStatsMaxQueriesPerFrame.GetValueOnRenderThread();
-		RenderQueryPool = RHICreateRenderQueryPool(RQT_AbsoluteTime, (MaxGPUQueries > 0) ? MaxGPUQueries : UINT32_MAX);
+		RenderQueryPool = RHICreateRenderQueryPool(RQT_AbsoluteTime);
 		for (int Index = 0; Index < NumGPUProfilerBufferedFrames; Index++)
 		{
 			Frames.Add(new FRealtimeGPUProfilerFrame(RenderQueryPool, QueryCount));
@@ -1121,9 +1084,8 @@ FScopedDrawStatCategory::~FScopedDrawStatCategory()
 		RHICmdList->SetDrawStatsCategory(Previous);
 	}
 }
-#endif // HAS_GPU_STATS
 
-#if GPUPROFILERTRACE_ENABLED && HAS_GPU_STATS
+#if GPUPROFILERTRACE_ENABLED
 FRealtimeGPUProfilerHistoryItem::FRealtimeGPUProfilerHistoryItem()
 {
 	FMemory::Memset(*this, 0);
@@ -1159,4 +1121,4 @@ void FRealtimeGPUProfiler::FetchPerfByDescription(TArray<FRealtimeGPUProfilerDes
 }
 #endif  // GPUPROFILERTRACE_ENABLED
 
-#endif // (RHI_NEW_GPU_PROFILER == 0)
+#endif // HAS_GPU_STATS
diff --git a/Engine/Source/Runtime/Slate/Private/Framework/Application/SlateApplication.cpp b/Engine/Source/Runtime/Slate/Private/Framework/Application/SlateApplication.cpp
--- a/Engine/Source/Runtime/Slate/Private/Framework/Application/SlateApplication.cpp
+++ b/Engine/Source/Runtime/Slate/Private/Framework/Application/SlateApplication.cpp
@@ -3739,11 +3739,6 @@ void FSlateApplication::EnterDebuggingMode()
 	GFirstFrameIntraFrameDebugging = true;
 #endif	//WITH_EDITORONLY_DATA
 
-	//Disable GPU Profiler during BluePrint Debugging to prevent leaking memory.
-	IConsoleVariable* CvarMaxQueriesPerFrame = IConsoleManager::Get().FindConsoleVariable(TEXT("r.GPUStatsMaxQueriesPerFrame"));
-	int MaxQueriesPerFrame = CvarMaxQueriesPerFrame->GetInt();
-	CvarMaxQueriesPerFrame->Set(0);
-
 	// Tick slate from here in the event that we should not return until the modal window is closed.
 	while (!bRequestLeaveDebugMode)
 	{
@@ -3767,8 +3762,6 @@ void FSlateApplication::EnterDebuggingMode()
 		GIntraFrameDebuggingGameThread = !bRequestLeaveDebugMode;
 #endif	//WITH_EDITORONLY_DATA
 	}
-
-	CvarMaxQueriesPerFrame->Set(MaxQueriesPerFrame);
 	
 	Renderer->BeginFrame();
 	bRequestLeaveDebugMode = false;
