Fix crash when running DumpGPU on application using memoryless textures

diff --git a/Engine/Source/Runtime/RenderCore/Private/DumpGPU.cpp b/Engine/Source/Runtime/RenderCore/Private/DumpGPU.cpp
--- a/Engine/Source/Runtime/RenderCore/Private/DumpGPU.cpp
+++ b/Engine/Source/Runtime/RenderCore/Private/DumpGPU.cpp
@@ -1358,12 +1358,15 @@ public:
 		bool bAllowDumpBinary)
 	{
 		int32 DumpTextureMode = GDumpTextureCVar.GetValueOnRenderThread();
-
-		if (DumpTextureMode == 0)
+		
+		// We can't dump memoryless textures
+		bool IsMemoryless = EnumHasAnyFlags(SubresourceDesc.Texture->Desc.Flags, TexCreate_Memoryless);
+		
+		if (DumpTextureMode == 0 || IsMemoryless)
 		{
 			return;
 		}
-
+ 
 		const FRDGTextureDesc& Desc = SubresourceDesc.Texture->Desc;
 		const FString UniqueResourceSubResourceName = GetUniqueSubResourceName(SubresourceDesc);
 		const FTextureSubresourceDumpDesc SubresourceDumpDesc = TranslateSubresourceDumpDesc(SubresourceDesc);
