Don't use shader parameter layout hash to retrieve the uniform buffer entries on the shader binding layout because it's not unique

diff --git a/Engine/Source/Runtime/RHI/Private/RHIContext.cpp b/Engine/Source/Runtime/RHI/Private/RHIContext.cpp
--- a/Engine/Source/Runtime/RHI/Private/RHIContext.cpp
+++ b/Engine/Source/Runtime/RHI/Private/RHIContext.cpp
@@ -61,7 +61,7 @@ void FUniformBufferStaticBindings::AddUniformBuffer(FRHIUniformBuffer* UniformBu
 	// Only care about the static slots if no shader desc is used, otherwise the desc is used to validate that it contains the uniform buffer
 	if (ShaderBindingLayout)
 	{
-		const FRHIUniformBufferShaderBindingLayout* UniformBufferEntry = ShaderBindingLayout->FindEntry(Layout.GetHash());
+		const FRHIUniformBufferShaderBindingLayout* UniformBufferEntry = ShaderBindingLayout->FindEntry(Layout.Name);
 		checkf(UniformBufferEntry, TEXT("Attempted to set a static uniform buffer %s which is not defined in the ShaderBindingLayout provided."), *Layout.GetDebugName());
 		UniformBuffers[UniformBufferEntry->CBVResourceIndex] = UniformBuffer;
 	}
diff --git a/Engine/Source/Runtime/RHI/Public/RHIShaderBindingLayout.h b/Engine/Source/Runtime/RHI/Public/RHIShaderBindingLayout.h
--- a/Engine/Source/Runtime/RHI/Public/RHIShaderBindingLayout.h
+++ b/Engine/Source/Runtime/RHI/Public/RHIShaderBindingLayout.h
@@ -12,7 +12,7 @@
 */
 struct FRHIUniformBufferShaderBindingLayout
 {
-	uint32 LayoutHash;
+	FString LayoutName;
 	union
 	{
 		struct
@@ -26,24 +26,24 @@ struct FRHIUniformBufferShaderBindingLayout
 		uint32 Flags;
 	};
 
-	FRHIUniformBufferShaderBindingLayout() : LayoutHash(0), Flags(0) {}
+	FRHIUniformBufferShaderBindingLayout() : Flags(0) {}
 			
 	friend uint32 GetTypeHash(const FRHIUniformBufferShaderBindingLayout& Entry)
 	{
-		uint32 Hash = GetTypeHash(Entry.LayoutHash);
+		uint32 Hash = GetTypeHash(Entry.LayoutName);
 		Hash = HashCombineFast(Hash, GetTypeHash(Entry.Flags));
 		return Hash;
 	}
 	
 	bool operator==(const FRHIUniformBufferShaderBindingLayout& Other) const
 	{
-		return LayoutHash == Other.LayoutHash &&
+		return LayoutName == Other.LayoutName &&
 			Flags == Other.Flags;			
 	}
 
 	friend inline FArchive& operator << (FArchive& Ar, FRHIUniformBufferShaderBindingLayout& F)
 	{
-		Ar << F.LayoutHash;
+		Ar << F.LayoutName;
 		Ar << F.Flags;
 		return Ar;
 	}
@@ -70,7 +70,7 @@ ENUM_CLASS_FLAGS(EShaderBindingLayoutFlags)
 class FRHIShaderBindingLayout
 {
 public:
-
+		
 	FRHIShaderBindingLayout() = default;
 	
 	FRHIShaderBindingLayout(EShaderBindingLayoutFlags InFlags, TConstArrayView<FRHIUniformBufferShaderBindingLayout> InUniformBufferEntries) : Flags(InFlags)
@@ -90,11 +90,11 @@ public:
 	uint32 GetNumUniformBufferEntries() const { return NumUniformBufferEntries; }
 	const FRHIUniformBufferShaderBindingLayout& GetUniformBufferEntry(uint32 Index) const { check(Index < NumUniformBufferEntries); return UniformBufferEntries[Index]; }
 	   	
-	const FRHIUniformBufferShaderBindingLayout* FindEntry(uint32 LayoutHash) const
+	const FRHIUniformBufferShaderBindingLayout* FindEntry(const FString& LayoutName) const
 	{
 		for (const FRHIUniformBufferShaderBindingLayout& Entry : UniformBufferEntries)
 		{
-			if (Entry.LayoutHash == LayoutHash)
+			if (Entry.LayoutName == LayoutName)
 			{
 				return &Entry;
 			}
diff --git a/Engine/Source/Runtime/RenderCore/Private/ShaderCore.cpp b/Engine/Source/Runtime/RenderCore/Private/ShaderCore.cpp
--- a/Engine/Source/Runtime/RenderCore/Private/ShaderCore.cpp
+++ b/Engine/Source/Runtime/RenderCore/Private/ShaderCore.cpp
@@ -839,26 +839,22 @@ uint32 FShaderParameterMap::CountParametersOfType(EShaderParameterType InType) c
 FThreadSafeSharedAnsiStringPtr FShaderBindingLayout::GetUniformBufferDeclarationAnsiPtr(const FShaderParametersMetadata* ShaderParametersMetadata) const
 {
 	// Does the binding layout contain a uniform buffer declaration for this shader paramaters
-	uint32 LayoutHash = ShaderParametersMetadata->GetLayout().GetHash();
-	const FRHIUniformBufferShaderBindingLayout* Entry = RHILayout.FindEntry(LayoutHash);
-	if (Entry)
+	FString UniformBufferName(ShaderParametersMetadata->GetShaderVariableName());
+	const FThreadSafeSharedAnsiStringPtr* UniformBufferDeclarationAnsi = UniformBufferMap.Find(UniformBufferName);
+	if (UniformBufferDeclarationAnsi)
 	{
-		uint8 UniformBufferIndex = Entry->CBVResourceIndex;
-		return UniformBufferData[UniformBufferIndex].UniformBufferDeclarationAnsi;
+		return *UniformBufferDeclarationAnsi;
 	}
 
 	// use the shared declaration from the metadata object itself (no fixed register or space defined)
 	return ShaderParametersMetadata->GetUniformBufferDeclarationAnsiPtr();
 }
 
-void FShaderBindingLayout::SetUniformBufferDeclarationAnsiPtr(const FShaderParametersMetadata* ShaderParametersMetadata, const FString& UniformBufferName, FThreadSafeSharedAnsiStringPtr UniformBufferDeclarationAnsi)
+void FShaderBindingLayout::SetUniformBufferDeclarationAnsiPtr(const FShaderParametersMetadata* ShaderParametersMetadata, FThreadSafeSharedAnsiStringPtr UniformBufferDeclarationAnsi)
 {
-	uint32 LayoutHash = ShaderParametersMetadata->GetLayout().GetHash();
-	const FRHIUniformBufferShaderBindingLayout* Entry = RHILayout.FindEntry(LayoutHash);
-	check(Entry);
-	UniformBufferData.SetNum(FMath::Max(UniformBufferData.Num(), int32(Entry->CBVResourceIndex + 1)));
-	UniformBufferData[Entry->CBVResourceIndex].UniformBufferName = UniformBufferName;
-	UniformBufferData[Entry->CBVResourceIndex].UniformBufferDeclarationAnsi = UniformBufferDeclarationAnsi;
+	FString UniformBufferName(ShaderParametersMetadata->GetShaderVariableName());
+	check(!UniformBufferMap.Contains(UniformBufferName));
+	UniformBufferMap.Add(UniformBufferName, UniformBufferDeclarationAnsi);
 }
 
 void FShaderBindingLayout::AddRequiredSymbols(TArray<FString>& RequiredSymbols) const
@@ -867,9 +863,9 @@ void FShaderBindingLayout::AddRequiredSymbols(TArray<FString>& RequiredSymbols)
 	check(EnumHasAllFlags(RHILayout.GetFlags(), EShaderBindingLayoutFlags::BindlessResources | EShaderBindingLayoutFlags::BindlessSamplers));
 
 	// Don't remove unused uniform buffers defined in the fixed shader binding layout because they are required to be declared for certain platforms
-	for (const FUniformBufferData& Data : UniformBufferData)
+	for (auto Iter = UniformBufferMap.CreateConstIterator(); Iter; ++Iter)
 	{
-		RequiredSymbols.Add(Data.UniformBufferName);
+		RequiredSymbols.Add(Iter.Key());
 	}
 }
 
diff --git a/Engine/Source/Runtime/RenderCore/Private/ShaderParameterMetadata.cpp b/Engine/Source/Runtime/RenderCore/Private/ShaderParameterMetadata.cpp
--- a/Engine/Source/Runtime/RenderCore/Private/ShaderParameterMetadata.cpp
+++ b/Engine/Source/Runtime/RenderCore/Private/ShaderParameterMetadata.cpp
@@ -178,7 +178,7 @@ void BuildShaderBindingLayout(TConstArrayView<FShaderParametersMetadata*> Unifor
 		check(UniformBufferLayout.ConstantBufferSize > 0);
 
 		FRHIUniformBufferShaderBindingLayout& NewEntry = UniformBufferEntries.AddDefaulted_GetRef();
-		NewEntry.LayoutHash = UniformBufferLayout.GetHash();
+		NewEntry.LayoutName = UniformBufferLayout.Name;
 		NewEntry.RegisterSpace = UE_HLSL_SPACE_STATIC_SHADER_BINDINGS;
 		NewEntry.CBVResourceIndex = CurrentCBVResourceIndex;
 		NewEntry.BaseSRVResourceIndex = CurrentBaseSRVResourceIndex;
@@ -208,7 +208,8 @@ void BuildShaderBindingLayout(TConstArrayView<FShaderParametersMetadata*> Unifor
 	for (uint8 UniformBufferIndex = 0; UniformBufferIndex < UniformBuffers.Num(); ++UniformBufferIndex)	
 	{
 		FShaderParametersMetadata* ShaderParametersMetaData = UniformBuffers[UniformBufferIndex];
-		const FRHIUniformBufferShaderBindingLayout* UniformBufferSBLayout = ShaderBindingLayout.RHILayout.FindEntry(ShaderParametersMetaData->GetLayout().GetHash());
+		const FString& LayoutName = ShaderParametersMetaData->GetLayout().Name;
+		const FRHIUniformBufferShaderBindingLayout* UniformBufferSBLayout = ShaderBindingLayout.RHILayout.FindEntry(LayoutName);
 		check(UniformBufferSBLayout);
 
 		FString UniformBufferName(ShaderParametersMetaData->GetShaderVariableName());
@@ -220,7 +221,7 @@ void BuildShaderBindingLayout(TConstArrayView<FShaderParametersMetadata*> Unifor
 		ShaderConvertAndStripComments(NewDeclaration, *NewDeclarationAnsi);
 		FThreadSafeSharedAnsiStringPtr UniformBufferDeclarationAnsi = MakeShareable(NewDeclarationAnsi);
 
-		ShaderBindingLayout.SetUniformBufferDeclarationAnsiPtr(ShaderParametersMetaData, UniformBufferName, UniformBufferDeclarationAnsi);
+		ShaderBindingLayout.SetUniformBufferDeclarationAnsiPtr(ShaderParametersMetaData, UniformBufferDeclarationAnsi);
 	}
 #endif //WITH_EDITOR
 
diff --git a/Engine/Source/Runtime/RenderCore/Public/ShaderCore.h b/Engine/Source/Runtime/RenderCore/Public/ShaderCore.h
--- a/Engine/Source/Runtime/RenderCore/Public/ShaderCore.h
+++ b/Engine/Source/Runtime/RenderCore/Public/ShaderCore.h
@@ -453,18 +453,13 @@ public:
 
 #if WITH_EDITOR
 	FThreadSafeSharedAnsiStringPtr GetUniformBufferDeclarationAnsiPtr(const FShaderParametersMetadata* ShaderParametersMetadata) const;
-	void SetUniformBufferDeclarationAnsiPtr(const FShaderParametersMetadata* ShaderParametersMetadata, const FString& UniformBufferName, FThreadSafeSharedAnsiStringPtr UniformBufferDeclarationAnsi);
+	void SetUniformBufferDeclarationAnsiPtr(const FShaderParametersMetadata* ShaderParametersMetadata, FThreadSafeSharedAnsiStringPtr UniformBufferDeclarationAnsi);
 
 	RENDERCORE_API void AddRequiredSymbols(TArray<FString>& RequiredSymbols) const;
 
 protected:
-	
-	struct FUniformBufferData
-	{
-		FString	UniformBufferName;
-		FThreadSafeSharedAnsiStringPtr UniformBufferDeclarationAnsi;
-	};
-	TArray<FUniformBufferData> UniformBufferData;
+		
+	TMap<FString, FThreadSafeSharedAnsiStringPtr> UniformBufferMap;
 #endif
 };
 
