Added shader parameter variant for scene texture extracts that can utilize shader parameter structs.

diff --git a/Engine/Source/Runtime/RenderCore/Public/ShaderParameterMacros.h b/Engine/Source/Runtime/RenderCore/Public/ShaderParameterMacros.h
--- a/Engine/Source/Runtime/RenderCore/Public/ShaderParameterMacros.h
+++ b/Engine/Source/Runtime/RenderCore/Public/ShaderParameterMacros.h
@@ -137,8 +137,17 @@ class TUniformBufferRef : public FUniformBufferRHIRef
 
 public:
 	/** Initializes the reference to null. */
-	TUniformBufferRef()
-	{}
+	TUniformBufferRef() = default;
+
+	/** Construct an instance from an existing RHI uniform buffer pointer. Validates at runtime that the type of the uniform buffer matches the template struct type. */
+	explicit TUniformBufferRef(FRHIUniformBuffer* InRHIRef)
+		: FUniformBufferRHIRef(InRHIRef)
+	{
+		checkf(!InRHIRef || TUniformBufferMetadataHelper<TBufferStruct>::GetStructMetadata()->GetLayoutPtr() == InRHIRef->GetLayoutPtr(),
+			TEXT("Attempted to create a uniform buffer of type '%s' from uniform buffer pointer of type '%s'"),
+			*TUniformBufferMetadataHelper<TBufferStruct>::GetStructMetadata()->GetLayout().Name,
+			*InRHIRef->GetLayout().Name);
+	}
 
 	/** Creates a uniform buffer with the given value, and returns a structured reference to it. */
 	static TUniformBufferRef<TBufferStruct> CreateUniformBufferImmediate(const TBufferStruct& Value, EUniformBufferUsage Usage, EUniformBufferValidation Validation = EUniformBufferValidation::ValidateResources)
@@ -166,11 +175,6 @@ public:
 
 private:
 
-	/** A private constructor used to coerce an arbitrary RHI uniform buffer reference to a structured reference. */
-	TUniformBufferRef(FRHIUniformBuffer* InRHIRef)
-	: FUniformBufferRHIRef(InRHIRef)
-	{}
-
 	template<typename TBufferStruct2>
 	friend class TUniformBuffer;
 
diff --git a/Engine/Source/Runtime/Renderer/Public/SceneRenderTargetParameters.h b/Engine/Source/Runtime/Renderer/Public/SceneRenderTargetParameters.h
--- a/Engine/Source/Runtime/Renderer/Public/SceneRenderTargetParameters.h
+++ b/Engine/Source/Runtime/Renderer/Public/SceneRenderTargetParameters.h
@@ -117,6 +117,11 @@ extern RENDERER_API FSceneTextureShaderParameters CreateSceneTextureShaderParame
 
 extern RENDERER_API FSceneTextureShaderParameters GetSceneTextureShaderParameters(const FSceneView& View);
 
+BEGIN_SHADER_PARAMETER_STRUCT(FSceneTextureExtractsParameters, )
+	SHADER_PARAMETER_STRUCT_REF(FSceneTextureUniformParameters, SceneTextures)
+	SHADER_PARAMETER_STRUCT_REF(FMobileSceneTextureUniformParameters, MobileSceneTextures)
+END_SHADER_PARAMETER_STRUCT()
+
 /** Struct containing references to extracted RHI resources after RDG execution. All textures are
  *  left in an SRV read state, so they can safely be used for read without being re-imported into
  *  RDG. Likewise, the uniform buffer is non-RDG and can be used as is.
@@ -139,6 +144,14 @@ public:
 		return MobileUniformBuffer;
 	}
 
+	FSceneTextureExtractsParameters GetShaderParameters() const
+	{
+		FSceneTextureExtractsParameters Parameters;
+		Parameters.SceneTextures = UniformBuffer;
+		Parameters.MobileSceneTextures = MobileUniformBuffer;
+		return Parameters;
+	}
+
 	FRHITexture* GetDepthTexture() const
 	{
 		return Depth ? Depth->GetRHI() : nullptr;
