Refactored RDG to allow certain setup tasks to run in parallel with RDG compilation. Disabled render pass merging if the underlying GPU is not tile based, as this avoids large render pass ranges from being batched into a single async task.

diff --git a/Engine/Plugins/Experimental/Water/Source/Runtime/Private/WaterInfoRendering.cpp b/Engine/Plugins/Experimental/Water/Source/Runtime/Private/WaterInfoRendering.cpp
--- a/Engine/Plugins/Experimental/Water/Source/Runtime/Private/WaterInfoRendering.cpp
+++ b/Engine/Plugins/Experimental/Water/Source/Runtime/Private/WaterInfoRendering.cpp
@@ -454,7 +454,7 @@ static void UpdateWaterInfoRendering_RenderThread(
 
 		SCOPED_DRAW_EVENT(RHICmdList, DepthRendering_RT);
 
-		FRDGBuilder GraphBuilder(RHICmdList, RDG_EVENT_NAME("WaterInfoDepthRendering"), ERDGBuilderFlags::AllowParallelExecute);
+		FRDGBuilder GraphBuilder(RHICmdList, RDG_EVENT_NAME("WaterInfoDepthRendering"), ERDGBuilderFlags::Parallel);
 
 		FRDGTextureRef TargetTexture = RegisterExternalTexture(GraphBuilder, RenderTarget->GetRenderTargetTexture(), TEXT("WaterDepthTarget"));
 
@@ -512,7 +512,7 @@ static void UpdateWaterInfoRendering_RenderThread(
 
 		SCOPED_DRAW_EVENT(RHICmdList, ColorRendering_RT);
 
-		FRDGBuilder GraphBuilder(RHICmdList, RDG_EVENT_NAME("WaterInfoColorRendering"), ERDGBuilderFlags::AllowParallelExecute);
+		FRDGBuilder GraphBuilder(RHICmdList, RDG_EVENT_NAME("WaterInfoColorRendering"), ERDGBuilderFlags::Parallel);
 
 		FRDGTextureRef TargetTexture = RegisterExternalTexture(GraphBuilder, RenderTarget->GetRenderTargetTexture(), TEXT("WaterColorTarget"));
 
@@ -564,7 +564,7 @@ static void UpdateWaterInfoRendering_RenderThread(
 
 		SCOPED_DRAW_EVENT(RHICmdList, DilationRendering_RT);
 
-		FRDGBuilder GraphBuilder(RHICmdList, RDG_EVENT_NAME("WaterInfoDilationRendering"), ERDGBuilderFlags::AllowParallelExecute);
+		FRDGBuilder GraphBuilder(RHICmdList, RDG_EVENT_NAME("WaterInfoDilationRendering"), ERDGBuilderFlags::Parallel);
 
 		FRDGTextureRef TargetTexture = RegisterExternalTexture(GraphBuilder, RenderTarget->GetRenderTargetTexture(), TEXT("WaterDilationTarget"));
 
diff --git a/Engine/Source/Runtime/RenderCore/Private/RenderGraphBuilder.cpp b/Engine/Source/Runtime/RenderCore/Private/RenderGraphBuilder.cpp
--- a/Engine/Source/Runtime/RenderCore/Private/RenderGraphBuilder.cpp
+++ b/Engine/Source/Runtime/RenderCore/Private/RenderGraphBuilder.cpp
@@ -542,7 +542,7 @@ bool FRDGBuilder::IsTransientInternal(FRDGViewableResource* Resource, bool bFast
 }
 
 FRDGBuilder::FRDGBuilder(FRHICommandListImmediate& InRHICmdList, FRDGEventName InName, ERDGBuilderFlags InFlags)
-	: FRDGScopeState(InRHICmdList, IsImmediateMode(), ::IsParallelExecuteEnabled() && EnumHasAnyFlags(InFlags, ERDGBuilderFlags::AllowParallelExecute))
+	: FRDGScopeState(InRHICmdList, IsImmediateMode(), ::IsParallelExecuteEnabled() && EnumHasAnyFlags(InFlags, ERDGBuilderFlags::ParallelExecute))
 	, RootAllocatorScope(Allocators.Root)
 	, Blackboard(Allocators.Root)
 	, BuilderName(InName)
@@ -561,8 +561,9 @@ FRDGBuilder::FRDGBuilder(FRHICommandListImmediate& InRHICmdList, FRDGEventName I
 
 	ProloguePass = SetupEmptyPass(Passes.Allocate<FRDGSentinelPass>(Allocators.Root, RDG_EVENT_NAME("Graph Prologue (Graphics)")));
 
-	ParallelExecute.bEnabled = ::IsParallelExecuteEnabled() && EnumHasAnyFlags(InFlags, ERDGBuilderFlags::AllowParallelExecute);
-	ParallelSetup.bEnabled   = ::IsParallelSetupEnabled()   && EnumHasAnyFlags(InFlags, ERDGBuilderFlags::AllowParallelExecute);
+	ParallelExecute.bEnabled = ::IsParallelExecuteEnabled() && EnumHasAnyFlags(InFlags, ERDGBuilderFlags::ParallelExecute);
+	ParallelSetup.bEnabled   = ::IsParallelSetupEnabled()   && EnumHasAnyFlags(InFlags, ERDGBuilderFlags::ParallelSetup);
+	bParallelCompileEnabled  = ::IsParallelSetupEnabled()   && EnumHasAnyFlags(InFlags, ERDGBuilderFlags::ParallelCompile);
 
 	if (TransientResourceAllocator)
 	{
@@ -1621,12 +1622,24 @@ void FRDGBuilder::FlushSetupQueue()
 
 ///////////////////////////////////////////////////////////////////////////////////////////////////
 
-void FRDGBuilder::WaitForParallelSetupTasks()
+void FRDGBuilder::WaitForParallelSetupTasks(ERDGSetupTaskWaitPoint WaitPoint)
 {
-	if (!ParallelSetup.Tasks.IsEmpty())
+	const auto WaitForTasksLambda = [this] (ERDGSetupTaskWaitPoint WaitPoint)
 	{
-		UE::Tasks::Wait(ParallelSetup.Tasks);
-		ParallelSetup.Tasks.Reset();
+		if (auto& Tasks = ParallelSetup.Tasks[(int32)WaitPoint]; !Tasks.IsEmpty())
+		{
+			UE::Tasks::Wait(Tasks);
+			Tasks.Reset();
+		}
+	};
+
+	switch (WaitPoint)
+	{
+	case ERDGSetupTaskWaitPoint::Execute:
+		WaitForTasksLambda(ERDGSetupTaskWaitPoint::Execute);
+		[[fallthrough]]; // Also flush any compile tasks that might have been added after the compile wait point.
+	case ERDGSetupTaskWaitPoint::Compile:
+		WaitForTasksLambda(ERDGSetupTaskWaitPoint::Compile);
 	}
 }
 
@@ -1663,7 +1676,7 @@ void FRDGBuilder::Execute()
 	if (!IsImmediateMode())
 	{
 		BeginFlushResourcesRHI();
-		WaitForParallelSetupTasks();
+		WaitForParallelSetupTasks(ERDGSetupTaskWaitPoint::Compile);
 
 		if (ParallelSetup.bEnabled)
 		{
@@ -1711,7 +1724,7 @@ void FRDGBuilder::Execute()
 			}
 			NumElementsCallbackBuffers.Empty();
 
-		}, TaskPriority);
+		}, TaskPriority, bParallelCompileEnabled);
 
 		UE::Tasks::FTask PrepareCollectResourcesTask = AddSetupTask([this]
 		{
@@ -1747,7 +1760,7 @@ void FRDGBuilder::Execute()
 				}
 			});
 
-		}, TaskPriority);
+		}, TaskPriority, bParallelCompileEnabled);
 
 		UE::Tasks::FTaskEvent AllocateUploadBuffersTask{ UE_SOURCE_LOCATION };
 
@@ -1755,7 +1768,7 @@ void FRDGBuilder::Execute()
 		{
 			SubmitBufferUploads(RHICmdListTask, &AllocateUploadBuffersTask);
 
-		}, BufferNumElementsCallbacksTask, TaskPriority);
+		}, BufferNumElementsCallbacksTask, TaskPriority, bParallelCompileEnabled);
 
 		Compile();
 
@@ -1764,7 +1777,7 @@ void FRDGBuilder::Execute()
 			CompilePassBarriers();
 			CollectPassBarriers();
 
-		}, TaskPriority);
+		}, TaskPriority, bParallelCompileEnabled);
 
 		if (ParallelExecute.bEnabled)
 		{
@@ -1781,7 +1794,7 @@ void FRDGBuilder::Execute()
 			SCOPE_CYCLE_COUNTER(STAT_RDG_CollectResourcesTime);
 			CSV_SCOPED_TIMING_STAT_EXCLUSIVE(RDG_CollectResources);
 			SCOPED_NAMED_EVENT_TEXT("FRDGBuilder::CollectResources", FColor::Magenta);
-			
+
 			PrepareCollectResourcesTask.Wait();
 
 			EnumerateExtendedLifetimeResources(Textures, [](FRDGTexture* Texture)
@@ -1843,13 +1856,13 @@ void FRDGBuilder::Execute()
 			{
 				AllocatePooledBuffers(RHICmdListTask, PooledBuffers);
 
-			}, AllocateUploadBuffersTask, TaskPriority);
+			}, AllocateUploadBuffersTask, TaskPriority, bParallelCompileEnabled);
 
 			AllocatePooledTexturesTask = AddCommandListSetupTask([this, PooledTextures = MoveTemp(CollectResourceContext.PooledTextures)] (FRHICommandListBase& RHICmdListTask)
 			{
 				AllocatePooledTextures(RHICmdListTask, PooledTextures);
 
-			}, TaskPriority);
+			}, TaskPriority, bParallelCompileEnabled);
 
 			AllocateTransientResources(MoveTemp(CollectResourceContext.TransientResources));
 
@@ -1857,13 +1870,13 @@ void FRDGBuilder::Execute()
 			{
 				FinalizeResources();
 
-			}, MakeArrayView<UE::Tasks::FTask>({ CollectPassBarriersTask, AllocatePooledBuffersTask, AllocatePooledTexturesTask }), TaskPriority);
+			}, MakeArrayView<UE::Tasks::FTask>({ CollectPassBarriersTask, AllocatePooledBuffersTask, AllocatePooledTexturesTask }), TaskPriority, bParallelCompileEnabled);
 
 			CreateViewsTask = AddCommandListSetupTask([this, InViews = MoveTemp(CollectResourceContext.Views)] (FRHICommandListBase& RHICmdListTask)
 			{
 				CreateViews(RHICmdListTask, InViews);
 
-			}, MakeArrayView<UE::Tasks::FTask>({ AllocatePooledBuffersTask, AllocatePooledTexturesTask, SubmitBufferUploadsTask}), TaskPriority);
+			}, MakeArrayView<UE::Tasks::FTask>({ AllocatePooledBuffersTask, AllocatePooledTexturesTask, SubmitBufferUploadsTask}), TaskPriority, bParallelCompileEnabled);
 
 			if (TransientResourceAllocator)
 			{
@@ -1879,7 +1892,7 @@ void FRDGBuilder::Execute()
 		{
 			CreateUniformBuffers(InUniformBuffers);
 
-		}, CreateViewsTask, TaskPriority); // Uniform buffer creation require views to be valid.
+		}, CreateViewsTask, TaskPriority, bParallelCompileEnabled); // Uniform buffer creation require views to be valid.
 
 		AllocatePooledBuffersTask.Wait();
 		AllocatePooledTexturesTask.Wait();
@@ -1891,7 +1904,7 @@ void FRDGBuilder::Execute()
 	}
 
 	EndFlushResourcesRHI();
-	WaitForParallelSetupTasks();
+	WaitForParallelSetupTasks(ERDGSetupTaskWaitPoint::Execute);
 
 	if (ParallelExecute.DispatchTaskEvent)
 	{
@@ -2416,7 +2429,6 @@ void FRDGBuilder::SetupAuxiliaryPasses(FRDGPass* Pass)
 		}
 
 		check(!EnumHasAnyFlags(Pass->Pipeline, ERHIPipeline::AsyncCompute));
-		check(ParallelSetup.Tasks.IsEmpty());
 
 		FCollectResourceContext Context;
 		SubmitBufferUploads(RHICmdList);
diff --git a/Engine/Source/Runtime/RenderCore/Private/RenderGraphPrivate.h b/Engine/Source/Runtime/RenderCore/Private/RenderGraphPrivate.h
--- a/Engine/Source/Runtime/RenderCore/Private/RenderGraphPrivate.h
+++ b/Engine/Source/Runtime/RenderCore/Private/RenderGraphPrivate.h
@@ -7,6 +7,7 @@
 #include "ProfilingDebugging/CountersTrace.h"
 #include "ProfilingDebugging/CsvProfiler.h"
 #include "RHI.h"
+#include "DataDrivenShaderPlatformInfo.h"
 
 DEFINE_LOG_CATEGORY_STATIC(LogRDG, Log, All);
 
@@ -245,7 +246,7 @@ FORCEINLINE bool IsImmediateMode()
 
 FORCEINLINE bool IsRenderPassMergeEnabled()
 {
-	return GRDGMergeRenderPasses != 0 && !IsImmediateMode();
+	return GRDGMergeRenderPasses != 0 && !IsImmediateMode() && RHIHasTiledGPU(GMaxRHIShaderPlatform);
 }
 
 FORCEINLINE bool IsAsyncComputeSupported()
diff --git a/Engine/Source/Runtime/RenderCore/Public/RenderGraphBuilder.h b/Engine/Source/Runtime/RenderCore/Public/RenderGraphBuilder.h
--- a/Engine/Source/Runtime/RenderCore/Public/RenderGraphBuilder.h
+++ b/Engine/Source/Runtime/RenderCore/Public/RenderGraphBuilder.h
@@ -177,6 +177,10 @@ public:
 	template <typename PODType>
 	PODType* AllocPODArray(uint32 Count);
 
+	/** Allocates POD memory using an allocator tied to the lifetime of the graph. Does not construct / destruct. */
+	template <typename PODType>
+	TArrayView<PODType> AllocPODArrayView(uint32 Count);
+
 	/** Allocates a C++ object using an allocator tied to the lifetime of the graph. Will destruct the object. */
 	template <typename ObjectType, typename... TArgs>
 	ObjectType* AllocObject(TArgs&&... Args);
@@ -260,35 +264,35 @@ public:
 
 	/** Launches a task that is synced prior to graph execution. If parallel execution is not enabled, the lambda is run immediately. */
 	template <typename TaskLambda>
-	UE::Tasks::FTask AddSetupTask(TaskLambda&& Task, bool bCondition = true);
+	UE::Tasks::FTask AddSetupTask(TaskLambda&& Task, bool bCondition = true, ERDGSetupTaskWaitPoint WaitPoint = ERDGSetupTaskWaitPoint::Compile);
 
 	template <typename TaskLambda>
-	UE::Tasks::FTask AddSetupTask(TaskLambda&& Task, UE::Tasks::ETaskPriority Priority, bool bCondition = true);
+	UE::Tasks::FTask AddSetupTask(TaskLambda&& Task, UE::Tasks::ETaskPriority Priority, bool bCondition = true, ERDGSetupTaskWaitPoint WaitPoint = ERDGSetupTaskWaitPoint::Compile);
 
 	template <typename TaskLambda>
-	UE::Tasks::FTask AddSetupTask(TaskLambda&& Task, UE::Tasks::FPipe* Pipe, UE::Tasks::ETaskPriority Priority = UE::Tasks::ETaskPriority::Normal, bool bCondition = true);
+	UE::Tasks::FTask AddSetupTask(TaskLambda&& Task, UE::Tasks::FPipe* Pipe, UE::Tasks::ETaskPriority Priority = UE::Tasks::ETaskPriority::Normal, bool bCondition = true, ERDGSetupTaskWaitPoint WaitPoint = ERDGSetupTaskWaitPoint::Compile);
 
 	template <typename TaskLambda, typename PrerequisitesCollectionType>
-	UE::Tasks::FTask AddSetupTask(TaskLambda&& Task, PrerequisitesCollectionType&& Prerequisites, UE::Tasks::ETaskPriority Priority = UE::Tasks::ETaskPriority::Normal, bool bCondition = true);
+	UE::Tasks::FTask AddSetupTask(TaskLambda&& Task, PrerequisitesCollectionType&& Prerequisites, UE::Tasks::ETaskPriority Priority = UE::Tasks::ETaskPriority::Normal, bool bCondition = true, ERDGSetupTaskWaitPoint WaitPoint = ERDGSetupTaskWaitPoint::Compile);
 
 	template <typename TaskLambda, typename PrerequisitesCollectionType>
-	UE::Tasks::FTask AddSetupTask(TaskLambda&& Task, UE::Tasks::FPipe* Pipe, PrerequisitesCollectionType&& Prerequisites, UE::Tasks::ETaskPriority Priority = UE::Tasks::ETaskPriority::Normal, bool bCondition = true);
+	UE::Tasks::FTask AddSetupTask(TaskLambda&& Task, UE::Tasks::FPipe* Pipe, PrerequisitesCollectionType&& Prerequisites, UE::Tasks::ETaskPriority Priority = UE::Tasks::ETaskPriority::Normal, bool bCondition = true, ERDGSetupTaskWaitPoint WaitPoint = ERDGSetupTaskWaitPoint::Compile);
 
 	/** Launches a task that is synced prior to graph execution. If parallel execution is not enabled, the lambda is run immediately. */
 	template <typename TaskLambda>
-	UE::Tasks::FTask AddCommandListSetupTask(TaskLambda&& Task, bool bCondition = true);
+	UE::Tasks::FTask AddCommandListSetupTask(TaskLambda&& Task, bool bCondition = true, ERDGSetupTaskWaitPoint WaitPoint = ERDGSetupTaskWaitPoint::Compile);
 
 	template <typename TaskLambda>
-	UE::Tasks::FTask AddCommandListSetupTask(TaskLambda&& Task, UE::Tasks::ETaskPriority Priority, bool bCondition = true);
+	UE::Tasks::FTask AddCommandListSetupTask(TaskLambda&& Task, UE::Tasks::ETaskPriority Priority, bool bCondition = true, ERDGSetupTaskWaitPoint WaitPoint = ERDGSetupTaskWaitPoint::Compile);
 
 	template <typename TaskLambda>
-	UE::Tasks::FTask AddCommandListSetupTask(TaskLambda&& Task, UE::Tasks::FPipe* Pipe, UE::Tasks::ETaskPriority Priority = UE::Tasks::ETaskPriority::Normal, bool bCondition = true);
+	UE::Tasks::FTask AddCommandListSetupTask(TaskLambda&& Task, UE::Tasks::FPipe* Pipe, UE::Tasks::ETaskPriority Priority = UE::Tasks::ETaskPriority::Normal, bool bCondition = true, ERDGSetupTaskWaitPoint WaitPoint = ERDGSetupTaskWaitPoint::Compile);
 
 	template <typename TaskLambda, typename PrerequisitesCollectionType>
-	UE::Tasks::FTask AddCommandListSetupTask(TaskLambda&& Task, PrerequisitesCollectionType&& Prerequisites, UE::Tasks::ETaskPriority Priority = UE::Tasks::ETaskPriority::Normal, bool bCondition = true);
+	UE::Tasks::FTask AddCommandListSetupTask(TaskLambda&& Task, PrerequisitesCollectionType&& Prerequisites, UE::Tasks::ETaskPriority Priority = UE::Tasks::ETaskPriority::Normal, bool bCondition = true, ERDGSetupTaskWaitPoint WaitPoint = ERDGSetupTaskWaitPoint::Compile);
 
 	template <typename TaskLambda, typename PrerequisitesCollectionType>
-	UE::Tasks::FTask AddCommandListSetupTask(TaskLambda&& Task, UE::Tasks::FPipe* Pipe, PrerequisitesCollectionType&& Prerequisites, UE::Tasks::ETaskPriority Priority = UE::Tasks::ETaskPriority::Normal, bool bCondition = true);
+	UE::Tasks::FTask AddCommandListSetupTask(TaskLambda&& Task, UE::Tasks::FPipe* Pipe, PrerequisitesCollectionType&& Prerequisites, UE::Tasks::ETaskPriority Priority = UE::Tasks::ETaskPriority::Normal, bool bCondition = true, ERDGSetupTaskWaitPoint WaitPoint = ERDGSetupTaskWaitPoint::Compile);
 
 	/** Whether RDG will launch async tasks when AddSetup{CommandList}Task is called. */
 	inline bool IsParallelSetupEnabled() const
@@ -931,17 +935,19 @@ private:
 	struct
 	{
 		/** Array of all tasks for variants of AddSetupTask. */
-		TArray<UE::Tasks::FTask, FRDGArrayAllocator> Tasks;
+		TStaticArray<TArray<UE::Tasks::FTask, FRDGArrayAllocator>, (int32)ERDGSetupTaskWaitPoint::MAX> Tasks;
 
 		bool bEnabled = false;
 
 	} ParallelSetup;
 
-	void WaitForParallelSetupTasks();
+	void WaitForParallelSetupTasks(ERDGSetupTaskWaitPoint WaitPoint);
 
 	/////////////////////////////////////////////////////////////////////////////
 	// Parallel Execution
 
+	bool bParallelCompileEnabled = false;
+
 	struct
 	{
 		TArray<FParallelPassSet, FRDGArrayAllocator> ParallelPassSets;
diff --git a/Engine/Source/Runtime/RenderCore/Public/RenderGraphBuilder.inl b/Engine/Source/Runtime/RenderCore/Public/RenderGraphBuilder.inl
--- a/Engine/Source/Runtime/RenderCore/Public/RenderGraphBuilder.inl
+++ b/Engine/Source/Runtime/RenderCore/Public/RenderGraphBuilder.inl
@@ -144,6 +144,12 @@ FORCEINLINE PODType* FRDGBuilder::AllocPODArray(uint32 Count)
 	return Allocators.Root.AllocUninitialized<PODType>(Count);
 }
 
+template <typename PODType>
+TArrayView<PODType> FRDGBuilder::AllocPODArrayView(uint32 Count)
+{
+	return TArrayView<PODType>(AllocPODArray<PODType>(Count), Count);
+}
+
 template <typename ObjectType, typename... TArgs>
 FORCEINLINE ObjectType* FRDGBuilder::AllocObject(TArgs&&... Args)
 {
@@ -455,15 +461,15 @@ inline void FRDGBuilder::AddDispatchHint()
 }
 
 template <typename TaskLambdaType>
-FORCEINLINE UE::Tasks::FTask FRDGBuilder::AddSetupTask(TaskLambdaType&& TaskLambda, bool bCondition)
+FORCEINLINE UE::Tasks::FTask FRDGBuilder::AddSetupTask(TaskLambdaType&& TaskLambda, bool bCondition, ERDGSetupTaskWaitPoint WaitPoint)
 {
-	return AddSetupTask(MoveTemp(TaskLambda), nullptr, TArray<UE::Tasks::FTask>{}, UE::Tasks::ETaskPriority::Normal, bCondition);
+	return AddSetupTask(MoveTemp(TaskLambda), nullptr, TArray<UE::Tasks::FTask>{}, UE::Tasks::ETaskPriority::Normal, bCondition, WaitPoint);
 }
 
 template <typename TaskLambdaType>
-FORCEINLINE UE::Tasks::FTask FRDGBuilder::AddSetupTask(TaskLambdaType&& TaskLambda, UE::Tasks::ETaskPriority Priority, bool bCondition)
+FORCEINLINE UE::Tasks::FTask FRDGBuilder::AddSetupTask(TaskLambdaType&& TaskLambda, UE::Tasks::ETaskPriority Priority, bool bCondition, ERDGSetupTaskWaitPoint WaitPoint)
 {
-	return AddSetupTask(MoveTemp(TaskLambda), nullptr, TArray<UE::Tasks::FTask>{}, Priority, bCondition);
+	return AddSetupTask(MoveTemp(TaskLambda), nullptr, TArray<UE::Tasks::FTask>{}, Priority, bCondition, WaitPoint);
 }
 
 template <typename TaskLambdaType>
@@ -471,7 +477,8 @@ FORCEINLINE UE::Tasks::FTask FRDGBuilder::AddSetupTask(
 	TaskLambdaType&& TaskLambda,
 	UE::Tasks::FPipe* Pipe,
 	UE::Tasks::ETaskPriority Priority,
-	bool bCondition)
+	bool bCondition,
+	ERDGSetupTaskWaitPoint WaitPoint)
 {
 	return AddSetupTask(MoveTemp(TaskLambda), Pipe, TArray<UE::Tasks::FTask>{}, Priority, bCondition);
 }
@@ -481,7 +488,8 @@ FORCEINLINE UE::Tasks::FTask FRDGBuilder::AddSetupTask(
 	TaskLambdaType&& TaskLambda,
 	PrerequisitesCollectionType&& Prerequisites,
 	UE::Tasks::ETaskPriority Priority,
-	bool bCondition)
+	bool bCondition,
+	ERDGSetupTaskWaitPoint WaitPoint)
 {
 	return AddSetupTask(MoveTemp(TaskLambda), nullptr, Forward<PrerequisitesCollectionType&&>(Prerequisites), Priority, bCondition);
 }
@@ -529,7 +537,8 @@ UE::Tasks::FTask FRDGBuilder::AddSetupTask(
 	UE::Tasks::FPipe* Pipe,
 	PrerequisitesCollectionType&& Prerequisites,
 	UE::Tasks::ETaskPriority Priority,
-	bool bCondition)
+	bool bCondition,
+	ERDGSetupTaskWaitPoint WaitPoint)
 {
 	UE::Tasks::FTask Task;
 
@@ -561,22 +570,22 @@ UE::Tasks::FTask FRDGBuilder::AddSetupTask(
 
 	if (Task.IsValid())
 	{
-		ParallelSetup.Tasks.Emplace(Task);
+		ParallelSetup.Tasks[(int32)WaitPoint].Emplace(Task);
 	}
 
 	return Task;
 }
 
 template <typename TaskLambdaType>
-FORCEINLINE UE::Tasks::FTask FRDGBuilder::AddCommandListSetupTask(TaskLambdaType&& TaskLambda, bool bCondition)
+FORCEINLINE UE::Tasks::FTask FRDGBuilder::AddCommandListSetupTask(TaskLambdaType&& TaskLambda, bool bCondition, ERDGSetupTaskWaitPoint WaitPoint)
 {
-	return AddCommandListSetupTask(MoveTemp(TaskLambda), nullptr, TArray<UE::Tasks::FTask>{}, UE::Tasks::ETaskPriority::Normal, bCondition);
+	return AddCommandListSetupTask(MoveTemp(TaskLambda), nullptr, TArray<UE::Tasks::FTask>{}, UE::Tasks::ETaskPriority::Normal, bCondition, WaitPoint);
 }
 
 template <typename TaskLambdaType>
-FORCEINLINE UE::Tasks::FTask FRDGBuilder::AddCommandListSetupTask(TaskLambdaType&& TaskLambda, UE::Tasks::ETaskPriority TaskPriority, bool bCondition)
+FORCEINLINE UE::Tasks::FTask FRDGBuilder::AddCommandListSetupTask(TaskLambdaType&& TaskLambda, UE::Tasks::ETaskPriority TaskPriority, bool bCondition, ERDGSetupTaskWaitPoint WaitPoint)
 {
-	return AddCommandListSetupTask(MoveTemp(TaskLambda), nullptr, TArray<UE::Tasks::FTask>{}, TaskPriority, bCondition);
+	return AddCommandListSetupTask(MoveTemp(TaskLambda), nullptr, TArray<UE::Tasks::FTask>{}, TaskPriority, bCondition, WaitPoint);
 }
 
 template <typename TaskLambdaType>
@@ -584,9 +593,10 @@ FORCEINLINE UE::Tasks::FTask FRDGBuilder::AddCommandListSetupTask(
 	TaskLambdaType&& TaskLambda,
 	UE::Tasks::FPipe* Pipe,
 	UE::Tasks::ETaskPriority Priority,
-	bool bCondition)
+	bool bCondition,
+	ERDGSetupTaskWaitPoint WaitPoint)
 {
-	return AddCommandListSetupTask(MoveTemp(TaskLambda), Pipe, TArray<UE::Tasks::FTask>{}, Priority, bCondition);
+	return AddCommandListSetupTask(MoveTemp(TaskLambda), Pipe, TArray<UE::Tasks::FTask>{}, Priority, bCondition, WaitPoint);
 }
 
 template <typename TaskLambdaType, typename PrerequisitesCollectionType>
@@ -594,9 +604,10 @@ FORCEINLINE UE::Tasks::FTask FRDGBuilder::AddCommandListSetupTask(
 	TaskLambdaType&& TaskLambda,
 	PrerequisitesCollectionType&& Prerequisites,
 	UE::Tasks::ETaskPriority Priority,
-	bool bCondition)
+	bool bCondition,
+	ERDGSetupTaskWaitPoint WaitPoint)
 {
-	return AddCommandListSetupTask(MoveTemp(TaskLambda), nullptr, Forward<PrerequisitesCollectionType&&>(Prerequisites), Priority, bCondition);
+	return AddCommandListSetupTask(MoveTemp(TaskLambda), nullptr, Forward<PrerequisitesCollectionType&&>(Prerequisites), Priority, bCondition, WaitPoint);
 }
 
 template <typename TaskLambdaType, typename PrerequisitesCollectionType>
@@ -605,7 +616,8 @@ UE::Tasks::FTask FRDGBuilder::AddCommandListSetupTask(
 	UE::Tasks::FPipe* Pipe,
 	PrerequisitesCollectionType&& Prerequisites,
 	UE::Tasks::ETaskPriority Priority,
-	bool bCondition)
+	bool bCondition,
+	ERDGSetupTaskWaitPoint WaitPoint)
 {
 	UE::Tasks::FTask Task;
 
@@ -660,7 +672,7 @@ UE::Tasks::FTask FRDGBuilder::AddCommandListSetupTask(
 
 	if (Task.IsValid())
 	{
-		ParallelSetup.Tasks.Emplace(Task);
+		ParallelSetup.Tasks[(int32)WaitPoint].Emplace(Task);
 	}
 
 	return Task;
diff --git a/Engine/Source/Runtime/RenderCore/Public/RenderGraphDefinitions.h b/Engine/Source/Runtime/RenderCore/Public/RenderGraphDefinitions.h
--- a/Engine/Source/Runtime/RenderCore/Public/RenderGraphDefinitions.h
+++ b/Engine/Source/Runtime/RenderCore/Public/RenderGraphDefinitions.h
@@ -67,8 +67,18 @@ enum class ERDGBuilderFlags
 {
 	None = 0,
 
+	/** Allows the builder to parallelize AddSetupPass calls. Without this flag, setup passes run serially. */
+	ParallelSetup = 1 << 0,
+	
+	/** Allows the builder to parallelize compilation of the graph. Without this flag, all passes execute on the render thread. */
+	ParallelCompile = 1 << 1,
+
 	/** Allows the builder to parallelize execution of passes. Without this flag, all passes execute on the render thread. */
-	AllowParallelExecute = 1 << 0
+	ParallelExecute = 1 << 2,
+
+	Parallel = ParallelSetup | ParallelCompile | ParallelExecute,
+
+	AllowParallelExecute UE_DEPRECATED(5.5, "Use ERDDGBuilderFlags::Parallel instead.") = Parallel,
 };
 ENUM_CLASS_FLAGS(ERDGBuilderFlags);
 
@@ -154,6 +164,17 @@ enum class ERDGTextureFlags : uint8
 };
 ENUM_CLASS_FLAGS(ERDGTextureFlags);
 
+enum class ERDGSetupTaskWaitPoint : uint8
+{
+	/** (Default) Setup task is synced prior to compilation. Use this mode if task mutates RDG resources (e.g. RDG buffer upload contents, buffer size callbacks, etc) */
+	Compile = 0,
+
+	/** Setup task is synced prior to execution. Use this mode if your task is stalling in RDG and doesn't affect RDG compilation in any way. */
+	Execute = 1,
+
+	MAX
+};
+
 /** Flags to annotate a view with when calling CreateUAV. */
 enum class ERDGUnorderedAccessViewFlags : uint8
 {
diff --git a/Engine/Source/Runtime/Renderer/Private/MeshEdgesRendering.cpp b/Engine/Source/Runtime/Renderer/Private/MeshEdgesRendering.cpp
--- a/Engine/Source/Runtime/Renderer/Private/MeshEdgesRendering.cpp
+++ b/Engine/Source/Runtime/Renderer/Private/MeshEdgesRendering.cpp
@@ -281,7 +281,7 @@ void RenderMeshEdges_RenderThread(
 	// update any resources that needed a deferred update
 	FDeferredUpdateResource::UpdateResources(RHICmdList);
 
-	FRDGBuilder GraphBuilder(RHICmdList, RDG_EVENT_NAME("MeshEdges"), ERDGBuilderFlags::AllowParallelExecute);
+	FRDGBuilder GraphBuilder(RHICmdList, RDG_EVENT_NAME("MeshEdges"), ERDGBuilderFlags::Parallel);
 	{
 		RDG_EVENT_SCOPE(GraphBuilder, "RenderMeshEdges_RenderThread");
 
diff --git a/Engine/Source/Runtime/Renderer/Private/PlanarReflectionRendering.cpp b/Engine/Source/Runtime/Renderer/Private/PlanarReflectionRendering.cpp
--- a/Engine/Source/Runtime/Renderer/Private/PlanarReflectionRendering.cpp
+++ b/Engine/Source/Runtime/Renderer/Private/PlanarReflectionRendering.cpp
@@ -351,7 +351,7 @@ static void UpdatePlanarReflectionContents_RenderThread(
 	FDeferredUpdateResource::UpdateResources(RHICmdList);
 
 	const ERHIFeatureLevel::Type FeatureLevel = SceneRenderer->FeatureLevel;
-	FRDGBuilder GraphBuilder(RHICmdList, RDG_EVENT_NAME("PlanarReflection"), ERDGBuilderFlags::AllowParallelExecute);
+	FRDGBuilder GraphBuilder(RHICmdList, RDG_EVENT_NAME("PlanarReflection"), ERDGBuilderFlags::Parallel);
 
 	// Make sure we render to the same set of GPUs as the main scene renderer.
 	if (MainSceneRenderer->ViewFamily.RenderTarget != nullptr)
diff --git a/Engine/Source/Runtime/Renderer/Private/ReflectionEnvironmentCapture.cpp b/Engine/Source/Runtime/Renderer/Private/ReflectionEnvironmentCapture.cpp
--- a/Engine/Source/Runtime/Renderer/Private/ReflectionEnvironmentCapture.cpp
+++ b/Engine/Source/Runtime/Renderer/Private/ReflectionEnvironmentCapture.cpp
@@ -519,7 +519,7 @@ void CaptureSceneToScratchCubemap(
 	// update any resources that needed a deferred update
 	FDeferredUpdateResource::UpdateResources(RHICmdList);
 
-	FRDGBuilder GraphBuilder(RHICmdList, RDG_EVENT_NAME("CubeMapCapture"), ERDGBuilderFlags::AllowParallelExecute);
+	FRDGBuilder GraphBuilder(RHICmdList, RDG_EVENT_NAME("CubeMapCapture"), ERDGBuilderFlags::Parallel);
 
 	{
 		RDG_EVENT_SCOPE(GraphBuilder, "CubeMapCapture");
diff --git a/Engine/Source/Runtime/Renderer/Private/SceneCaptureRendering.cpp b/Engine/Source/Runtime/Renderer/Private/SceneCaptureRendering.cpp
--- a/Engine/Source/Runtime/Renderer/Private/SceneCaptureRendering.cpp
+++ b/Engine/Source/Runtime/Renderer/Private/SceneCaptureRendering.cpp
@@ -387,10 +387,10 @@ static void UpdateSceneCaptureContentDeferred_RenderThread(
 
 #if WANTS_DRAW_MESH_EVENTS
 	SCOPED_DRAW_EVENTF(RHICmdList, SceneCapture, TEXT("SceneCapture %s"), EventName);
-	FRDGBuilder GraphBuilder(RHICmdList, RDG_EVENT_NAME("SceneCapture %s", *EventName), ERDGBuilderFlags::AllowParallelExecute);
+	FRDGBuilder GraphBuilder(RHICmdList, RDG_EVENT_NAME("SceneCapture %s", *EventName), ERDGBuilderFlags::Parallel);
 #else
 	SCOPED_DRAW_EVENT(RHICmdList, UpdateSceneCaptureContent_RenderThread);
-	FRDGBuilder GraphBuilder(RHICmdList, RDG_EVENT_NAME("SceneCapture"), ERDGBuilderFlags::AllowParallelExecute);
+	FRDGBuilder GraphBuilder(RHICmdList, RDG_EVENT_NAME("SceneCapture"), ERDGBuilderFlags::Parallel);
 #endif
 
 	{
diff --git a/Engine/Source/Runtime/Renderer/Private/SceneRendering.cpp b/Engine/Source/Runtime/Renderer/Private/SceneRendering.cpp
--- a/Engine/Source/Runtime/Renderer/Private/SceneRendering.cpp
+++ b/Engine/Source/Runtime/Renderer/Private/SceneRendering.cpp
@@ -5161,7 +5161,7 @@ static void RenderViewFamilies_RenderThread(FRHICommandListImmediate& RHICmdList
 				ViewFamily.EngineShowFlags.HitProxies ? TEXT("RenderHitProxies") : TEXT("Render"),
 				ViewFamily.bResolveScene ? TEXT("Primary") : TEXT("Auxiliary")
 			),
-			ERDGBuilderFlags::AllowParallelExecute
+			ERDGBuilderFlags::Parallel
 		);
 
 #if WITH_GPUDEBUGCRASH
