Fix Shader SymbolsInfo Causing Non-Determinism and Defeating Deduplication
Info is now stored in shader map in editor only data alongside platform data
Includes shader bump

#tests local tests on Lyra cooking and DDC-Verify

diff --git a/Engine/Plugins/Compositing/OpenColorIO/Source/OpenColorIO/Private/OpenColorIOShaderMap.cpp b/Engine/Plugins/Compositing/OpenColorIO/Source/OpenColorIO/Private/OpenColorIOShaderMap.cpp
--- a/Engine/Plugins/Compositing/OpenColorIO/Source/OpenColorIO/Private/OpenColorIOShaderMap.cpp
+++ b/Engine/Plugins/Compositing/OpenColorIO/Source/OpenColorIO/Private/OpenColorIOShaderMap.cpp
@@ -433,7 +433,7 @@ FShader* FOpenColorIOShaderMap::ProcessCompilationResultsForSingleJob(const TRef
 	FShaderCompileJob* CurrentJob = SingleJob->GetSingleShaderJob();
 	check(CurrentJob->Id == CompilingId);
 
-	GetResourceCode()->AddShaderCompilerOutput(CurrentJob->Output, CurrentJob->Key.ToString());
+	GetResourceCode()->AddShaderCompilerOutput(CurrentJob->Output, CurrentJob->Key.ToString(), CurrentJob->Input.GenerateDebugInfo());
 
 	FShader* Shader = nullptr;
 
diff --git a/Engine/Plugins/FX/Niagara/Source/NiagaraShader/Private/NiagaraShader.cpp b/Engine/Plugins/FX/Niagara/Source/NiagaraShader/Private/NiagaraShader.cpp
--- a/Engine/Plugins/FX/Niagara/Source/NiagaraShader/Private/NiagaraShader.cpp
+++ b/Engine/Plugins/FX/Niagara/Source/NiagaraShader/Private/NiagaraShader.cpp
@@ -777,7 +777,7 @@ FShader* FNiagaraShaderMap::ProcessCompilationResultsForSingleJob(const TRefCoun
 	auto CurrentJob = SingleJob->GetSingleShaderJob();
 	check(CurrentJob->Id == CompilingId);
 
-	GetResourceCode()->AddShaderCompilerOutput(CurrentJob->Output, CurrentJob->Key.ToString());
+	GetResourceCode()->AddShaderCompilerOutput(CurrentJob->Output, CurrentJob->Key.ToString(), SingleJob->GetSingleShaderJob()->Input.GenerateDebugInfo());
 
 	FShader* Shader = nullptr;
 
diff --git a/Engine/Plugins/Runtime/ComputeFramework/Source/ComputeFramework/Private/ComputeKernelShaderMap.cpp b/Engine/Plugins/Runtime/ComputeFramework/Source/ComputeFramework/Private/ComputeKernelShaderMap.cpp
--- a/Engine/Plugins/Runtime/ComputeFramework/Source/ComputeFramework/Private/ComputeKernelShaderMap.cpp
+++ b/Engine/Plugins/Runtime/ComputeFramework/Source/ComputeFramework/Private/ComputeKernelShaderMap.cpp
@@ -432,7 +432,7 @@ FShader* FComputeKernelShaderMap::ProcessCompilationResultsForSingleJob(FShaderC
 {
 	check(CurrentJob.Id == CompilingId);
 
-	GetResourceCode()->AddShaderCompilerOutput(CurrentJob.Output, CurrentJob.Key.ToString());
+	GetResourceCode()->AddShaderCompilerOutput(CurrentJob.Output, CurrentJob.Key.ToString(), CurrentJob.Input.GenerateDebugInfo());
 
 	FShader* Shader = nullptr;
 
diff --git a/Engine/Shaders/Public/ShaderVersion.ush b/Engine/Shaders/Public/ShaderVersion.ush
--- a/Engine/Shaders/Public/ShaderVersion.ush
+++ b/Engine/Shaders/Public/ShaderVersion.ush
@@ -3,4 +3,4 @@
 // Each time the console command is executed it generates a new GUID. As this file is included
 // in Platform.ush (which should be included in any shader) it allows to invalidate the shader DDC.
 // 
-#pragma message("UESHADERMETADATA_VERSION 033f833d-eb12-485d-a2fc-d8ed752a806d")
+#pragma message("UESHADERMETADATA_VERSION D6EB0A35-C442-4C75-89FD-C4D7D8E06556")
diff --git a/Engine/Source/Developer/ShaderCompilerCommon/Public/ShaderSymbolExport.h b/Engine/Source/Developer/ShaderCompilerCommon/Public/ShaderSymbolExport.h
--- a/Engine/Source/Developer/ShaderCompilerCommon/Public/ShaderSymbolExport.h
+++ b/Engine/Source/Developer/ShaderCompilerCommon/Public/ShaderSymbolExport.h
@@ -25,7 +25,7 @@ public:
 	*   Template type is the platform specific symbol data structure.
 	*/
 	template<typename TPlatformShaderSymbolData>
-	void NotifyShaderCompiled(const TConstArrayView<uint8>& PlatformSymbolData);
+	void NotifyShaderCompiled(const TConstArrayView<uint8>& PlatformSymbolData, const FString& DebugInfo = FString());
 
 	/** Called at the end of a cook to free resources and finalize artifacts created during the cook. */
 	void NotifyShaderCompilersShutdown();
@@ -61,7 +61,7 @@ private:
 };
 
 template<typename TPlatformShaderSymbolData>
-inline void FShaderSymbolExport::NotifyShaderCompiled(const TConstArrayView<uint8>& PlatformSymbolData)
+inline void FShaderSymbolExport::NotifyShaderCompiled(const TConstArrayView<uint8>& PlatformSymbolData, const FString& DebugInfo)
 {
 	static bool bFirst = true;
 	if (bFirst)
@@ -82,7 +82,6 @@ inline void FShaderSymbolExport::NotifyShaderCompiled(const TConstArrayView<uint
 		for (const auto& SymbolData : FullSymbolData.GetAllSymbolData())
 		{
 			const FString FileName = SymbolData.GetFilename();
-			const FString DebugInfo = SymbolData.GetDebugInfo();
 			TConstArrayView<uint8> Contents = SymbolData.GetContents();
 
 			WriteSymbolData(FileName, DebugInfo, Contents);
diff --git a/Engine/Source/Developer/TargetPlatform/Public/Interfaces/IShaderFormat.h b/Engine/Source/Developer/TargetPlatform/Public/Interfaces/IShaderFormat.h
--- a/Engine/Source/Developer/TargetPlatform/Public/Interfaces/IShaderFormat.h
+++ b/Engine/Source/Developer/TargetPlatform/Public/Interfaces/IShaderFormat.h
@@ -108,7 +108,7 @@ public:
 	 * Called when a shader resource is cooked, so the shader format can perform platform-specific operations on the debug data.
 	 * Does nothing on platforms that make no use of the platform debug data.
 	 */
-	virtual void NotifyShaderCompiled(const TConstArrayView<uint8>& PlatformDebugData, FName Format) const { }
+	virtual void NotifyShaderCompiled(const TConstArrayView<uint8>& PlatformDebugData, FName Format, const FString& DebugInfo = FString()) const { }
 
 	/** Called at the end of a cook to free resources and finalize artifacts created during the cook. */
 	virtual void NotifyShaderCompilersShutdown(FName Format) const { }
diff --git a/Engine/Source/Developer/Windows/ShaderFormatD3D/Private/D3DShaderCompilerDXC.cpp b/Engine/Source/Developer/Windows/ShaderFormatD3D/Private/D3DShaderCompilerDXC.cpp
--- a/Engine/Source/Developer/Windows/ShaderFormatD3D/Private/D3DShaderCompilerDXC.cpp
+++ b/Engine/Source/Developer/Windows/ShaderFormatD3D/Private/D3DShaderCompilerDXC.cpp
@@ -926,7 +926,6 @@ bool CompileAndProcessD3DShaderDXC(
 	{
 		FD3DSM6ShaderDebugData DebugData;
 		DebugData.Name = PdbName;
-		DebugData.DebugInfo = Input.GenerateDebugInfo();
 
 		// We don't export the PDB contents here because it would result in duplicate data,
 		// as we use embedded PDBs. Once we are able to use external PDBs, the PDB contents
diff --git a/Engine/Source/Developer/Windows/ShaderFormatD3D/Private/ShaderFormatD3D.cpp b/Engine/Source/Developer/Windows/ShaderFormatD3D/Private/ShaderFormatD3D.cpp
--- a/Engine/Source/Developer/Windows/ShaderFormatD3D/Private/ShaderFormatD3D.cpp
+++ b/Engine/Source/Developer/Windows/ShaderFormatD3D/Private/ShaderFormatD3D.cpp
@@ -168,11 +168,11 @@ public:
 	}
 
 #if WITH_ENGINE
-	virtual void NotifyShaderCompiled(const TConstArrayView<uint8>& PlatformDebugData, FName Format) const override
+	virtual void NotifyShaderCompiled(const TConstArrayView<uint8>& PlatformDebugData, FName Format, const FString& DebugInfo) const override
 	{
 		if (Format == NAME_PCD3D_SM6)
 		{
-			ShaderSymbolExportSM6.NotifyShaderCompiled<FD3DSM6ShaderDebugData>(PlatformDebugData);
+			ShaderSymbolExportSM6.NotifyShaderCompiled<FD3DSM6ShaderDebugData>(PlatformDebugData, DebugInfo);
 		}
 	}
 	virtual void NotifyShaderCompilersShutdown(FName Format) const override
diff --git a/Engine/Source/Developer/Windows/ShaderFormatD3D/Private/ShaderFormatD3D.h b/Engine/Source/Developer/Windows/ShaderFormatD3D/Private/ShaderFormatD3D.h
--- a/Engine/Source/Developer/Windows/ShaderFormatD3D/Private/ShaderFormatD3D.h
+++ b/Engine/Source/Developer/Windows/ShaderFormatD3D/Private/ShaderFormatD3D.h
@@ -69,7 +69,6 @@ struct FD3DSM6ShaderDebugData
 	inline friend FArchive& operator<<(FArchive& Ar, FD3DSM6ShaderDebugData& DebugData)
 	{
 		Ar << DebugData.Name;
-		Ar << DebugData.DebugInfo;
 		Ar << DebugData.Contents;
 		return Ar;
 	}
@@ -84,11 +83,6 @@ struct FD3DSM6ShaderDebugData
 		return Name;
 	}
 
-	inline FString GetDebugInfo() const
-	{
-		return DebugInfo;
-	}
-
 	TConstArrayView<FD3DSM6ShaderDebugData> GetAllSymbolData() const
 	{
 		return TConstArrayView<FD3DSM6ShaderDebugData>(this, 1);
diff --git a/Engine/Source/Runtime/Engine/Private/Materials/MaterialShader.cpp b/Engine/Source/Runtime/Engine/Private/Materials/MaterialShader.cpp
--- a/Engine/Source/Runtime/Engine/Private/Materials/MaterialShader.cpp
+++ b/Engine/Source/Runtime/Engine/Private/Materials/MaterialShader.cpp
@@ -2706,7 +2706,7 @@ FShader* FMaterialShaderMap::ProcessCompilationResultsForSingleJob(FShaderCompil
 		*LexToString(CurrentJob.Input.Target.GetPlatform()),
 		*LexToString(GetShaderPlatform()));
 
-	GetResourceCode()->AddShaderCompilerOutput(CurrentJob.Output, CurrentJob.Key.ToString());
+	GetResourceCode()->AddShaderCompilerOutput(CurrentJob.Output, CurrentJob.Key.ToString(), SingleJob->Input.GenerateDebugInfo());
 
 #if ALLOW_SHADERMAP_DEBUG_DATA
 	CompileTime += SingleJob->Output.CompileTime;
diff --git a/Engine/Source/Runtime/Engine/Private/ShaderCompiler/ShaderCompilerEditor.cpp b/Engine/Source/Runtime/Engine/Private/ShaderCompiler/ShaderCompilerEditor.cpp
--- a/Engine/Source/Runtime/Engine/Private/ShaderCompiler/ShaderCompilerEditor.cpp
+++ b/Engine/Source/Runtime/Engine/Private/ShaderCompiler/ShaderCompilerEditor.cpp
@@ -148,7 +148,7 @@ FShader* FGlobalShaderTypeCompiler::FinishCompileShader(const FGlobalShaderType*
 		EShaderPlatform Platform = CurrentJob.Input.Target.GetPlatform();
 		FGlobalShaderMapSection* Section = GGlobalShaderMap[Platform]->FindOrAddSection(ShaderType);
 
-		Section->GetResourceCode()->AddShaderCompilerOutput(CurrentJob.Output, CurrentJob.Key.ToString());
+		Section->GetResourceCode()->AddShaderCompilerOutput(CurrentJob.Output, CurrentJob.Key.ToString(), CurrentJob.Input.GenerateDebugInfo());
 
 		if (ShaderPipelineType && !ShaderPipelineType->ShouldOptimizeUnusedOutputs(CurrentJob.Input.Target.GetPlatform()))
 		{
diff --git a/Engine/Source/Runtime/RenderCore/Private/ShaderResource.cpp b/Engine/Source/Runtime/RenderCore/Private/ShaderResource.cpp
--- a/Engine/Source/Runtime/RenderCore/Private/ShaderResource.cpp
+++ b/Engine/Source/Runtime/RenderCore/Private/ShaderResource.cpp
@@ -240,7 +240,7 @@ int32 FShaderMapResourceCode::FindShaderIndex(const FSHAHash& InHash) const
 	return Algo::BinarySearch(ShaderHashes, InHash);
 }
 
-void FShaderMapResourceCode::AddShaderCompilerOutput(const FShaderCompilerOutput& Output, const FString& DebugName)
+void FShaderMapResourceCode::AddShaderCompilerOutput(const FShaderCompilerOutput& Output, const FString& DebugName, FString DebugInfo)
 {
 	TRACE_CPUPROFILER_EVENT_SCOPE(FShaderMapResourceCode::AddShaderCode);
 
@@ -253,7 +253,7 @@ void FShaderMapResourceCode::AddShaderCompilerOutput(const FShaderCompilerOutput
 
 #if WITH_EDITORONLY_DATA
 		// Output.Errors contains warnings in the case any exist (no errors since if there were the job would have failed)
-		AddEditorOnlyData(Index, DebugName, Output.PlatformDebugData, Output.Errors, Output.ShaderStatistics);
+		AddEditorOnlyData(Index, DebugName, Output.PlatformDebugData, Output.Errors, Output.ShaderStatistics, DebugInfo);
 #endif
 
 		FShaderEntry& Entry = ShaderEntries.InsertDefaulted_GetRef(Index);
@@ -321,16 +321,16 @@ void FShaderMapResourceCode::AddShaderCompilerOutput(const FShaderCompilerOutput
 	else
 	{
 		// Output.Errors contains warnings in the case any exist (no errors since if there were the job would have failed)
-		// We append the warnings for any additional jobs which resulted in the same bytecode for the sake of determinism in the
-		// results saved to DDC. 
-		AppendWarningsToEditorOnlyData(Index, DebugName, Output.Errors);
+		// We append the warnings and deduplicate other data like DebugInfo for any additional jobs which resulted in the
+		// same bytecode for the sake of determinism in the results saved to DDC. 
+		UpdateEditorOnlyData(Index, DebugName, Output.Errors, DebugInfo);
 		ValidateShaderStatisticsEditorOnlyData(Index, Output.ShaderStatistics);
 	}
 #endif
 }
 
 #if WITH_EDITORONLY_DATA
-void FShaderMapResourceCode::AddEditorOnlyData(int32 Index, const FString& DebugName, TConstArrayView<uint8> InPlatformDebugData, TConstArrayView<FShaderCompilerError> InCompilerWarnings, const TArray<FGenericShaderStat>& ShaderStatistics)
+void FShaderMapResourceCode::AddEditorOnlyData(int32 Index, const FString& DebugName, TConstArrayView<uint8> InPlatformDebugData, TConstArrayView<FShaderCompilerError> InCompilerWarnings, const TArray<FGenericShaderStat>& ShaderStatistics, const FString& DebugInfo)
 {
 	FShaderEditorOnlyDataEntry& Entry = ShaderEditorOnlyDataEntries.InsertDefaulted_GetRef(Index);
 	Entry.PlatformDebugData = InPlatformDebugData;
@@ -339,12 +339,19 @@ void FShaderMapResourceCode::AddEditorOnlyData(int32 Index, const FString& Debug
 	check(Entry.ShaderStatistics.Num() == 0);
 	Entry.ShaderStatistics = ShaderStatistics;
 
-	AppendWarningsToEditorOnlyData(Index, DebugName, InCompilerWarnings);
+	UpdateEditorOnlyData(Index, DebugName, InCompilerWarnings, DebugInfo);
 }
 
-void FShaderMapResourceCode::AppendWarningsToEditorOnlyData(int32 Index, const FString& DebugName, TConstArrayView<FShaderCompilerError> InCompilerWarnings)
+void FShaderMapResourceCode::UpdateEditorOnlyData(int32 Index, const FString& DebugName, TConstArrayView<FShaderCompilerError> InCompilerWarnings, const FString& DebugInfo)
 {
 	FShaderEditorOnlyDataEntry& Entry = ShaderEditorOnlyDataEntries[Index];
+
+	// Keep a single DebugInfo as it doesn't matter which one we use, but make sure it is the same one for determinism
+	if (!DebugInfo.IsEmpty() && (Entry.DebugInfo.IsEmpty() || (DebugInfo < Entry.DebugInfo)))
+	{
+		Entry.DebugInfo = DebugInfo;
+	}
+
 	for (const FShaderCompilerError& Warning : InCompilerWarnings)
 	{
 		FString ModifiedWarning = !DebugName.IsEmpty() ? FString::Printf(TEXT("%s [%s]"), *Warning.GetErrorString(), *DebugName) : Warning.GetErrorString();
@@ -444,7 +451,7 @@ void FShaderMapResourceCode::NotifyShadersCompiled(FName FormatName)
 		{
 			for (const FShaderEditorOnlyDataEntry& Entry : ShaderEditorOnlyDataEntries)
 			{
-				ShaderFormat->NotifyShaderCompiled(Entry.PlatformDebugData, FormatName);
+				ShaderFormat->NotifyShaderCompiled(Entry.PlatformDebugData, FormatName, Entry.DebugInfo);
 			}
 		}
 	}
diff --git a/Engine/Source/Runtime/RenderCore/Public/Shader.h b/Engine/Source/Runtime/RenderCore/Public/Shader.h
--- a/Engine/Source/Runtime/RenderCore/Public/Shader.h
+++ b/Engine/Source/Runtime/RenderCore/Public/Shader.h
@@ -498,6 +498,12 @@ public:
 	struct FShaderEditorOnlyDataEntry
 	{
 		TArray<uint8> PlatformDebugData;
+
+		/** This field contains a debug string stored in a ShaderSymbols.info file when using r.Shaders.SymbolInfo=1
+		*   Used to facilitate reverse lookup of a shader from a platform shader hash when no full shader symbol 
+		*   information is available */
+		FString DebugInfo;
+
 		/** A (deduplicated/sorted) array of all the compiler warnings that were emitted when all shaders resulting 
 		 *  in the associated bytecode were compiled (i.e. if multiple shader sources have warnings but compile to
 		 *  the same code, all warnings for each unique source will be reported).
@@ -509,7 +515,7 @@ public:
 
 		friend FArchive& operator<<(FArchive& Ar, FShaderEditorOnlyDataEntry& Entry)
 		{
-			return Ar << Entry.PlatformDebugData << Entry.CompilerWarnings << Entry.ShaderStatistics;
+			return Ar << Entry.PlatformDebugData << Entry.DebugInfo << Entry.CompilerWarnings << Entry.ShaderStatistics;
 		}
 	};
 #endif // WITH_EDITORONLY_DATA
@@ -527,13 +533,14 @@ public:
 
 	RENDERCORE_API uint32 GetSizeBytes() const;
 
-	RENDERCORE_API void AddShaderCompilerOutput(const FShaderCompilerOutput& Output, const FString& DebugName = FString());
+	RENDERCORE_API void AddShaderCompilerOutput(const FShaderCompilerOutput& Output, const FString& DebugName = FString(), FString DebugInfo = FString());
 
 	int32 FindShaderIndex(const FSHAHash& InHash) const;
 
 #if WITH_EDITORONLY_DATA
-	void AddEditorOnlyData(int32 Index, const FString& DebugName, TConstArrayView<uint8> InPlatformDebugData, TConstArrayView<FShaderCompilerError> InCompilerWarnings, const TArray<FGenericShaderStat>& ShaderStatistics);
-	void AppendWarningsToEditorOnlyData(int32 Index, const FString& DebugName, TConstArrayView<FShaderCompilerError> InCompilerWarnings);
+	void AddEditorOnlyData(int32 Index, const FString& DebugName, TConstArrayView<uint8> InPlatformDebugData, TConstArrayView<FShaderCompilerError> InCompilerWarnings, const TArray<FGenericShaderStat>& ShaderStatistics, const FString& DebugInfo = FString());
+	void UpdateEditorOnlyData(int32 Index, const FString& DebugName, TConstArrayView<FShaderCompilerError> InCompilerWarnings, const FString& DebugInfo);
+
 	RENDERCORE_API void LogShaderCompilerWarnings();
 #endif
 
diff --git a/Engine/Source/Runtime/RenderCore/Public/ShaderCompilerCore.h b/Engine/Source/Runtime/RenderCore/Public/ShaderCompilerCore.h
--- a/Engine/Source/Runtime/RenderCore/Public/ShaderCompilerCore.h
+++ b/Engine/Source/Runtime/RenderCore/Public/ShaderCompilerCore.h
@@ -33,14 +33,14 @@ typedef TSharedPtr<TArray<ANSICHAR>, ESPMode::ThreadSafe> FShaderSharedAnsiStrin
 // this is for the protocol, not the data, bump if FShaderCompilerInput/FShaderPreprocessOutput serialization, SerializeWorkerInput or ProcessInputFromArchive changes.
 inline const int32 ShaderCompileWorkerInputVersion = 28;
 // this is for the protocol, not the data, bump if FShaderCompilerOutput or WriteToOutputArchive changes.
-inline const int32 ShaderCompileWorkerOutputVersion = 21;
+inline const int32 ShaderCompileWorkerOutputVersion = 22;
 // this is for the protocol, not the data.
 inline const int32 ShaderCompileWorkerSingleJobHeader = 'S';
 // this is for the protocol, not the data.
 inline const int32 ShaderCompileWorkerPipelineJobHeader = 'P';
 
 // modify this for changes to the FShaderCompilerOutput data structure (in addition to ShaderCompileWorkerOutputVersion)
-inline const int32 FShaderCompilerOutputStructVersion = 3;
+inline const int32 FShaderCompilerOutputStructVersion = 4;
 
 namespace UE::ShaderCompiler
 {
@@ -363,7 +363,7 @@ struct FShaderCompilerInput
 	// Currently uses DebugGroupName, but can be updated to contain other important information as needed
 	FString GenerateDebugInfo() const
 	{
-		return DebugGroupName;
+		return Environment.CompilerFlags.Contains(CFLAG_GenerateSymbolsInfo) ? DebugGroupName : FString();
 	}
 
 	UE_DEPRECATED(5.5, "GatherSharedInputs no longer processes external includes, as these are not needed by SCW anymore")
