[CF] Fix missed errors/warnings due to deduplication of line marker info within shader compiler
[PCG] Fix line numbers for errors/warnings and add them to the error summary when available.

#rnx

diff --git a/Engine/Plugins/PCG/Source/PCG/Private/Compute/Elements/PCGComputeGraphElement.cpp b/Engine/Plugins/PCG/Source/PCG/Private/Compute/Elements/PCGComputeGraphElement.cpp
--- a/Engine/Plugins/PCG/Source/PCG/Private/Compute/Elements/PCGComputeGraphElement.cpp
+++ b/Engine/Plugins/PCG/Source/PCG/Private/Compute/Elements/PCGComputeGraphElement.cpp
@@ -300,7 +300,18 @@ void FPCGComputeGraphElement::LogCompilationMessages(FPCGComputeGraphContext* In
 						FPCGStack StackWithNode = *InContext->Stack;
 						StackWithNode.PushFrame(NodeAndCompileMessages.Get<0>().ResolveObjectPtr());
 
-						Subsystem->GetNodeVisualLogsMutable().Log(StackWithNode, Verbosity, FText::FromString(Message.Text));
+						if (Message.Line != INDEX_NONE)
+						{
+							Subsystem->GetNodeVisualLogsMutable().Log(StackWithNode, Verbosity, FText::Format(
+								LOCTEXT("ErrorWithLineFormat", "[{0},{1}] {2}"),
+								Message.Line,
+								Message.ColumnStart,
+								FText::FromString(Message.Text)));
+						}
+						else
+						{
+							Subsystem->GetNodeVisualLogsMutable().Log(StackWithNode, Verbosity, FText::FromString(Message.Text));
+						}
 					}
 				}
 			}
diff --git a/Engine/Plugins/PCG/Source/PCG/Private/Compute/Elements/PCGCustomHLSL.cpp b/Engine/Plugins/PCG/Source/PCG/Private/Compute/Elements/PCGCustomHLSL.cpp
--- a/Engine/Plugins/PCG/Source/PCG/Private/Compute/Elements/PCGCustomHLSL.cpp
+++ b/Engine/Plugins/PCG/Source/PCG/Private/Compute/Elements/PCGCustomHLSL.cpp
@@ -1008,7 +1008,7 @@ FString UPCGCustomHLSLSettings::GetCookedKernelSource(const TMap<FPCGKernelAttri
 	const FIntVector GroupSize = GetThreadGroupSize();
 
 	// FIXME: Create source range mappings so that we can go from error location to our source.
-	FString Source = (TEXT("    // User kernel\n") + ShaderSource);
+	FString Source = ShaderSource;
 	FString Functions = (TEXT("// User kernel functions\n") + ShaderFunctions);
 
 #if PLATFORM_WINDOWS
@@ -1242,7 +1242,7 @@ FString UPCGCustomHLSLSettings::GetCookedKernelSource(const TMap<FPCGKernelAttri
 		Source.ReplaceInline(TEXT("KERNEL"), TEXT("void __kernel_func(uint ThreadIndex)"), ESearchCase::CaseSensitive);
 
 		Result = FString::Printf(TEXT(
-			"#line 1 \"%s\"\n" // ShaderPathName
+			"#line 0 \"%s\"\n" // ShaderPathName
 			"%s\n" // Includes
 			"%s\n" // Functions
 			"%s\n" // Source
@@ -1259,7 +1259,7 @@ FString UPCGCustomHLSLSettings::GetCookedKernelSource(const TMap<FPCGKernelAttri
 			"    const uint ThreadIndex = %s;\n" // UnWrappedDispatchThreadId
 			"%s\n" // HeaderWriters
 			"%s\n" // KernelSpecificPreamble
-			"#line 1 \"%s\"\n\n" // ShaderPathName
+			"#line 0 \"%s\"\n\n" // ShaderPathName
 			"%s\n" // Source
 			"}\n"),
 			*Includes, *Functions, *KernelFunc, *UnWrappedDispatchThreadId, *HeaderWriters, *KernelSpecificPreamble, *ShaderPathName, *Source);
diff --git a/Engine/Plugins/Runtime/ComputeFramework/Source/ComputeFramework/Private/ComputeKernelShaderCompilationManager.cpp b/Engine/Plugins/Runtime/ComputeFramework/Source/ComputeFramework/Private/ComputeKernelShaderCompilationManager.cpp
--- a/Engine/Plugins/Runtime/ComputeFramework/Source/ComputeFramework/Private/ComputeKernelShaderCompilationManager.cpp
+++ b/Engine/Plugins/Runtime/ComputeFramework/Source/ComputeFramework/Private/ComputeKernelShaderCompilationManager.cpp
@@ -108,7 +108,6 @@ void FComputeKernelShaderCompilationManager::ProcessAsyncResults()
 static bool ParseShaderCompilerError(FShaderCompilerError const& InError, FComputeKernelCompileMessage& OutMessage)
 {
 	FShaderCompilerError Error = InError;
-	Error.ExtractSourceLocation();
 
 	// We ignore error messages that don't have a line information.
 	FString Line, Column;
@@ -248,10 +247,13 @@ void FComputeKernelShaderCompilationManager::ProcessCompiledComputeKernelShaderM
 
 				if (GShowComputeKernelShaderWarnings || !CurrentJob.bSucceeded)
 				{
-					for (int32 ErrorIndex = 0; ErrorIndex < CurrentJob.Output.Errors.Num(); ErrorIndex++)
+					TArray<FShaderCompilerError> Errors = CurrentJob.Output.Errors;
+					FShaderCompilerError::ExtractSourceLocations(Errors);
+
+					for (int32 ErrorIndex = 0; ErrorIndex < Errors.Num(); ErrorIndex++)
 					{
 						FComputeKernelCompileMessage Message;
-						if (ParseShaderCompilerError(CurrentJob.Output.Errors[ErrorIndex], Message))
+						if (ParseShaderCompilerError(Errors[ErrorIndex], Message))
 						{
 							ProcessedCompileResults.Messages.AddUnique(Message);
 						}
diff --git a/Engine/Source/Runtime/RenderCore/Private/ShaderCore.cpp b/Engine/Source/Runtime/RenderCore/Private/ShaderCore.cpp
--- a/Engine/Source/Runtime/RenderCore/Private/ShaderCore.cpp
+++ b/Engine/Source/Runtime/RenderCore/Private/ShaderCore.cpp
@@ -3633,6 +3633,49 @@ bool FShaderCompilerError::ExtractSourceLocation()
 	return false;
 }
 
+void FShaderCompilerError::ExtractSourceLocations(TArray<FShaderCompilerError>& InOutErrors)
+{
+	FString CurrentLine, CurrentColumn; // Local to loop but hoisted for performance.
+	FString PreviousLine, PreviousColumn;
+
+	for (int ErrorIndex = 0; ErrorIndex < InOutErrors.Num(); ++ErrorIndex)
+	{
+		FShaderCompilerError& CurrentError = InOutErrors[ErrorIndex];
+
+		CurrentError.ExtractSourceLocation();
+
+		if (!CurrentError.ErrorLineString.Split(TEXT(","), &CurrentLine, &CurrentColumn))
+		{
+			PreviousLine.Reset();
+			PreviousColumn.Reset();
+			continue;
+		}
+
+		if (!CurrentLine.IsNumeric() || !CurrentColumn.IsNumeric())
+		{
+			PreviousLine.Reset();
+			PreviousColumn.Reset();
+			continue;
+		}
+
+		// The shader compiler may omit line marker info after the first error for that line/column. Copy this information from the previous error
+		// if the line/column matches.
+		if (!CurrentError.HasLineMarker() && ErrorIndex > 0 && PreviousLine == CurrentLine && PreviousColumn == CurrentColumn)
+		{
+			FShaderCompilerError& PreviousError = InOutErrors[ErrorIndex - 1];
+			if (PreviousError.HasLineMarker())
+			{
+				// Issue pertains to same code. Copy marker.
+				CurrentError.HighlightedLine = PreviousError.HighlightedLine;
+				CurrentError.HighlightedLineMarker = PreviousError.HighlightedLineMarker;
+			}
+		}
+
+		PreviousLine = MoveTemp(CurrentLine);
+		PreviousColumn = MoveTemp(CurrentColumn);
+	}
+}
+
 FString FShaderCompilerError::GetShaderSourceFilePath() const
 {
 	// Always return error file path as-is if it doesn't denote a virtual path.
diff --git a/Engine/Source/Runtime/RenderCore/Public/ShaderCompilerCore.h b/Engine/Source/Runtime/RenderCore/Public/ShaderCompilerCore.h
--- a/Engine/Source/Runtime/RenderCore/Public/ShaderCompilerCore.h
+++ b/Engine/Source/Runtime/RenderCore/Public/ShaderCompilerCore.h
@@ -523,6 +523,10 @@ struct FShaderCompilerError
 	/** Extracts the file path and source line from StrippedErrorMessage to ErrorVirtualFilePath and ErrorLineString. */
 	bool RENDERCORE_API ExtractSourceLocation();
 
+	/** Extracts the file path and source line for each error from the error message. Propagates highlighted line and marker to
+	 * all errors pertaining to the same source location. */
+	static void RENDERCORE_API ExtractSourceLocations(TArray<FShaderCompilerError>& InOutErrors);
+
 	/** Returns the path of the underlying source file relative to the process base dir. */
 	FString RENDERCORE_API GetShaderSourceFilePath() const;
 
