Disable render pass merging if r.RDG.Debug.FlushGPU is enabled.

The flush debug mode makes each pass use a different command list, and we can't currently have render passes which start on one command list and end on another.

#rnx

diff --git a/Engine/Source/Runtime/RenderCore/Private/RenderGraphPrivate.h b/Engine/Source/Runtime/RenderCore/Private/RenderGraphPrivate.h
--- a/Engine/Source/Runtime/RenderCore/Private/RenderGraphPrivate.h
+++ b/Engine/Source/Runtime/RenderCore/Private/RenderGraphPrivate.h
@@ -246,7 +246,7 @@ FORCEINLINE bool IsImmediateMode()
 
 FORCEINLINE bool IsRenderPassMergeEnabled()
 {
-	return GRDGMergeRenderPasses != 0 && !IsImmediateMode() && RHIHasTiledGPU(GMaxRHIShaderPlatform);
+	return GRDGMergeRenderPasses != 0 && !IsImmediateMode() && RHIHasTiledGPU(GMaxRHIShaderPlatform) && !GRDGDebugFlushGPU;
 }
 
 FORCEINLINE bool IsAsyncComputeSupported()
