Changing SetExternalPooledBufferRHI to take a TRefCountPtr to reduce reference count traffic.

diff --git a/Engine/Source/Runtime/RenderCore/Private/RenderGraphBuilder.cpp b/Engine/Source/Runtime/RenderCore/Private/RenderGraphBuilder.cpp
--- a/Engine/Source/Runtime/RenderCore/Private/RenderGraphBuilder.cpp
+++ b/Engine/Source/Runtime/RenderCore/Private/RenderGraphBuilder.cpp
@@ -4219,7 +4219,7 @@ void FRDGBuilder::SetTransientTextureRHI(FRDGTexture* Texture, FRHITransientText
 	SetDiscardPass(Texture, TransientTexture);
 }
 
-void FRDGBuilder::SetExternalPooledBufferRHI(FRDGBuffer* Buffer, FRDGPooledBuffer* PooledBuffer)
+void FRDGBuilder::SetExternalPooledBufferRHI(FRDGBuffer* Buffer, const TRefCountPtr<FRDGPooledBuffer>& PooledBuffer)
 {
 	SetPooledBufferRHI(Buffer, PooledBuffer);
 	Buffer->Allocation = PooledBuffer;
diff --git a/Engine/Source/Runtime/RenderCore/Public/RenderGraphBuilder.h b/Engine/Source/Runtime/RenderCore/Public/RenderGraphBuilder.h
--- a/Engine/Source/Runtime/RenderCore/Public/RenderGraphBuilder.h
+++ b/Engine/Source/Runtime/RenderCore/Public/RenderGraphBuilder.h
@@ -693,7 +693,7 @@ private:
 	void SetPooledTextureRHI(FRDGTexture* Texture, FRDGPooledTexture* PooledTexture);
 	void SetTransientTextureRHI(FRDGTexture* Texture, FRHITransientTexture* TransientTexture);
 	void SetDiscardPass(FRDGTexture* Texture, FRHITransientTexture* TransientTexture);
-	void SetExternalPooledBufferRHI(FRDGBuffer* Buffer, FRDGPooledBuffer* PooledBuffer);
+	void SetExternalPooledBufferRHI(FRDGBuffer* Buffer, const TRefCountPtr<FRDGPooledBuffer>& PooledBuffer);
 	void SetPooledBufferRHI(FRDGBuffer* Buffer, FRDGPooledBuffer* PooledBuffer);
 	void SetTransientBufferRHI(FRDGBuffer* Buffer, FRHITransientBuffer* TransientBuffer);
 
