[shaders]
- move population of output.target field into core code so each backend doesn't need to do it manually (do so before calling the compile function so any existing code expecting the output field to be set at any point during the compile process is unaffected)
- add a check in FShaderCompileJob::SerializeOutput and FShaderCompileJob::SerializeWorkerOutput that the FShaderTarget for a job output matches that of its input. This should catch cases of shaders with the wrong frequency being associated with jobs; further the callstack should indicate where this incorrect association is coming from (since SerializeOutput is called in different places for each cache path: duplicate in-flight jobs, jobs which hit in the in-memory job cache, and jobs which hit in the DDC cache, and SerializeWorkerOutput is only called when the job is read back from SCW output).

diff --git a/Engine/Source/Developer/Apple/MetalShaderFormat/Private/MetalCompileShaderMSC.cpp b/Engine/Source/Developer/Apple/MetalShaderFormat/Private/MetalCompileShaderMSC.cpp
--- a/Engine/Source/Developer/Apple/MetalShaderFormat/Private/MetalCompileShaderMSC.cpp
+++ b/Engine/Source/Developer/Apple/MetalShaderFormat/Private/MetalCompileShaderMSC.cpp
@@ -802,7 +802,6 @@ void FMetalCompileShaderMSC::DoCompileMetalShader(
 
 	if (Result != 0)
 	{
-		Output.Target = Input.Target;
 		BuildMetalShaderOutput(Output, Input, MetalSource.c_str(), MetalSource.length(), CRCLen, CRC, VersionEnum, *Standard, *MinOSVersion, Output.Errors, OutputData.TypedBuffers, OutputData.InvariantBuffers, OutputData.TypedUAVs, OutputData.ConstantBuffers, bAllowFastIntrinsics
 		  , NumCBVs, OutputSizeVS, MaxInputPrimitivesPerMeshThreadgroupGS, bUsesDiscard, ReflectionJSON, MetalBytecode
 		);
diff --git a/Engine/Source/Developer/Apple/MetalShaderFormat/Private/MetalCompileShaderSPIRV.cpp b/Engine/Source/Developer/Apple/MetalShaderFormat/Private/MetalCompileShaderSPIRV.cpp
--- a/Engine/Source/Developer/Apple/MetalShaderFormat/Private/MetalCompileShaderSPIRV.cpp
+++ b/Engine/Source/Developer/Apple/MetalShaderFormat/Private/MetalCompileShaderSPIRV.cpp
@@ -1297,7 +1297,6 @@ void FMetalCompileShaderSPIRV::DoCompileMetalShader(
 
 	if (Result != 0)
 	{
-		Output.Target = Input.Target;
 		BuildMetalShaderOutput(Output, Input, MetalSource.c_str(), MetalSource.length(), CRCLen, CRC, VersionEnum, *Standard, *MinOSVersion, Output.Errors, OutputData.TypedBuffers, OutputData.InvariantBuffers, OutputData.TypedUAVs, OutputData.ConstantBuffers, bAllowFastIntrinsics
 #if UE_METAL_USE_METAL_SHADER_CONVERTER
 		  , 0, 0, 0, false, nullptr, FMetalShaderBytecode()
diff --git a/Engine/Source/Developer/ShaderFormatOpenGL/Private/OpenGLShaderCompiler.cpp b/Engine/Source/Developer/ShaderFormatOpenGL/Private/OpenGLShaderCompiler.cpp
--- a/Engine/Source/Developer/ShaderFormatOpenGL/Private/OpenGLShaderCompiler.cpp
+++ b/Engine/Source/Developer/ShaderFormatOpenGL/Private/OpenGLShaderCompiler.cpp
@@ -967,7 +967,6 @@ static void PrecompileShaderInternal(FShaderCompilerOutput& ShaderOutput, const
 		glGetShaderiv(Shader, GL_COMPILE_STATUS, &CompileStatus);
 		if (CompileStatus == GL_TRUE)
 		{
-			ShaderOutput.Target = ShaderInput.Target;
 			BuildShaderOutputInternal(
 				ShaderOutput,
 				ShaderInput,
@@ -3132,7 +3131,6 @@ void CompileOpenGLShader(const FShaderCompilerInput& Input, const FShaderPreproc
 		PrecompileShaderInternal(Output, Input, GlslShaderSource, Version, HlslFrequency);
 #else // VALIDATE_GLSL_WITH_DRIVER
 		int32 SourceLen = FCStringAnsi::Strlen(GlslShaderSource); //-V595
-		Output.Target = Input.Target;
 		BuildShaderOutputInternal(Output, Input, GlslShaderSource, SourceLen, Version);
 #endif // VALIDATE_GLSL_WITH_DRIVER
 
diff --git a/Engine/Source/Developer/VulkanShaderFormat/Private/SpirVShaderCompiler.inl b/Engine/Source/Developer/VulkanShaderFormat/Private/SpirVShaderCompiler.inl
--- a/Engine/Source/Developer/VulkanShaderFormat/Private/SpirVShaderCompiler.inl
+++ b/Engine/Source/Developer/VulkanShaderFormat/Private/SpirVShaderCompiler.inl
@@ -812,8 +812,6 @@ static bool BuildShaderOutputFromSpirv(
 		AddReflectionInfos(Bindings.AccelerationStructures, VK_DESCRIPTOR_TYPE_ACCELERATION_STRUCTURE_KHR, 0);
 	}
 
-	Output.Target = InternalState.Input.Target;
-
 	// Overwrite updated SPIRV code
 	SerializedOutput.Spirv.Data = TArray<uint32>(Reflection.GetCode(), Reflection.GetCodeSize() / 4);
 
diff --git a/Engine/Source/Developer/Windows/ShaderFormatD3D/Private/D3DShaderCompiler.inl b/Engine/Source/Developer/Windows/ShaderFormatD3D/Private/D3DShaderCompiler.inl
--- a/Engine/Source/Developer/Windows/ShaderFormatD3D/Private/D3DShaderCompiler.inl
+++ b/Engine/Source/Developer/Windows/ShaderFormatD3D/Private/D3DShaderCompiler.inl
@@ -555,9 +555,6 @@ inline void GenerateFinalOutput(
 
 	Output.NumTextureSamplers = PackedResourceCounts.NumSamplers;
 
-	// Pass the target through to the output.
-	Output.Target = Input.Target;
-
 	// SRV Limits
 	{
 		if (CompileData.bBindlessResources)
diff --git a/Engine/Source/Runtime/RenderCore/Private/ShaderCore.cpp b/Engine/Source/Runtime/RenderCore/Private/ShaderCore.cpp
--- a/Engine/Source/Runtime/RenderCore/Private/ShaderCore.cpp
+++ b/Engine/Source/Runtime/RenderCore/Private/ShaderCore.cpp
@@ -1648,9 +1648,11 @@ public:
 		__try
 #endif
 		{
+			Job.Output.Target = Job.Input.Target;
 			if (Job.SecondaryOutput.IsValid())
 			{
 				check(Job.SecondaryPreprocessOutput.IsValid());
+				Job.SecondaryOutput->Target = Job.Input.Target;
 				Compiler->CompilePreprocessedShader(Job.Input, Job.PreprocessOutput, *Job.SecondaryPreprocessOutput, Job.Output, *Job.SecondaryOutput, WorkingDirectory);
 			}
 			else
@@ -4011,6 +4013,8 @@ void FShaderCompileJob::SerializeOutput(FArchive& Ar)
 	// output hash is now serialized as part of the output, as the shader code is compressed in SCWs
 	checkf(!Output.bSucceeded || Output.OutputHash != FSHAHash(), TEXT("Successful compile job does not have an OutputHash generated."));
 
+	checkf(Output.Target == Input.Target, TEXT("Output FShaderTarget does not match the input struct; incorrect results associated with job?"));
+
 	if (Ar.IsLoading())
 	{
 		bFinalized = true;
@@ -4082,6 +4086,8 @@ void FShaderCompileJob::SerializeWorkerOutput(FArchive& Ar)
 
 	Ar << Output;
 
+	checkf(Output.Target == Input.Target, TEXT("Output FShaderTarget does not match the input struct; incorrect results associated with job?"));
+
 	bool bSecondaryOutput = SecondaryOutput.IsValid();
 	Ar << bSecondaryOutput;
 
