Modifying the GameThreadWaitForTask function to report a hang on the render thread instead of firing an assert when a timeout occurs.

diff --git a/Engine/Source/Runtime/RenderCore/Private/RenderingThread.cpp b/Engine/Source/Runtime/RenderCore/Private/RenderingThread.cpp
--- a/Engine/Source/Runtime/RenderCore/Private/RenderingThread.cpp
+++ b/Engine/Source/Runtime/RenderCore/Private/RenderingThread.cpp
@@ -31,6 +31,7 @@
 #include "ProfilingDebugging/CsvProfiler.h"
 #include "Async/TaskTrace.h"
 #include "DataDrivenShaderPlatformInfo.h"
+#include "HAL/ThreadManager.h"
 #include "ProfilingDebugging/CountersTrace.h"
 
 //
@@ -1223,6 +1224,62 @@ static FAutoConsoleVariableRef CVarTimeoutForBlockOnRenderFence(
 	TEXT("Number of milliseconds the game thread should wait before failing when waiting on a render thread fence.")
 );
 
+static void HandleRenderTaskHang(uint32 ThreadThatHung, double HangDuration)
+{
+	// Get the name of the hung thread
+	FString ThreadName = FThreadManager::GetThreadName(ThreadThatHung);
+	if (ThreadName.IsEmpty())
+	{
+		ThreadName = FString::Printf(TEXT("unknown thread (%u)"), ThreadThatHung);
+	}
+	
+#if !PLATFORM_WINDOWS || (PLATFORM_USE_MINIMAL_HANG_DETECTION && 1)
+	UE_LOG(LogRendererCore, Fatal, TEXT("GameThread timed out waiting for %s after %.02f secs"), *ThreadName, HangDuration);
+#else
+	// Capture the stack in the thread that hung
+	static const int32 MaxStackFrames = 100;
+	uint64 StackFrames[MaxStackFrames];
+	int32 NumStackFrames = FPlatformStackWalk::CaptureThreadStackBackTrace(ThreadThatHung, StackFrames, MaxStackFrames);
+
+	// Convert the stack trace to text
+	TArray<FString> StackLines;
+	for (int32 Idx = 0; Idx < NumStackFrames; Idx++)
+	{
+		ANSICHAR Buffer[1024];
+		Buffer[0] = '\0';
+		FPlatformStackWalk::ProgramCounterToHumanReadableString(Idx, StackFrames[Idx], Buffer, sizeof(Buffer));
+		StackLines.Add(Buffer);
+	}
+	
+	// Dump the callstack and the thread name to log
+	FString StackTrimmed;
+	UE_LOG(LogRendererCore, Error, TEXT("GameThread timed out waiting for %s after %.02f seconds:"), *ThreadName, HangDuration);
+	for (int32 Idx = 0; Idx < StackLines.Num(); Idx++)
+	{
+		UE_LOG(LogRendererCore, Error, TEXT("  %s"), *StackLines[Idx]);
+		if (StackTrimmed.Len() < 512)
+		{
+			StackTrimmed += TEXT("  ");
+			StackTrimmed += StackLines[Idx];
+			StackTrimmed += LINE_TERMINATOR;
+		}
+	}
+
+	const FString ErrorMessage = FString::Printf(TEXT("GameThread timed out waiting for %s after %.02f seconds:%s%s%sCheck log for full callstack."),
+		*ThreadName, HangDuration, LINE_TERMINATOR, *StackTrimmed, LINE_TERMINATOR);
+	
+	GLog->Panic();
+	ReportHang(*ErrorMessage, StackFrames, NumStackFrames, ThreadThatHung);
+	if (FApp::CanEverRender())
+	{
+		FPlatformMisc::MessageBoxExt(EAppMsgType::Ok,
+			*NSLOCTEXT("MessageDialog", "ReportHangError_Body", "The application has hung and will now close. We apologize for the inconvenience.").ToString(),
+			*NSLOCTEXT("MessageDialog", "ReportHangError_Title", "Application Hang Detected").ToString());
+	}
+	FPlatformMisc::RequestExit(true, TEXT("GameThreadWaitForTask"));
+#endif
+}
+
 /**
  * Block the game thread waiting for a task to finish on the rendering thread.
  */
@@ -1307,7 +1364,12 @@ static void GameThreadWaitForTask(const UE::Tasks::FTask& Task, bool bEmptyGameT
 				{
 					if (bOverdue && !bDisabled && !IsRenderThreadTimeoutSuspended() && !FPlatformMisc::IsDebuggerPresent())
 					{
-						UE_LOG(LogRendererCore, Fatal, TEXT("GameThread timed out waiting for RenderThread after %.02f secs"), RenderThreadTimeoutClock.Seconds() - StartTime);
+						double HangDuration = RenderThreadTimeoutClock.Seconds() - StartTime;
+						// TODO: Walk the wait chain instead of explicitly setting the render thread as the hung thread id
+						PRAGMA_DISABLE_DEPRECATION_WARNINGS
+						uint32 ThreadThatHung = GRenderThreadId;
+						PRAGMA_ENABLE_DEPRECATION_WARNINGS
+						HandleRenderTaskHang(ThreadThatHung, HangDuration);
 					}
 				}
 #endif
