Add #define to enable mmapping of the metalmap file.

This allows the OS to page in areas of the file on demand.
The OS should be able to page it out when under memory pressure.

requires that the on-disk layout of the metalmap file matches the in-memory form.
for mapping to succeed metalmap files stored within pak files must not be compressed.

diff --git a/Engine/Source/Programs/AutomationTool/Scripts/CopyBuildToStagingDirectory.Automation.cs b/Engine/Source/Programs/AutomationTool/Scripts/CopyBuildToStagingDirectory.Automation.cs
--- a/Engine/Source/Programs/AutomationTool/Scripts/CopyBuildToStagingDirectory.Automation.cs
+++ b/Engine/Source/Programs/AutomationTool/Scripts/CopyBuildToStagingDirectory.Automation.cs
@@ -202,6 +202,23 @@ namespace AutomationScripts
 			}
 		}
 
+		/// <summary>
+		/// returns a hash set containing file extensions that will not be compressed within pak files
+		/// or null if no settings could be found.
+		/// </summary>
+		private static HashSet<string> GetFileExtensionsToSkipCompression(ConfigHierarchy PlatformGameConfig)
+		{
+			if (PlatformGameConfig != null)
+			{
+				List<string> FileExtensionsToSkipCompressionList = null;
+				PlatformGameConfig.GetArray("/Script/UnrealEd.ProjectPackagingSettings", "FileExtensionsToSkipCompressionInPak", out FileExtensionsToSkipCompressionList);
+				HashSet<string> FileExtensionsToSkipCompression = (FileExtensionsToSkipCompressionList is not null) ? new HashSet<string>(FileExtensionsToSkipCompressionList) : null;
+				return FileExtensionsToSkipCompression;
+			}
+
+			return null;
+		}
+
 		/// <summary>
 		/// Writes a pak response file to disk
 		/// </summary>
@@ -211,7 +228,7 @@ namespace AutomationScripts
 		/// <param name="RehydrateAssets"></param>
 		/// <param name="CryptoSettings"></param>
 		/// <param name="bForceFullEncryption"></param>
-		private static void WritePakResponseFile(string Filename, Dictionary<string, string> ResponseFile, bool Compressed, bool RehydrateAssets, EncryptionAndSigning.CryptoSettings CryptoSettings, bool bForceFullEncryption)
+		private static void WritePakResponseFile(string Filename, Dictionary<string, string> ResponseFile, bool Compressed, bool RehydrateAssets, EncryptionAndSigning.CryptoSettings CryptoSettings, bool bForceFullEncryption, HashSet<string> FileExtensionsToSkipCompression)
 		{
 			using (var Writer = new StreamWriter(Filename, false, new System.Text.UTF8Encoding(true)))
 			{
@@ -220,8 +237,10 @@ namespace AutomationScripts
 					string Extension = Path.GetExtension(Entry.Key);
 					string Line = String.Format("\"{0}\" \"{1}\"", Entry.Key, Entry.Value);
 
+					bool bSkipCompression = FileExtensionsToSkipCompression != null && FileExtensionsToSkipCompression.Contains(Path.GetExtension(Entry.Key));
+					
 					// explicitly exclude some file types from compression
-					if (Compressed && !Path.GetExtension(Entry.Key).Contains(".mp4") && !Extension.Contains("ushaderbytecode") && !Path.GetExtension(Entry.Key).Contains("upipelinecache"))
+					if (Compressed && !bSkipCompression && !Path.GetExtension(Entry.Key).Contains(".mp4") && !Extension.Contains("ushaderbytecode") && !Path.GetExtension(Entry.Key).Contains("upipelinecache"))
 					{
 						Line += " -compress";
 					}
@@ -301,7 +320,7 @@ namespace AutomationScripts
 			return CmdLine.ToString();
 		}
 
-		static private string GetPakFileSpecificUnrealPakArguments(Dictionary<string, string> UnrealPakResponseFile, FileReference OutputLocation, string AdditionalOptions, bool Compressed, bool RehydrateAssets,  EncryptionAndSigning.CryptoSettings CryptoSettings, String PatchSourceContentPath, string EncryptionKeyGuid)
+		static private string GetPakFileSpecificUnrealPakArguments(Dictionary<string, string> UnrealPakResponseFile, FileReference OutputLocation, string AdditionalOptions, bool Compressed, bool RehydrateAssets,  EncryptionAndSigning.CryptoSettings CryptoSettings, String PatchSourceContentPath, string EncryptionKeyGuid, HashSet<string> FileExtensionsToSkipCompression)
 		{
 			StringBuilder CmdLine = new StringBuilder(MakePathSafeToUseWithCommandLine(OutputLocation.FullName));
 
@@ -312,7 +331,7 @@ namespace AutomationScripts
 			string ResponseFilesPath = CombinePaths(CmdEnv.EngineSavedFolder, "ResponseFiles");
 			InternalUtils.SafeCreateDirectory(ResponseFilesPath);
 			string UnrealPakResponseFileName = CombinePaths(ResponseFilesPath, "PakList_" + PakName + ".txt");
-			WritePakResponseFile(UnrealPakResponseFileName, UnrealPakResponseFile, Compressed, RehydrateAssets, CryptoSettings, bForceEncryption);
+			WritePakResponseFile(UnrealPakResponseFileName, UnrealPakResponseFile, Compressed, RehydrateAssets, CryptoSettings, bForceEncryption, FileExtensionsToSkipCompression);
 			CmdLine.AppendFormat(" -create={0}", CommandUtils.MakePathSafeToUseWithCommandLine(UnrealPakResponseFileName));
 
 			if (!String.IsNullOrEmpty(PatchSourceContentPath))
@@ -343,9 +362,10 @@ namespace AutomationScripts
 				bool bUnattended)
 		{
 			bool RehydrateAssets = false;
+			HashSet<string> FileExtensionsToSkipCompression = null;
 
 			StringBuilder CmdLine = new StringBuilder();
-			CmdLine.Append(GetPakFileSpecificUnrealPakArguments(UnrealPakResponseFile, OutputLocation, "", Compressed, RehydrateAssets, CryptoSettings, PatchSourceContentPath, EncryptionKeyGuid));
+			CmdLine.Append(GetPakFileSpecificUnrealPakArguments(UnrealPakResponseFile, OutputLocation, "", Compressed, RehydrateAssets, CryptoSettings, PatchSourceContentPath, EncryptionKeyGuid, FileExtensionsToSkipCompression));
 			CmdLine.Append(" ");
 			CmdLine.Append(GetCommonUnrealPakArguments(PakOrderFileLocations, AdditionalOptions, CryptoSettings, CryptoKeysCacheFilename, SecondaryPakOrderFileLocations, bUnattended));
 			return CmdLine.ToString();
@@ -363,7 +383,8 @@ namespace AutomationScripts
 			string EncryptionKeyGuid,
 			string PatchSourceContentPath,
 			bool bGenerateDiffPatch,
-			bool bIsDLC)
+			bool bIsDLC,
+			HashSet<string> FileExtensionsToSkipCompression)
 		{
 			StringBuilder CmdLine = new StringBuilder();
 			CmdLine.AppendFormat("-Output={0}", MakePathSafeToUseWithCommandLine(Path.ChangeExtension(PakOutputLocation.FullName, ".utoc")));
@@ -394,7 +415,7 @@ namespace AutomationScripts
 			string ResponseFilesPath = CombinePaths(CmdEnv.EngineSavedFolder, "ResponseFiles");
 			InternalUtils.SafeCreateDirectory(ResponseFilesPath);
 			string UnrealPakResponseFileName = CombinePaths(ResponseFilesPath, "PakListIoStore_" + ContainerName + ".txt");
-			WritePakResponseFile(UnrealPakResponseFileName, UnrealPakResponseFile, bCompressed, RehydrateAssets, CryptoSettings, bForceEncryption);
+			WritePakResponseFile(UnrealPakResponseFileName, UnrealPakResponseFile, bCompressed, RehydrateAssets, CryptoSettings, bForceEncryption, FileExtensionsToSkipCompression);
 			CmdLine.AppendFormat(" -ResponseFile={0}", CommandUtils.MakePathSafeToUseWithCommandLine(UnrealPakResponseFileName));
 
 			if (CryptoSettings != null && CryptoSettings.bDataCryptoRequired)
@@ -3138,6 +3159,8 @@ namespace AutomationScripts
 			bool bForceUseProjectCompressionFormatIgnoreHardwareOverride = false;// do we want to override HW compression with project? if so, read it below
 			PlatformGameConfig.GetBool("/Script/UnrealEd.ProjectPackagingSettings", "bForceUseProjectCompressionFormatIgnoreHardwareOverride", out bForceUseProjectCompressionFormatIgnoreHardwareOverride);
 
+			HashSet<string> FileExtensionsToSkipCompressionInPaks = GetFileExtensionsToSkipCompression(PlatformGameConfig);
+
 			string HardwareCompressionFormat = null;
 			if ( ! bForceUseProjectCompressionFormatIgnoreHardwareOverride )
 			{
@@ -3771,7 +3794,8 @@ namespace AutomationScripts
 								Params.SkipEncryption ? "" : PakParams.EncryptionKeyGuid,
 								ContainerPatchSourcePath,
 								bGenerateDiffPatch,
-								Params.HasDLCName));
+								Params.HasDLCName,
+								FileExtensionsToSkipCompressionInPaks));
 						}
 
 						if (!PakParams.bStageLoose)
@@ -3790,7 +3814,8 @@ namespace AutomationScripts
 								PakParams.bRehydrateAssets,
 								Params.SkipEncryption ? null : CryptoSettings,
 								PatchSourceContentPath,
-								Params.SkipEncryption ? "" : PakParams.EncryptionKeyGuid));
+								Params.SkipEncryption ? "" : PakParams.EncryptionKeyGuid,
+								FileExtensionsToSkipCompressionInPaks));
 						}
 						LogNames.Add(OutputLocation.GetFileNameWithoutExtension());
 					}
diff --git a/Engine/Source/Runtime/Apple/MetalRHI/Private/MetalDynamicRHI_Shaders.cpp b/Engine/Source/Runtime/Apple/MetalRHI/Private/MetalDynamicRHI_Shaders.cpp
--- a/Engine/Source/Runtime/Apple/MetalRHI/Private/MetalDynamicRHI_Shaders.cpp
+++ b/Engine/Source/Runtime/Apple/MetalRHI/Private/MetalDynamicRHI_Shaders.cpp
@@ -9,6 +9,7 @@
 #include "MetalShaderTypes.h"
 #include "Shaders/MetalShaderLibrary.h"
 #include "DataDrivenShaderPlatformInfo.h"
+#include "Serialization/StaticMemoryReader.h"
 #include "Interfaces/IPluginManager.h"
 
 //------------------------------------------------------------------------------
@@ -126,17 +127,70 @@ FRHIShaderLibraryRef FMetalDynamicRHI::RHICreateShaderLibrary(EShaderPlatform Pl
         return *FoundShaderLibrary;
     }
 
+    auto SerializeShaderCode = [](FMetalShaderLibrary::FShaderCodeArrayType& Array, FArchive& Ar)
+    {
+#if !USE_MMAPPED_SHADERARCHIVE
+		Ar << Array;
+#else
+       if (Ar.GetArchiveName() == TEXT("FStaticMemoryReader"))
+       {
+			using ArrayType = std::remove_cvref_t <decltype(Array)>;
+			FStaticMemoryReader* MemReader = static_cast<FStaticMemoryReader*>(&Ar);
+
+			typename ArrayType::SizeType SerializeNum;
+			Ar << SerializeNum;
+
+			uint64 ArrayBytes = SerializeNum * sizeof(typename ArrayType::ElementType);
+			uint64 Offset = Ar.Tell();
+
+			Array = FMetalShaderLibrary::FShaderCodeArrayType(MemReader->GetData() + Offset, SerializeNum);
+			Ar.Seek(Offset + ArrayBytes);
+       }
+	   else
+	   {
+		   UE_LOG(LogMetal, Fatal, TEXT("mmapped array must be serialized via FStaticMemoryReader"));
+	   }
+#endif
+	};
+
+ #if !USE_MMAPPED_SHADERARCHIVE
     FArchive* BinaryShaderAr = IFileManager::Get().CreateFileReader(*BinaryShaderFile);
+#else
+	FStaticMemoryReader* BinaryShaderAr = nullptr;
+    TUniquePtr<FMetalShaderLibrary::FShaderLibDataOwner> MemOwner = MakeUnique<FMetalShaderLibrary::FShaderLibDataOwner>();
+	if (FPlatformProperties::SupportsMemoryMappedFiles())
+	{
+		MemOwner->MappedCacheFile = TUniquePtr<IMappedFileHandle>(FPlatformFileManager::Get().GetPlatformFile().OpenMapped(*BinaryShaderFile));
+		if (MemOwner->MappedCacheFile.IsValid())
+		{
+			MemOwner->MappedRegion = TUniquePtr<IMappedFileRegion>(MemOwner->MappedCacheFile->MapRegion(0, MemOwner->MappedCacheFile->GetFileSize()));
+			if(MemOwner->MappedRegion.IsValid())
+			{
+				UE_LOG(LogMetal, Display, TEXT("mmapping %s, %d bytes"), *BinaryShaderFile, MemOwner->MappedCacheFile->GetFileSize());
+				BinaryShaderAr = new FStaticMemoryReader(MemOwner->MappedRegion->GetMappedPtr(), MemOwner->MappedCacheFile->GetFileSize());
+			}
+		}
+	}
 
+	if(BinaryShaderAr == nullptr)
+	{
+		TArray<uint8>& FileData = MemOwner->Mem;
+		if (FFileHelper::LoadFileToArray(FileData, *BinaryShaderFile))
+		{
+			UE_LOG(LogMetal, Display, TEXT("emulating mmapping %s, %d bytes!"), *BinaryShaderFile, FileData.Num());
+			BinaryShaderAr = new FStaticMemoryReader(FileData.GetData(), FileData.Num());
+		}
+	}
+#endif
     if( BinaryShaderAr != NULL )
     {
         FMetalShaderLibraryHeader Header;
         FSerializedShaderArchive SerializedShaders;
-        TArray<uint8> ShaderCode;
+        FMetalShaderLibrary::FShaderCodeArrayType ShaderCode;
 
         *BinaryShaderAr << Header;
         *BinaryShaderAr << SerializedShaders;
-        *BinaryShaderAr << ShaderCode;
+		SerializeShaderCode(ShaderCode,*BinaryShaderAr);
         BinaryShaderAr->Flush();
         delete BinaryShaderAr;
 
@@ -165,19 +219,8 @@ FRHIShaderLibraryRef FMetalDynamicRHI::RHICreateShaderLibrary(EShaderPlatform Pl
                     static const bool bRunningWithZenStore = FPlatformFileManager::Get().FindPlatformFile(TEXT("StorageServer")) != nullptr;
                     if( bRunningWithZenStore )
                     {
-                        FArchive* Reader = IFileManager::Get().CreateFileReader(*MetalLibraryFilePath);
                         TArray<uint8> LibraryData;
-                        bool Success = false;
-                        if (Reader)
-                        {
-                            const int64 FileSize = Reader->TotalSize();
-                            LibraryData.Reset( FileSize + 2 );
-                            LibraryData.AddUninitialized( FileSize );
-                            Reader->Serialize(LibraryData.GetData(), LibraryData.Num());
-                            Success = Reader->Close();
-                            delete Reader;
-                        }
-                        if (Success)
+                        if (FFileHelper::LoadFileToArray(LibraryData, *MetalLibraryFilePath))
                         {
                             dispatch_data_t data = dispatch_data_create(LibraryData.GetData(), LibraryData.Num(), nil, DISPATCH_DATA_DESTRUCTOR_DEFAULT);
                             Library = NS::TransferPtr(GetMetalDeviceContext().GetDevice()->newLibrary(data, &Error));
@@ -206,9 +249,18 @@ FRHIShaderLibraryRef FMetalDynamicRHI::RHICreateShaderLibrary(EShaderPlatform Pl
                 }
             }
 
-            Result = new FMetalShaderLibrary(Platform, Name, BinaryShaderFile, Header, SerializedShaders, ShaderCode, Libraries);
+            FMetalShaderLibrary* MtlLib = new FMetalShaderLibrary(Platform, Name, BinaryShaderFile, Header, MoveTemp(SerializedShaders), MoveTemp(ShaderCode), Libraries
+#if USE_MMAPPED_SHADERARCHIVE
+				,MoveTemp(MemOwner)
+#endif
+			);
+            Result = MtlLib;
             FMetalShaderLibrary::LoadedShaderLibraryMap.Add(BinaryShaderFile, Result.GetReference());
         }
+		else
+		{
+		    UE_LOG(LogMetal, Display, TEXT("Unknown shader format for %s!"), *LibName);
+		}
     }
     else
     {
diff --git a/Engine/Source/Runtime/Apple/MetalRHI/Private/Shaders/MetalShaderLibrary.cpp b/Engine/Source/Runtime/Apple/MetalRHI/Private/Shaders/MetalShaderLibrary.cpp
--- a/Engine/Source/Runtime/Apple/MetalRHI/Private/Shaders/MetalShaderLibrary.cpp
+++ b/Engine/Source/Runtime/Apple/MetalRHI/Private/Shaders/MetalShaderLibrary.cpp
@@ -50,15 +50,15 @@ FMetalShaderLibrary::FMetalShaderLibrary(EShaderPlatform Platform,
 										 FString const& Name,
 										 const FString& InShaderLibraryFilename,
 										 const FMetalShaderLibraryHeader& InHeader,
-										 const FSerializedShaderArchive& InSerializedShaders,
-										 const TArray<uint8>& InShaderCode,
+										 FSerializedShaderArchive&& InSerializedShaders,
+                                         FShaderCodeArrayType&& InShaderCode,
 										 const TArray<MTLLibraryPtr>& InLibrary)
 	: FRHIShaderLibrary(Platform, Name)
 	, ShaderLibraryFilename(InShaderLibraryFilename)
 	, Library(InLibrary)
 	, Header(InHeader)
-	, SerializedShaders(InSerializedShaders)
-	, ShaderCode(InShaderCode)
+	, SerializedShaders(MoveTemp(InSerializedShaders))
+	, ShaderCode(MoveTemp(InShaderCode))
 {
 #if !UE_BUILD_SHIPPING
 	DebugFile = nullptr;
@@ -88,28 +88,32 @@ bool FMetalShaderLibrary::IsNativeLibrary() const
 
 int32 FMetalShaderLibrary::GetNumShaders() const
 {
-	return SerializedShaders.ShaderEntries.Num();
+	return SerializedShaders.GetShaderEntries().Num();
 }
 
 int32 FMetalShaderLibrary::GetNumShaderMaps() const
 {
-	return SerializedShaders.ShaderMapEntries.Num();
+	return SerializedShaders.GetShaderMapEntries().Num();
 }
 
 uint32 FMetalShaderLibrary::GetSizeBytes() const
 {
+#if USE_MMAPPED_SHADERARCHIVE
+	return SerializedShaders.GetAllocatedSize() + ShaderCode.Num() * ShaderCode.GetTypeSize();
+#else
 	return SerializedShaders.GetAllocatedSize() + ShaderCode.GetAllocatedSize();
+#endif
 }
 
 int32 FMetalShaderLibrary::GetNumShadersForShaderMap(int32 ShaderMapIndex) const
 {
-	return SerializedShaders.ShaderMapEntries[ShaderMapIndex].NumShaders;
+	return SerializedShaders.GetShaderMapEntries()[ShaderMapIndex].NumShaders;
 }
 
 int32 FMetalShaderLibrary::GetShaderIndex(int32 ShaderMapIndex, int32 i) const
 {
-	const FShaderMapEntry& ShaderMapEntry = SerializedShaders.ShaderMapEntries[ShaderMapIndex];
-	return SerializedShaders.ShaderIndices[ShaderMapEntry.ShaderIndicesOffset + i];
+	const FShaderMapEntry& ShaderMapEntry = SerializedShaders.GetShaderMapEntries()[ShaderMapIndex];
+	return SerializedShaders.GetShaderIndices()[ShaderMapEntry.ShaderIndicesOffset + i];
 }
 
 int32 FMetalShaderLibrary::FindShaderMapIndex(const FSHAHash& Hash)
@@ -124,12 +128,12 @@ int32 FMetalShaderLibrary::FindShaderIndex(const FSHAHash& Hash)
 
 TRefCountPtr<FRHIShader> FMetalShaderLibrary::CreateShader(int32 Index)
 {
-	const FShaderCodeEntry& ShaderEntry = SerializedShaders.ShaderEntries[Index];
+	const FShaderCodeEntry& ShaderEntry = SerializedShaders.GetShaderEntries()[Index];
 
 	// We don't handle compressed shaders here, since typically these are just tiny headers.
 	check(ShaderEntry.Size == ShaderEntry.UncompressedSize);
 
-	const TArrayView<uint8> Code = MakeArrayView(ShaderCode.GetData() + ShaderEntry.Offset, ShaderEntry.Size);
+	const TArrayView<const uint8> Code = MakeArrayView(ShaderCode.GetData() + ShaderEntry.Offset, ShaderEntry.Size);
 	const int32 LibraryIndex = Index / Header.NumShadersPerLibrary;
 
 	TRefCountPtr<FRHIShader> Shader;
@@ -179,7 +183,7 @@ TRefCountPtr<FRHIShader> FMetalShaderLibrary::CreateShader(int32 Index)
 
 	if (Shader)
 	{
-		Shader->SetHash(SerializedShaders.ShaderHashes[Index]);
+		Shader->SetHash(SerializedShaders.GetShaderHashes()[Index]);
 	}
 
 	return Shader;
diff --git a/Engine/Source/Runtime/Apple/MetalRHI/Private/Shaders/MetalShaderLibrary.h b/Engine/Source/Runtime/Apple/MetalRHI/Private/Shaders/MetalShaderLibrary.h
--- a/Engine/Source/Runtime/Apple/MetalRHI/Private/Shaders/MetalShaderLibrary.h
+++ b/Engine/Source/Runtime/Apple/MetalRHI/Private/Shaders/MetalShaderLibrary.h
@@ -10,10 +10,42 @@
 #include "MetalShaderResources.h"
 #include "RHIShaderLibrary.h"
 #include "ShaderCodeArchive.h"
+#include "Async/MappedFileHandle.h"
 
 class FMetalShaderLibrary final : public FRHIShaderLibrary
 {
 public:
+
+#if !USE_MMAPPED_SHADERARCHIVE
+    using FShaderCodeArrayType = TArray<uint8>;
+#else
+    struct FShaderLibDataOwner
+    {
+        TArray<uint8> Mem;
+        TUniquePtr<IMappedFileHandle> MappedCacheFile;
+        TUniquePtr<IMappedFileRegion> MappedRegion;
+    };
+    
+	using FShaderCodeArrayType = TArrayView<const uint8>;
+private:
+    TUniquePtr<FShaderLibDataOwner> MemOwner;
+public:
+    
+    FMetalShaderLibrary(EShaderPlatform Platform,
+        FString const& Name,
+        const FString& InShaderLibraryFilename,
+        const FMetalShaderLibraryHeader& InHeader,
+        FSerializedShaderArchive&& InSerializedShaders,
+        FShaderCodeArrayType&& InShaderCode,
+        const TArray<MTLLibraryPtr>& InLibrary,
+        TUniquePtr<FShaderLibDataOwner>&& InMemOwner)
+        :
+        FMetalShaderLibrary(Platform, Name, InShaderLibraryFilename, InHeader, MoveTemp(InSerializedShaders), MoveTemp(InShaderCode), InLibrary)
+        {
+            MemOwner = MoveTemp(InMemOwner);
+        }
+#endif
+  
 	static FCriticalSection LoadedShaderLibraryMutex;
 	static TMap<FString, FRHIShaderLibrary*> LoadedShaderLibraryMap;
 
@@ -21,8 +53,8 @@ public:
 						FString const& Name,
 						const FString& InShaderLibraryFilename,
 						const FMetalShaderLibraryHeader& InHeader,
-						const FSerializedShaderArchive& InSerializedShaders,
-						const TArray<uint8>& InShaderCode,
+						FSerializedShaderArchive&& InSerializedShaders,
+						FShaderCodeArrayType&& InShaderCode,
 						const TArray<MTLLibraryPtr>& InLibrary);
 
 	virtual ~FMetalShaderLibrary();
@@ -39,7 +71,7 @@ public:
 	virtual int32 FindShaderIndex(const FSHAHash& Hash) override;
 	virtual FSHAHash GetShaderHash(int32 ShaderMapIndex, int32 ShaderIndex) override
 	{ 
-		return SerializedShaders.ShaderHashes[GetShaderIndex(ShaderMapIndex, ShaderIndex)];
+		return SerializedShaders.GetShaderHashes()[GetShaderIndex(ShaderMapIndex, ShaderIndex)];
 	};
 
 	virtual bool PreloadShader(int32 ShaderIndex, FGraphEventArray& OutCompletionEvents) override { return false; }
@@ -52,7 +84,7 @@ private:
 	TArray<MTLLibraryPtr> Library;
 	FMetalShaderLibraryHeader Header;
 	FSerializedShaderArchive SerializedShaders;
-	TArray<uint8> ShaderCode;
+	FShaderCodeArrayType ShaderCode;
 #if !UE_BUILD_SHIPPING
 	class FMetalShaderDebugZipFile* DebugFile;
 #endif
diff --git a/Engine/Source/Runtime/Core/Public/Serialization/StaticMemoryReader.h b/Engine/Source/Runtime/Core/Public/Serialization/StaticMemoryReader.h
--- a/Engine/Source/Runtime/Core/Public/Serialization/StaticMemoryReader.h
+++ b/Engine/Source/Runtime/Core/Public/Serialization/StaticMemoryReader.h
@@ -30,6 +30,9 @@ public:
 		return NumBytes;
 	}
 
+	// returns the raw data ptr assigned to the memory reader.
+	const uint8* GetData() const { return Bytes; }
+
 	void Serialize( void* Data, int64 Num )
 	{
 		if (Num && !IsError())
diff --git a/Engine/Source/Runtime/RenderCore/Private/ShaderCodeArchive.cpp b/Engine/Source/Runtime/RenderCore/Private/ShaderCodeArchive.cpp
--- a/Engine/Source/Runtime/RenderCore/Private/ShaderCodeArchive.cpp
+++ b/Engine/Source/Runtime/RenderCore/Private/ShaderCodeArchive.cpp
@@ -19,6 +19,7 @@
 #include "ShaderCodeLibrary.h"
 #include "ShaderCore.h"
 #include "Stats/Stats.h"
+#include "Serialization/StaticMemoryReader.h"
 
 #if WITH_EDITOR
 #include "Misc/Optional.h"
@@ -104,6 +105,7 @@ int32 FSerializedShaderArchive::FindShaderMap(const FSHAHash& Hash) const
 	return FindShaderMapWithKey(Hash, Key);
 }
 
+#if !USE_MMAPPED_SHADERARCHIVE
 bool FSerializedShaderArchive::FindOrAddShaderMap(const FSHAHash& Hash, int32& OutIndex, const FShaderMapAssetPaths* AssociatedAssets)
 {
 	const uint32 Key = GetTypeHash(Hash);
@@ -127,6 +129,7 @@ bool FSerializedShaderArchive::FindOrAddShaderMap(const FSHAHash& Hash, int32& O
 	OutIndex = Index;
 	return bAdded;
 }
+#endif
 
 int32 FSerializedShaderArchive::FindShaderWithKey(const FSHAHash& Hash, uint32 Key) const
 {
@@ -146,6 +149,7 @@ int32 FSerializedShaderArchive::FindShader(const FSHAHash& Hash) const
 	return FindShaderWithKey(Hash, Key);
 }
 
+#if !USE_MMAPPED_SHADERARCHIVE
 bool FSerializedShaderArchive::FindOrAddShader(const FSHAHash& Hash, int32& OutIndex)
 {
 	const uint32 Key = GetTypeHash(Hash);
@@ -246,7 +250,8 @@ bool LoadFromCompactBinary(FCbFieldView Field, FSerializedShaderArchive& OutArch
 
 	return bOk;
 }
-#endif
+#endif //WITH_EDITOR
+#endif //!USE_MMAPPED_SHADERARCHIVE
 
 
 #if UE_SCA_VISUALIZE_SHADER_USAGE
@@ -445,6 +450,7 @@ void FSerializedShaderArchive::DecompressShader(int32 Index, const TArray<TArray
 	}
 }
 
+#if !USE_MMAPPED_SHADERARCHIVE
 void FSerializedShaderArchive::Finalize()
 {
 	// Set the correct offsets
@@ -496,15 +502,47 @@ void FSerializedShaderArchive::Finalize()
 		check(CurrentPreloadEntry.Size == 0);
 	}
 }
+#endif
 
 void FSerializedShaderArchive::Serialize(FArchive& Ar)
 {
+#if USE_MMAPPED_SHADERARCHIVE
+	auto SerializeMappedToArrayView = [](auto& ArrayView, FStaticMemoryReader& Ar)
+	{
+        using ArrayType = std::remove_cvref_t <decltype(ArrayView)>;
+        typename ArrayType::SizeType SerializeNum;
+        
+        Ar << SerializeNum;
+
+        uint64 ArrayBytes = SerializeNum * sizeof(typename ArrayType::ElementType);
+        uint64 Offset = Ar.Tell();
+
+        ArrayView = TArrayView<typename ArrayType::ElementType >((typename ArrayType::ElementType*)(Ar.GetData() + Offset), SerializeNum);
+        Ar.Seek(Offset + ArrayBytes);
+	};
+    
+    FStaticMemoryReader& MemReaderAr = static_cast<FStaticMemoryReader&>(Ar);
+    if (Ar.GetArchiveName() != TEXT("FStaticMemoryReader"))
+    {
+        UE_LOG(LogShaderLibrary, Fatal, TEXT("mmapped shader archive must be serialized via FStaticMemoryReader"));
+    }
+    else
+    {
+        SerializeMappedToArrayView(ShaderMapHashes, MemReaderAr);
+        SerializeMappedToArrayView(ShaderHashes, MemReaderAr);
+        SerializeMappedToArrayView(ShaderMapEntries, MemReaderAr);
+        SerializeMappedToArrayView(ShaderEntries, MemReaderAr);
+        SerializeMappedToArrayView(PreloadEntries, MemReaderAr);
+        SerializeMappedToArrayView(ShaderIndices, MemReaderAr);
+    }
+#else
 	Ar << ShaderMapHashes;
 	Ar << ShaderHashes;
 	Ar << ShaderMapEntries;
 	Ar << ShaderEntries;
 	Ar << PreloadEntries;
 	Ar << ShaderIndices;
+#endif
 
 	check(ShaderHashes.Num() == ShaderEntries.Num());
 	check(ShaderMapHashes.Num() == ShaderMapEntries.Num());
@@ -1037,9 +1075,9 @@ FShaderCodeArchive* FShaderCodeArchive::Create(EShaderPlatform InPlatform, FArch
 	// Open library for async reads
 	Library->FileCacheHandle = IFileCacheHandle::CreateFileCacheHandle(*InDestFilePath);
 
-	Library->DebugVisualizer.Initialize(Library->SerializedShaders.ShaderEntries.Num());
+	Library->DebugVisualizer.Initialize(Library->SerializedShaders.GetShaderEntries().Num());
 
-	UE_LOG(LogShaderLibrary, Display, TEXT("Using %s for material shader code. Total %d unique shaders."), *InDestFilePath, Library->SerializedShaders.ShaderEntries.Num());
+	UE_LOG(LogShaderLibrary, Display, TEXT("Using %s for material shader code. Total %d unique shaders."), *InDestFilePath, Library->SerializedShaders.GetShaderEntries().Num());
 
 	INC_DWORD_STAT_BY(STAT_Shaders_ShaderResourceMemory, Library->GetSizeBytes());
 
@@ -1073,7 +1111,7 @@ void FShaderCodeArchive::Teardown()
 		FShaderPreloadEntry& ShaderPreloadEntry = ShaderPreloads[ShaderIndex];
 		if (ShaderPreloadEntry.Code)
 		{
-			const FShaderCodeEntry& ShaderEntry = SerializedShaders.ShaderEntries[ShaderIndex];
+			const FShaderCodeEntry& ShaderEntry = SerializedShaders.GetShaderEntries()[ShaderIndex];
 			FMemory::Free(ShaderPreloadEntry.Code);
 			ShaderPreloadEntry.Code = nullptr;
 			DEC_DWORD_STAT_BY(STAT_Shaders_ShaderPreloadMemory, ShaderEntry.Size);
@@ -1085,7 +1123,7 @@ void FShaderCodeArchive::Teardown()
 
 void FShaderCodeArchive::OnShaderPreloadFinished(int32 ShaderIndex, const IMemoryReadStreamRef& PreloadData)
 {
-	const FShaderCodeEntry& ShaderEntry = SerializedShaders.ShaderEntries[ShaderIndex];
+	const FShaderCodeEntry& ShaderEntry = SerializedShaders.GetShaderEntries()[ShaderIndex];
 	PreloadData->EnsureReadNonBlocking();		// Ensure data is ready before taking the lock
 	{
 		FWriteScopeLock Lock(ShaderPreloadLock);
@@ -1129,7 +1167,7 @@ bool FShaderCodeArchive::PreloadShader(int32 ShaderIndex, FGraphEventArray& OutC
 	{
 		check(!ShaderPreloadEntry.PreloadEvent);
 
-		const FShaderCodeEntry& ShaderEntry = SerializedShaders.ShaderEntries[ShaderIndex];
+		const FShaderCodeEntry& ShaderEntry = SerializedShaders.GetShaderEntries()[ShaderIndex];
 		ShaderPreloadEntry.Code = FMemory::Malloc(ShaderEntry.Size);
 		ShaderPreloadEntry.FramePreloadStarted = GFrameNumber;
 		DebugVisualizer.MarkExplicitlyPreloadedForVisualization(ShaderIndex);
@@ -1158,18 +1196,19 @@ bool FShaderCodeArchive::PreloadShaderMap(int32 ShaderMapIndex, FGraphEventArray
 {
 	LLM_SCOPE(ELLMTag::Shaders);
 
-	const FShaderMapEntry& ShaderMapEntry = SerializedShaders.ShaderMapEntries[ShaderMapIndex];
+	const FShaderMapEntry& ShaderMapEntry = SerializedShaders.GetShaderMapEntries()[ShaderMapIndex];
 	const EAsyncIOPriorityAndFlags IOPriority = (EAsyncIOPriorityAndFlags)GShaderCodeLibraryAsyncLoadingPriority;
 	const uint32 FrameNumber = GFrameNumber;
 	uint32 PreloadMemory = 0u;
 	
 	FWriteScopeLock Lock(ShaderPreloadLock);
 
+	TArrayView ShaderIndices = SerializedShaders.GetShaderIndices();
 	for (uint32 i = 0u; i < ShaderMapEntry.NumShaders; ++i)
 	{
-		const int32 ShaderIndex = SerializedShaders.ShaderIndices[ShaderMapEntry.ShaderIndicesOffset + i];
+		const int32 ShaderIndex = ShaderIndices[ShaderMapEntry.ShaderIndicesOffset + i];
 		FShaderPreloadEntry& ShaderPreloadEntry = ShaderPreloads[ShaderIndex];
-		const FShaderCodeEntry& ShaderEntry = SerializedShaders.ShaderEntries[ShaderIndex];
+		const FShaderCodeEntry& ShaderEntry = SerializedShaders.GetShaderEntries()[ShaderIndex];
 
 #if RHI_RAYTRACING
 		if (!IsRayTracingAllowed() && !IsCreateShadersOnLoadEnabled() && IsRayTracingShaderFrequency(static_cast<EShaderFrequency>(ShaderEntry.Frequency)))
@@ -1248,7 +1287,7 @@ void FShaderCodeArchive::ReleasePreloadedShader(int32 ShaderIndex)
 		{
 			FMemory::Free(ShaderPreloadEntry.Code);
 			ShaderPreloadEntry.Code = nullptr;
-			const FShaderCodeEntry& ShaderEntry = SerializedShaders.ShaderEntries[ShaderIndex];
+			const FShaderCodeEntry& ShaderEntry = SerializedShaders.GetShaderEntries()[ShaderIndex];
 			DEC_DWORD_STAT_BY(STAT_Shaders_ShaderPreloadMemory, ShaderEntry.Size);
 		}
 	}
@@ -1263,7 +1302,7 @@ TRefCountPtr<FRHIShader> FShaderCodeArchive::CreateShader(int32 Index)
 	FMemStackBase& MemStack = FMemStack::Get();
 	FMemMark Mark(MemStack);
 
-	const FShaderCodeEntry& ShaderEntry = SerializedShaders.ShaderEntries[Index];
+	const FShaderCodeEntry& ShaderEntry = SerializedShaders.GetShaderEntries()[Index];
 	FShaderPreloadEntry& ShaderPreloadEntry = ShaderPreloads[Index];
 	checkf(!ShaderPreloadEntry.bNeverToBePreloaded, TEXT("We are creating a shader that shouldn't be preloaded in this run (e.g. raytracing shader on D3D11)."));
 
@@ -1297,7 +1336,7 @@ TRefCountPtr<FRHIShader> FShaderCodeArchive::CreateShader(int32 Index)
 	ensureAlwaysMsgf(IsInRenderingThread() || GRHISupportsMultithreadedShaderCreation, TEXT("More than one thread is creating shaders, but GRHISupportsMultithreadedShaderCreation is false."));
 
 	const auto ShaderCodeView = MakeArrayView(ShaderCode, ShaderEntry.UncompressedSize);
-	const FSHAHash& ShaderHash = SerializedShaders.ShaderHashes[Index];
+	const FSHAHash& ShaderHash = SerializedShaders.GetShaderHashes()[Index];
 	switch (ShaderEntry.Frequency)
 	{
 	case SF_Vertex: Shader = RHICreateVertexShader(ShaderCodeView, ShaderHash); CheckShaderCreation(Shader, Index); break;
@@ -1359,12 +1398,12 @@ FIoChunkId FIoStoreShaderCodeArchive::GetShaderCodeChunkId(const FSHAHash& Shade
 
 void FIoStoreShaderCodeArchive::CreateIoStoreShaderCodeArchiveHeader(const FName& Format, const FSerializedShaderArchive& SerializedShaders, FIoStoreShaderCodeArchiveHeader& OutHeader)
 {
-	OutHeader.ShaderMapHashes = SerializedShaders.ShaderMapHashes;
-	OutHeader.ShaderHashes = SerializedShaders.ShaderHashes;
+	OutHeader.ShaderMapHashes = SerializedShaders.GetShaderMapHashes();
+	OutHeader.ShaderHashes = SerializedShaders.GetShaderHashes();
 	// shader group hashes will be populated later
 
-	OutHeader.ShaderMapEntries.Empty(SerializedShaders.ShaderMapEntries.Num());
-	for (const FShaderMapEntry& ShaderMapEntry : SerializedShaders.ShaderMapEntries)
+	OutHeader.ShaderMapEntries.Empty(SerializedShaders.GetShaderMapEntries().Num());
+	for (const FShaderMapEntry& ShaderMapEntry : SerializedShaders.GetShaderMapEntries())
 	{
 		FIoStoreShaderMapEntry& IoStoreShaderMapEntry = OutHeader.ShaderMapEntries.AddDefaulted_GetRef();
 		IoStoreShaderMapEntry.ShaderIndicesOffset = ShaderMapEntry.ShaderIndicesOffset;
@@ -1372,11 +1411,11 @@ void FIoStoreShaderCodeArchive::CreateIoStoreShaderCodeArchiveHeader(const FName
 	}
 
 	// indices should be copied before grouping as the groups will append to the array
-	OutHeader.ShaderIndices = SerializedShaders.ShaderIndices;
+	OutHeader.ShaderIndices = SerializedShaders.GetShaderIndices();
 
 	// shader entries are copied, the remainder of the field will be assigned when splitting into groups
-	OutHeader.ShaderEntries.Empty(SerializedShaders.ShaderEntries.Num());
-	for (const FShaderCodeEntry& ShaderEntry : SerializedShaders.ShaderEntries)
+	OutHeader.ShaderEntries.Empty(SerializedShaders.GetShaderEntries().Num());
+	for (const FShaderCodeEntry& ShaderEntry : SerializedShaders.GetShaderEntries())
 	{
 		FIoStoreShaderCodeEntry& IoStoreShaderEntry = OutHeader.ShaderEntries.AddDefaulted_GetRef();
 		IoStoreShaderEntry.Frequency = ShaderEntry.Frequency;
@@ -1484,8 +1523,8 @@ void FIoStoreShaderCodeArchive::CreateIoStoreShaderCodeArchiveHeader(const FName
 		ShaderIndicesInGroup.Sort(
 			[&SerializedShaders](const int32 ShaderIndexA, const int32 ShaderIndexB)
 			{
-				const FShaderCodeEntry& ShaderEntryA = SerializedShaders.ShaderEntries[ShaderIndexA];
-				const FShaderCodeEntry& ShaderEntryB = SerializedShaders.ShaderEntries[ShaderIndexB];
+				const FShaderCodeEntry& ShaderEntryA = SerializedShaders.GetShaderEntries()[ShaderIndexA];
+				const FShaderCodeEntry& ShaderEntryB = SerializedShaders.GetShaderEntries()[ShaderIndexB];
 				if (ShaderEntryA.UncompressedSize != ShaderEntryB.UncompressedSize)
 				{
 					return ShaderEntryA.UncompressedSize < ShaderEntryB.UncompressedSize;
@@ -1509,6 +1548,9 @@ void FIoStoreShaderCodeArchive::CreateIoStoreShaderCodeArchiveHeader(const FName
 		GroupEntry.NumShaders = ShaderIndicesInGroup.Num();
 		// ShaderIndicesOffset will be filled later, once we know all the groups (see comment about StoredGroupShaderIndices above).
 
+		const TArrayView<const FSHAHash> ShaderHashes = SerializedShaders.GetShaderHashes();
+		const TArrayView<const FShaderCodeEntry> ShaderEntries = SerializedShaders.GetShaderEntries();
+
 		// update shader entries both with the group number and their uncompressed offset in the group
 		FSHA1 GroupHasher;
 		uint32 CurrentGroupSize = 0;
@@ -1520,11 +1562,11 @@ void FIoStoreShaderCodeArchive::CreateIoStoreShaderCodeArchiveHeader(const FName
 			IoStoreShaderEntry.UncompressedOffsetInGroup = CurrentGroupSize;
 
 			// group hash is constructed from hashing the shaders in the group.
-			GroupHasher.Update(SerializedShaders.ShaderHashes[ShaderIndex].Hash, sizeof(FSHAHash));
+			GroupHasher.Update(ShaderHashes[ShaderIndex].Hash, sizeof(FSHAHash));
 			// shader hash as of now excludes optional data, so we cannot rely on it, especially across the shader formats. Make the group hash a bit more robust by including the shader size in it.
-			GroupHasher.Update(reinterpret_cast<const uint8*>(&SerializedShaders.ShaderEntries[ShaderIndex].UncompressedSize), sizeof(FShaderCodeEntry::UncompressedSize));
+			GroupHasher.Update(reinterpret_cast<const uint8*>(&ShaderEntries[ShaderIndex].UncompressedSize), sizeof(FShaderCodeEntry::UncompressedSize));
 
-			CurrentGroupSize += SerializedShaders.ShaderEntries[ShaderIndex].UncompressedSize;
+			CurrentGroupSize += ShaderEntries[ShaderIndex].UncompressedSize;
 		}
 		// Shader hashes cannot be used to uniquely identify across shader formats due to aforementioned exclusion of optional data from it.
 		// Include the shader format (in a platform-agnostic way) into the group hash to lower the risk of collision of shaders of different formats.
@@ -1545,9 +1587,10 @@ void FIoStoreShaderCodeArchive::CreateIoStoreShaderCodeArchiveHeader(const FName
 	{
 		// calculate current group size
 		uint32 GroupSize = 0;
+		const TArrayView<const FShaderCodeEntry> ShaderEntries = SerializedShaders.GetShaderEntries();
 		for (uint32 ShaderIdx : CurrentShaderGroup)
 		{
-			GroupSize += SerializedShaders.ShaderEntries[ShaderIdx].UncompressedSize;
+			GroupSize += ShaderEntries[ShaderIdx].UncompressedSize;
 		}
 		
 		if (LIKELY(GroupSize <= MaxUncompressedShaderGroupSize || CurrentShaderGroup.Num() == 1))
@@ -1569,8 +1612,9 @@ void FIoStoreShaderCodeArchive::CreateIoStoreShaderCodeArchiveHeader(const FName
 			CurrentShaderGroup.Sort(
 				[&SerializedShaders](const int32 ShaderIndexA, const int32 ShaderIndexB)
 				{
-					const FShaderCodeEntry& ShaderEntryA = SerializedShaders.ShaderEntries[ShaderIndexA];
-					const FShaderCodeEntry& ShaderEntryB = SerializedShaders.ShaderEntries[ShaderIndexB];
+					const TArrayView<const FShaderCodeEntry> ShaderEntries = SerializedShaders.GetShaderEntries();
+					const FShaderCodeEntry& ShaderEntryA = ShaderEntries[ShaderIndexA];
+					const FShaderCodeEntry& ShaderEntryB = ShaderEntries[ShaderIndexB];
 					if (ShaderEntryA.UncompressedSize != ShaderEntryB.UncompressedSize)
 					{
 						return ShaderEntryA.UncompressedSize > ShaderEntryB.UncompressedSize;
@@ -1600,7 +1644,7 @@ void FIoStoreShaderCodeArchive::CreateIoStoreShaderCodeArchiveHeader(const FName
 				}
 
 				NewGroups[SmallestNewGroupIdx].Add(ShaderIdx);
-				NewGroupSizes[SmallestNewGroupIdx] += SerializedShaders.ShaderEntries[ShaderIdx].UncompressedSize;
+				NewGroupSizes[SmallestNewGroupIdx] += ShaderEntries[ShaderIdx].UncompressedSize;
 			}
 			
 #if DO_CHECK // sanity checks
@@ -1698,7 +1742,7 @@ void FIoStoreShaderCodeArchive::CreateIoStoreShaderCodeArchiveHeader(const FName
 	}
 
 	/** Tries to find whether NewIndices exist as a subsequence in ExistingIndices.Returns - 1 if not found. */
-	auto FindSequenceInArray = [](const TArray<uint32>& ExistingIndices, const TArray<uint32>& NewIndices) -> int32
+	auto FindSequenceInArray = [](const TArrayView<const uint32> ExistingIndices, const TArrayView<const uint32> NewIndices) -> int32
 	{
 		check(NewIndices.Num() > 0);
 
@@ -1742,7 +1786,7 @@ void FIoStoreShaderCodeArchive::CreateIoStoreShaderCodeArchiveHeader(const FName
 			// See if we can find indices in that order somewhere in the ShaderIndices array already, to avoid adding new indices.
 			// We are looking in the read-only original array, because there's no sense to look in OutHeader.ShaderIndices - groups don't overlap,
 			// so we know that newly added (by some previous group) indices aren't useful for us.
-			int32 ExistingOffset = FindSequenceInArray(SerializedShaders.ShaderIndices, ShaderIndicesInGroup);
+			int32 ExistingOffset = FindSequenceInArray(SerializedShaders.GetShaderIndices(), ShaderIndicesInGroup);
 			if (ExistingOffset != INDEX_NONE)
 			{
 				GroupEntry.ShaderIndicesOffset = static_cast<uint32>(ExistingOffset);
@@ -1769,8 +1813,8 @@ void FIoStoreShaderCodeArchive::CreateIoStoreShaderCodeArchiveHeader(const FName
 		}
 	}
 
-	checkf(OutHeader.ShaderEntries.Num() == SerializedShaders.ShaderEntries.Num(), TEXT("Error creating IoStoreShaderArchive header - shader entries differ (%d in IoStore, %d original). Bug in grouping logic?"),
-		OutHeader.ShaderEntries.Num(), SerializedShaders.ShaderEntries.Num());
+	checkf(OutHeader.ShaderEntries.Num() == SerializedShaders.GetShaderEntries().Num(), TEXT("Error creating IoStoreShaderArchive header - shader entries differ (%d in IoStore, %d original). Bug in grouping logic?"),
+		OutHeader.ShaderEntries.Num(), SerializedShaders.GetShaderEntries().Num());
 	checkf(OutHeader.ShaderGroupIoHashes.Num() == OutHeader.ShaderGroupEntries.Num(), TEXT("Error creating IoStoreShaderArchive header - mismatch between shader group hashes and descriptors (%d descriptors, %d hashes). Bug in grouping logic?"),
 		OutHeader.ShaderGroupEntries.Num(), OutHeader.ShaderGroupIoHashes.Num());
 	checkf(OutHeader.ShaderGroupEntries.Num() != 0, TEXT("At least one group must have been created"));
diff --git a/Engine/Source/Runtime/RenderCore/Public/ShaderCodeArchive.h b/Engine/Source/Runtime/RenderCore/Public/ShaderCodeArchive.h
--- a/Engine/Source/Runtime/RenderCore/Public/ShaderCodeArchive.h
+++ b/Engine/Source/Runtime/RenderCore/Public/ShaderCodeArchive.h
@@ -30,6 +30,11 @@ class FCbFieldView;
 class FCbWriter;
 #endif
 
+#ifndef USE_MMAPPED_SHADERARCHIVE
+    // If enabled FSerializedShaderArchive will use a mmapping of the on-disk shader archive instead of normal serialization. The on-disk format must match the expected in memory layout.
+   #define USE_MMAPPED_SHADERARCHIVE 0
+#endif
+
 // enable visualization in the desktop Development builds only as it has a memory hit and writes files
 #define UE_SCA_VISUALIZE_SHADER_USAGE			(!WITH_EDITOR && UE_BUILD_DEVELOPMENT && PLATFORM_DESKTOP)
 
@@ -54,6 +59,10 @@ static FArchive& operator <<(FArchive& Ar, FFileCachePreloadEntry& Ref)
 	return Ar << Ref.Offset << Ref.Size;
 }
 
+
+#if USE_MMAPPED_SHADERARCHIVE
+#pragma pack(push, 1)
+#endif
 struct FShaderCodeEntry
 {
 	uint64 Offset = 0;
@@ -71,32 +80,49 @@ struct FShaderCodeEntry
 		return Ar << Ref.Offset << Ref.Size << Ref.UncompressedSize << Ref.Frequency;
 	}
 };
+#if USE_MMAPPED_SHADERARCHIVE
+#pragma pack(pop)
+#endif
 
 // Portion of shader code archive that's serialize to disk
 class FSerializedShaderArchive
 {
 public:
+#if USE_MMAPPED_SHADERARCHIVE
+	template<typename T> using TArrayType = TArrayView<T>;
+	// Ensure use of the const Get functions when mmapping is used.
+	private:
+#else
+	template<typename T> using TArrayType = TArray<T>;
+#endif
 
 	/** Hashes of all shadermaps in the library */
-	TArray<FSHAHash> ShaderMapHashes;
+	TArrayType<FSHAHash> ShaderMapHashes;
 
 	/** Output hashes of all shaders in the library */
-	TArray<FSHAHash> ShaderHashes;
+	TArrayType<FSHAHash> ShaderHashes;
 
 	/** An array of a shadermap descriptors. Each shadermap can reference an arbitrary number of shaders */
-	TArray<FShaderMapEntry> ShaderMapEntries;
+	TArrayType<FShaderMapEntry> ShaderMapEntries;
 
 	/** An array of all shaders descriptors, deduplicated */
-	TArray<FShaderCodeEntry> ShaderEntries;
-
+	TArrayType<FShaderCodeEntry> ShaderEntries;
 	/** An array of entries for the bytes of shadercode that need to be preloaded for a shadermap.
 	  * Each shadermap has a range in this array, beginning of which is stored in FShaderMapEntry.FirstPreloadIndex. */
-	TArray<FFileCachePreloadEntry> PreloadEntries;
+	TArrayType<FFileCachePreloadEntry> PreloadEntries;
 
 	/** Flat array of shaders referenced by all shadermaps. Each shadermap has a range in this array, beginning of which is
 	  * stored as ShaderIndicesOffset in the shadermap's descriptor (FShaderMapEntry).
 	  */
-	TArray<uint32> ShaderIndices;
+	TArrayType<uint32> ShaderIndices;
+
+public:
+	const TArrayView<const uint32> GetShaderIndices() const { return ShaderIndices; }
+	const TArrayView<const FFileCachePreloadEntry> GetPreloadEntries() const { return PreloadEntries; }
+	const TArrayView<const FShaderCodeEntry> GetShaderEntries() const { return ShaderEntries; }
+	const TArrayView<const FShaderMapEntry> GetShaderMapEntries() const { return ShaderMapEntries; }
+	const TArrayView<const FSHAHash> GetShaderHashes() const { return ShaderHashes; }
+	const TArrayView<const FSHAHash> GetShaderMapHashes() const { return ShaderMapHashes; }
 
 	FHashTable ShaderMapHashTable;
 	FHashTable ShaderHashTable;
@@ -154,12 +180,21 @@ public:
 
 	int64 GetAllocatedSize() const
 	{
+#if USE_MMAPPED_SHADERARCHIVE
+		return ShaderHashes.Num() * ShaderHashes.GetTypeSize() +
+			ShaderEntries.Num() * ShaderEntries.GetTypeSize() +
+			ShaderMapHashes.Num() * ShaderMapHashes.GetTypeSize() +
+			ShaderMapEntries.Num() * ShaderMapEntries.GetTypeSize() +
+			PreloadEntries.Num() * PreloadEntries.GetTypeSize() +
+			ShaderIndices.Num() * ShaderIndices.GetTypeSize()
+#else
 		return ShaderHashes.GetAllocatedSize() +
 			ShaderEntries.GetAllocatedSize() +
 			ShaderMapHashes.GetAllocatedSize() +
 			ShaderMapEntries.GetAllocatedSize() +
 			PreloadEntries.GetAllocatedSize() +
 			ShaderIndices.GetAllocatedSize()
+#endif
 #if WITH_EDITOR
 			+ ShaderCodeToAssets.GetAllocatedSize()
 #endif // WITH_EDITOR
@@ -169,18 +204,21 @@ public:
 	void Empty()
 	{
 		EmptyShaderMaps();
-
+#if !USE_MMAPPED_SHADERARCHIVE
 		ShaderHashes.Empty();
 		ShaderEntries.Empty();
+#endif
 		ShaderHashTable.Clear();
 	}
 
 	void EmptyShaderMaps()
 	{
+#if !USE_MMAPPED_SHADERARCHIVE
 		ShaderMapHashes.Empty();
 		ShaderMapEntries.Empty();
 		PreloadEntries.Empty();
 		ShaderIndices.Empty();
+#endif
 		ShaderMapHashTable.Clear();
 #if WITH_EDITOR
 		ShaderCodeToAssets.Empty();
@@ -208,16 +246,18 @@ public:
 
 	RENDERCORE_API int32 FindShaderMapWithKey(const FSHAHash& Hash, uint32 Key) const;
 	RENDERCORE_API int32 FindShaderMap(const FSHAHash& Hash) const;
-	RENDERCORE_API bool FindOrAddShaderMap(const FSHAHash& Hash, int32& OutIndex, const FShaderMapAssetPaths* AssociatedAssets);
-
 	RENDERCORE_API int32 FindShaderWithKey(const FSHAHash& Hash, uint32 Key) const;
 	RENDERCORE_API int32 FindShader(const FSHAHash& Hash) const;
+
+#if !USE_MMAPPED_SHADERARCHIVE
 	RENDERCORE_API bool FindOrAddShader(const FSHAHash& Hash, int32& OutIndex);
+	RENDERCORE_API bool FindOrAddShaderMap(const FSHAHash& Hash, int32& OutIndex, const FShaderMapAssetPaths* AssociatedAssets);
 	RENDERCORE_API void RemoveLastAddedShader();
+	RENDERCORE_API void Finalize();
+#endif
 
 	RENDERCORE_API void DecompressShader(int32 Index, const TArray<TArray<uint8>>& ShaderCode, TArray<uint8>& OutDecompressedShader) const;
 
-	RENDERCORE_API void Finalize();
 	RENDERCORE_API void Serialize(FArchive& Ar);
 #if WITH_EDITOR
 	RENDERCORE_API void SaveAssetInfo(FArchive& Ar);
@@ -327,14 +367,14 @@ public:
 			ShaderPreloads.GetAllocatedSize();
 	}
 
-	virtual int32 GetNumShaders() const override { return SerializedShaders.ShaderEntries.Num(); }
-	virtual int32 GetNumShaderMaps() const override { return SerializedShaders.ShaderMapEntries.Num(); }
-	virtual int32 GetNumShadersForShaderMap(int32 ShaderMapIndex) const override { return SerializedShaders.ShaderMapEntries[ShaderMapIndex].NumShaders; }
+	virtual int32 GetNumShaders() const override { return SerializedShaders.GetShaderEntries().Num(); }
+	virtual int32 GetNumShaderMaps() const override { return SerializedShaders.GetShaderMapEntries().Num(); }
+	virtual int32 GetNumShadersForShaderMap(int32 ShaderMapIndex) const override { return SerializedShaders.GetShaderMapEntries()[ShaderMapIndex].NumShaders; }
 
 	virtual int32 GetShaderIndex(int32 ShaderMapIndex, int32 i) const override
 	{
-		const FShaderMapEntry& ShaderMapEntry = SerializedShaders.ShaderMapEntries[ShaderMapIndex];
-		return SerializedShaders.ShaderIndices[ShaderMapEntry.ShaderIndicesOffset + i];
+		const FShaderMapEntry& ShaderMapEntry = SerializedShaders.GetShaderMapEntries()[ShaderMapIndex];
+		return SerializedShaders.GetShaderIndices()[ShaderMapEntry.ShaderIndicesOffset + i];
 	}
 
 	virtual int32 FindShaderMapIndex(const FSHAHash& Hash) override
@@ -349,7 +389,7 @@ public:
 
 	virtual FSHAHash GetShaderHash(int32 ShaderMapIndex, int32 ShaderIndex) override
 	{
-		return SerializedShaders.ShaderHashes[GetShaderIndex(ShaderMapIndex, ShaderIndex)];
+		return SerializedShaders.GetShaderHashes()[GetShaderIndex(ShaderMapIndex, ShaderIndex)];
 	};
 
 	virtual bool PreloadShader(int32 ShaderIndex, FGraphEventArray& OutCompletionEvents) override;
