Remove the use of FCompareShaderPipelineNameTypes since it does null checks on references.  This is getting flagged as a warning by the compiler.

- We already maintain a sorted list of pipelines so just use that for generating the global shader map ID.
- Add a check to ensure we are only ever putting in valid pipelines.

diff --git a/Engine/Source/Runtime/RenderCore/Private/GlobalShader.cpp b/Engine/Source/Runtime/RenderCore/Private/GlobalShader.cpp
--- a/Engine/Source/Runtime/RenderCore/Private/GlobalShader.cpp
+++ b/Engine/Source/Runtime/RenderCore/Private/GlobalShader.cpp
@@ -415,9 +415,11 @@ FGlobalShaderMapId::FGlobalShaderMapId(EShaderPlatform Platform, const ITargetPl
 		}
 	}
 
-	for (TLinkedList<FShaderPipelineType*>::TIterator ShaderPipelineIt(FShaderPipelineType::GetTypeList()); ShaderPipelineIt; ShaderPipelineIt.Next())
+	const TArray<FShaderPipelineType*>& SortedMaterialPipelineTypes = FShaderPipelineType::GetSortedTypes(FShaderType::EShaderTypeForDynamicCast::Global);
+
+	for (FShaderPipelineType* Pipeline : SortedMaterialPipelineTypes)
 	{
-		const FShaderPipelineType* Pipeline = *ShaderPipelineIt;
+		check(Pipeline);
 		if (Pipeline->IsGlobalTypePipeline())
 		{
 			int32 NumStagesNeeded = 0;
@@ -456,8 +458,6 @@ FGlobalShaderMapId::FGlobalShaderMapId(EShaderPlatform Platform, const ITargetPl
 		Dependencies.Add(Dependency);
 	}
 
-	// Shader pipeline dependencies
-	ShaderPipelineTypes.Sort(FCompareShaderPipelineNameTypes());
 	for (int32 TypeIndex = 0; TypeIndex < ShaderPipelineTypes.Num(); TypeIndex++)
 	{
 		const FShaderPipelineType* Pipeline = ShaderPipelineTypes[TypeIndex];
diff --git a/Engine/Source/Runtime/RenderCore/Public/Shader.h b/Engine/Source/Runtime/RenderCore/Public/Shader.h
--- a/Engine/Source/Runtime/RenderCore/Public/Shader.h
+++ b/Engine/Source/Runtime/RenderCore/Public/Shader.h
@@ -2052,39 +2052,6 @@ public:
 };
 
 
-/** Used to compare two shader pipeline types by name. */
-class FCompareShaderPipelineNameTypes
-{
-public:
-	/*FORCEINLINE*/ bool operator()(const FShaderPipelineType& A, const FShaderPipelineType& B) const
-	{
-		//#todo-rco: Avoid this by adding an FNullShaderPipelineType
-		bool bNullA = &A == nullptr;
-		bool bNullB = &B == nullptr;
-		if (bNullA && bNullB)
-		{
-			return false;
-		}
-		else if (bNullA)
-		{
-			return true;
-		}
-		else if (bNullB)
-		{
-			return false;
-		}
-
-
-		int32 AL = FCString::Strlen(A.GetName());
-		int32 BL = FCString::Strlen(B.GetName());
-		if (AL == BL)
-		{
-			return FCString::Strncmp(A.GetName(), B.GetName(), AL) > 0;
-		}
-		return AL > BL;
-	}
-};
-
 // A Shader Pipeline instance with compiled stages
 class FShaderPipeline
 {
