Fixing more Bindless ShaderBundle issues by restoring the CFLAG_ForceBindful flag for platforms that need it.

diff --git a/Engine/Source/Runtime/RenderCore/Private/ShaderBundles.cpp b/Engine/Source/Runtime/RenderCore/Private/ShaderBundles.cpp
--- a/Engine/Source/Runtime/RenderCore/Private/ShaderBundles.cpp
+++ b/Engine/Source/Runtime/RenderCore/Private/ShaderBundles.cpp
@@ -21,6 +21,10 @@ bool FDispatchShaderBundleCS::ShouldCompilePermutation(const FGlobalShaderPermut
 
 void FDispatchShaderBundleCS::ModifyCompilationEnvironment(const FGlobalShaderPermutationParameters& Parameters, FShaderCompilerEnvironment& OutEnvironment)
 {
+	if (FDataDrivenShaderPlatformInfo::GetRequiresBindfulUtilityShaders(Parameters.Platform))
+	{
+		OutEnvironment.CompilerFlags.Add(CFLAG_ForceBindful);
+	}
 	OutEnvironment.CompilerFlags.Add(CFLAG_RootConstants);
 
 	OutEnvironment.SetDefine(TEXT("THREADGROUP_SIZEX"), ThreadGroupSizeX);
