Add validation to catch issues where there is a miss match between the shader and its type.

- During compile shader we need to store out the target (VS, PS, CS, etc.)
- When we serialize the shader output from the Worker we double check the input matches the output.
- When we serialize the shader job (whether that is from the in-memory job cache or the DDC) we double check the input matches the output.

diff --git a/Engine/Shaders/Public/ShaderVersion.ush b/Engine/Shaders/Public/ShaderVersion.ush
--- a/Engine/Shaders/Public/ShaderVersion.ush
+++ b/Engine/Shaders/Public/ShaderVersion.ush
@@ -3,4 +3,4 @@
 // Each time the console command is executed it generates a new GUID. As this file is included
 // in Platform.ush (which should be included in any shader) it allows to invalidate the shader DDC.
 // 
-#pragma message("UESHADERMETADATA_VERSION 9DB6FF28-9CA1-4E83-97B3-9A355390EE10")
+#pragma message("UESHADERMETADATA_VERSION 60ae06fd-b60f-4424-80dd-25b05e4faaa9")
diff --git a/Engine/Source/Runtime/RenderCore/Private/ShaderCore.cpp b/Engine/Source/Runtime/RenderCore/Private/ShaderCore.cpp
--- a/Engine/Source/Runtime/RenderCore/Private/ShaderCore.cpp
+++ b/Engine/Source/Runtime/RenderCore/Private/ShaderCore.cpp
@@ -1689,6 +1689,8 @@ public:
 			if (Job.SecondaryOutput.IsValid())
 			{
 				Job.Output.bSucceeded = Job.Output.bSucceeded && Job.SecondaryOutput->bSucceeded;
+				// ensure the target field is set on the job output struct as we use it for validation during serialization
+				Job.SecondaryOutput->Target = Job.Input.Target;
 				if (Job.Output.bSucceeded)
 				{
 					Job.SecondaryOutput->GenerateOutputHash();
@@ -1701,6 +1703,8 @@ public:
 			Job.Output.bSucceeded = false;
 		}
 
+		// ensure the target field is set on the job output struct as we use it for validation during serialization
+		Job.Output.Target = Job.Input.Target;
 		if (Job.Output.bSucceeded)
 		{
 			Job.Output.GenerateOutputHash();
@@ -4010,6 +4014,7 @@ void FShaderCompileJob::SerializeOutput(FArchive& Ar)
 	Ar << Output;
 	// output hash is now serialized as part of the output, as the shader code is compressed in SCWs
 	checkf(!Output.bSucceeded || Output.OutputHash != FSHAHash(), TEXT("Successful compile job does not have an OutputHash generated."));
+	checkf(Output.Target == Input.Target, TEXT("Output FShaderTarget does not match the input struct; incorrect results associated with job?"));
 
 	if (Ar.IsLoading())
 	{
@@ -4097,6 +4102,12 @@ void FShaderCompileJob::SerializeWorkerOutput(FArchive& Ar)
 	bool bSucceededTemp = (bool)bSucceeded;
 	Ar << bSucceededTemp;
 	bSucceeded = bSucceededTemp;
+
+	checkf(Output.Target == Input.Target, TEXT("Output FShaderTarget does not match the input struct; incorrect results associated with job?"));
+	if (bSecondaryOutput)
+	{
+		checkf(SecondaryOutput->Target == Input.Target, TEXT("Output FShaderTarget does not match the input struct; incorrect results associated with job?"));
+	}
 }
 
 void FShaderCompileJob::SerializeWorkerInput(FArchive& Ar)
