Force Substrate's max closure count to 1 on SM5 platforms.

diff --git a/Engine/Config/Windows/BaseWindowsEngine.ini b/Engine/Config/Windows/BaseWindowsEngine.ini
--- a/Engine/Config/Windows/BaseWindowsEngine.ini
+++ b/Engine/Config/Windows/BaseWindowsEngine.ini
@@ -13,6 +13,7 @@ DefaultExternalCredentialTypeStr=Steam
 [PCD3D_SM5]
 BindlessResources=Disabled
 BindlessSamplers=Disabled
+MaxClosuresPerPixel=1
 
 [PCD3D_SM6]
 BindlessResources=RayTracingOnly
@@ -21,6 +22,7 @@ BindlessSamplers=RayTracingOnly
 [SF_VULKAN_SM5]
 BindlessResources=Disabled
 BindlessSamplers=Disabled
+MaxClosuresPerPixel=1
 
 [SF_VULKAN_SM6]
 ; Bindless is necessary for Vulkan ray tracing
diff --git a/Engine/Source/Runtime/RenderCore/Private/RenderUtils.cpp b/Engine/Source/Runtime/RenderCore/Private/RenderUtils.cpp
--- a/Engine/Source/Runtime/RenderCore/Private/RenderUtils.cpp
+++ b/Engine/Source/Runtime/RenderCore/Private/RenderUtils.cpp
@@ -20,7 +20,7 @@
 #include "Animation/MeshDeformerProvider.h"
 #include "Interfaces/ITargetPlatform.h"
 
-#if WITH_EDITOR
+#if WITH_EDITORONLY_DATA
 #include "Interfaces/ITargetPlatformManagerModule.h"
 #include "RHIShaderFormatDefinitions.inl"
 #endif
@@ -1889,8 +1889,38 @@ namespace Substrate
 	uint32 GetClosurePerPixel(EShaderPlatform InPlatform)
 	{
 		// Variant for shader compilation per platform
-		static FShaderPlatformCachedIniValue<int32> CVarClosureBudget(TEXT("r.Substrate.ClosuresPerPixel"));
-		return uint32(FMath::Max(0, CVarClosureBudget.Get(InPlatform)));
+		const TCHAR* CVarName = TEXT("r.Substrate.ClosuresPerPixel");
+		const TCHAR* ClosureName = TEXT("MaxClosuresPerPixel");
+		static FShaderPlatformCachedIniValue<int32> CVarClosureBudget(CVarName);
+		int32 OutClosurePerPixel = CVarClosureBudget.Get(InPlatform);
+
+		// Override r.Substrate.ClosuresPerPixel with ShaderFormat/MaxClosuresPerPixel for supporting sm5/sm6 differences
+		#if WITH_EDITORONLY_DATA
+		const FName ShaderFormatName = LegacyShaderPlatformToShaderFormat(InPlatform);
+		if (ITargetPlatform* TargetPlatform = GetTargetPlatformManager()->FindTargetPlatformWithSupport(TEXT("ShaderFormat"), ShaderFormatName))
+		{
+			if (FConfigCacheIni* PlatformConfig = TargetPlatform->GetConfigSystem())
+			{
+				const FString ShaderFormatStr = ShaderFormatName.ToString();
+				int32 ConfigClosurePerPixel = 0;
+				if (PlatformConfig->GetInt(*ShaderFormatStr, ClosureName, ConfigClosurePerPixel, GEngineIni))
+				{
+					OutClosurePerPixel = FMath::Min(OutClosurePerPixel, ConfigClosurePerPixel);
+				}
+			}
+		}
+		#else
+		static int32 ConfigClosurePerPixel = 0;
+		if (GConfig && ConfigClosurePerPixel == 0)
+		{
+			const FString ShaderFormatStr = LegacyShaderPlatformToShaderFormat(InPlatform).ToString();
+			const bool bExists = GConfig->GetInt(*ShaderFormatStr, ClosureName, ConfigClosurePerPixel, GEngineIni);
+			ConfigClosurePerPixel = bExists ? ConfigClosurePerPixel : OutClosurePerPixel;
+		}		
+		OutClosurePerPixel = FMath::Min(OutClosurePerPixel, ConfigClosurePerPixel);
+		#endif
+
+		return uint32(FMath::Max(1, OutClosurePerPixel));
 	}
 
 	uint32 GetNormalQuality()
