Flush any pending RDG async deletion tasks when shutting down the render core module.

diff --git a/Engine/Source/Runtime/RenderCore/Private/RenderCore.cpp b/Engine/Source/Runtime/RenderCore/Private/RenderCore.cpp
--- a/Engine/Source/Runtime/RenderCore/Private/RenderCore.cpp
+++ b/Engine/Source/Runtime/RenderCore/Private/RenderCore.cpp
@@ -19,6 +19,8 @@
 void UpdateShaderDevelopmentMode();
 
 void InitRenderGraph();
+void ShutdownRenderGraph();
+
 static void InitPixelRenderCounters();
 
 class FRenderCoreModule : public FDefaultModuleImpl
@@ -33,6 +35,11 @@ public:
 		InitRenderGraph();
 		InitPixelRenderCounters();
 	}
+
+	virtual void ShutdownModule() override
+	{
+		ShutdownRenderGraph();
+	}
 };
 
 IMPLEMENT_MODULE(FRenderCoreModule, RenderCore);
diff --git a/Engine/Source/Runtime/RenderCore/Private/RenderGraphPrivate.cpp b/Engine/Source/Runtime/RenderCore/Private/RenderGraphPrivate.cpp
--- a/Engine/Source/Runtime/RenderCore/Private/RenderGraphPrivate.cpp
+++ b/Engine/Source/Runtime/RenderCore/Private/RenderGraphPrivate.cpp
@@ -634,6 +634,11 @@ void InitRenderGraph()
 #endif
 }
 
+void ShutdownRenderGraph()
+{
+	FRDGBuilder::WaitForAsyncDeleteTask();
+}
+
 bool IsParallelExecuteEnabled()
 {
 	return GRDGParallelExecute > 0
