Fix some platforms not showing the GPU track in Insights
 - Platforms that do not support CPU/GPU timestamp calibration have not shown the GPU track since the DevPR copy up due to a bug in how discarded events were handled.
 - RDG builds the event tree on the render thread and inserts the timestamp query commands on parallel threads. Some of these events go unused due to pass culling etc, but the legacy GPU profiler couldn't handle this, hence needing to add support for "discarded" events.
 - GatherQueryResults wasn't correctly handling discarded events. We also need to bring back old logic for aligning the GPU frame to the start of the CPU frame, which was deleted in DevPR.
 - Also remove the ringbuffer of GPU profiler frames, and replace this with a queue, which prevents further issues with timestamp readback.

diff --git a/Engine/Source/Runtime/RenderCore/Private/ProfilingDebugging/RealtimeGPUProfiler.cpp b/Engine/Source/Runtime/RenderCore/Private/ProfilingDebugging/RealtimeGPUProfiler.cpp
--- a/Engine/Source/Runtime/RenderCore/Private/ProfilingDebugging/RealtimeGPUProfiler.cpp
+++ b/Engine/Source/Runtime/RenderCore/Private/ProfilingDebugging/RealtimeGPUProfiler.cpp
@@ -43,7 +43,6 @@ static TAutoConsoleVariable<int> CVarGPUStatsChildTimesIncluded(
 	TEXT("to the total GPU time, so we probably want this disabled.\n")
 );
 
-static const int32 NumGPUProfilerBufferedFrames = 4;
 static const uint64 InvalidQueryResult = 0xFFFFFFFFFFFFFFFFull;
 
 /*-----------------------------------------------------------------------------
@@ -102,18 +101,19 @@ public:
 	bool GatherQueryResults(FRHICommandListImmediate& RHICmdList)
 	{
 		//QUICK_SCOPE_CYCLE_COUNTER(STAT_SceneUtils_GatherQueryResults);
-		if (bDiscarded)
+		bool bUsed = bStarted && bEnded;
+		if (bDiscarded || !bUsed)
 		{
 			for (uint32 GPUIndex : GPUMask)
 			{
 				StartResultMicroseconds[GPUIndex] = 0;
 				EndResultMicroseconds[GPUIndex] = 0;
 			}
-			return true;
+
+			return bDiscarded;
 		}
 
 		// Get the query results which are still outstanding
-		check(GFrameNumberRenderThread != FrameNumber);
 		check(StartQuery.IsValid() && EndQuery.IsValid());
 
 		for (uint32 GPUIndex : GPUMask)
@@ -233,6 +233,11 @@ public:
 		return 2u;
 	}
 
+	bool IsDiscarded() const
+	{
+		return bDiscarded;
+	}
+
 	TStaticArray<uint64, MAX_NUM_GPUS> StartResultMicroseconds;
 	TStaticArray<uint64, MAX_NUM_GPUS> EndResultMicroseconds;
 
@@ -378,11 +383,8 @@ Container for a single frame's GPU stats
 class FRealtimeGPUProfilerFrame
 {
 public:
-	FRealtimeGPUProfilerFrame(FRenderQueryPoolRHIRef InRenderQueryPool, uint32& InQueryCount)
-		: NextEventIdx(1)
-		, NextResultPendingEventIdx(1)
-		, QueryCount(InQueryCount)
-		, RenderQueryPool(InRenderQueryPool)
+	FRealtimeGPUProfilerFrame(FRHIRenderQueryPool* InRenderQueryPool)
+		: RenderQueryPool(InRenderQueryPool)
 	{
 		GpuProfilerEvents.Empty(GPredictedMaxNumEvents);
 		for (uint32 Idx = 0; Idx < GPredictedMaxNumEvents; ++Idx)
@@ -390,8 +392,6 @@ public:
 			GpuProfilerEvents.Add(new FRealtimeGPUProfilerEvent(*RenderQueryPool));
 		}
 
-		QueryCount += (GPredictedMaxNumEvents - 1u) * FRealtimeGPUProfilerEvent::GetNumRHIQueriesPerEvent();
-
 		GpuProfilerEventParentIndices.Empty(GPredictedMaxNumEvents);
 		GpuProfilerEventParentIndices.AddUninitialized();
 
@@ -402,36 +402,17 @@ public:
 		EventAggregates.AddUninitialized();
 
 		DescriptionStringBuffer.Empty();
-	}
 
-	~FRealtimeGPUProfilerFrame()
-	{
-		QueryCount -= (GpuProfilerEvents.Num() - 1) * FRealtimeGPUProfilerEvent::GetNumRHIQueriesPerEvent();
+		CPUFrameStartTimestamp = FPlatformTime::Cycles64();
 	}
 
-	void Clear(void* Dummy)
-	{
-		NextEventIdx = 1;
-		NextResultPendingEventIdx = 1;
-
-		GpuProfilerEventParentIndices.Reset();
-		GpuProfilerEventParentIndices.AddUninitialized();
-
-		EventStack.Reset();
-		EventStack.Add(0);
-
-		EventAggregates.Reset();
-		EventAggregates.AddUninitialized();
-
-		DescriptionStringBuffer.Empty();
-	}
+	~FRealtimeGPUProfilerFrame() = default;
 
 	FRealtimeGPUProfilerQuery PushEvent(FRHIGPUMask GPUMask, const FName& Name, const TStatId& Stat, const TCHAR* Description)
 	{
 		if (NextEventIdx >= GpuProfilerEvents.Num())
 		{
 			GpuProfilerEvents.Add(new FRealtimeGPUProfilerEvent(*RenderQueryPool));
-			QueryCount += FRealtimeGPUProfilerEvent::GetNumRHIQueriesPerEvent();
 		}
 
 		const int32 EventIdx = NextEventIdx++;
@@ -463,18 +444,12 @@ public:
 		)
 	{
 		// Gather any remaining results and check all the results are ready
-		const int32 NumEventsThisFramePlusOne = NextEventIdx;
-
-		for (; NextResultPendingEventIdx < NumEventsThisFramePlusOne; ++NextResultPendingEventIdx)
+		for (; NextResultPendingEventIdx < NextEventIdx; ++NextResultPendingEventIdx)
 		{
 			FRealtimeGPUProfilerEvent& Event = GpuProfilerEvents[NextResultPendingEventIdx];
 
 			if (!Event.GatherQueryResults(RHICmdList))
 			{
-                // TODO: clloyd - Commmented out for now to stop spam on Mac
-//#if !(UE_BUILD_SHIPPING || UE_BUILD_TEST)
-				//UE_LOG(LogRendererCore, Warning, TEXT("Query '%s' not ready."), *Event.GetName().ToString());
-//#endif
 				// The frame isn't ready yet. Don't update stats - we'll try again next frame. 
 				return false;
 			}
@@ -503,7 +478,7 @@ public:
 		uint64 TotalUs = 0llu;
 		FNameSet StatSeenSet;
 
-		for (int32 Idx = 1; Idx < NumEventsThisFramePlusOne; ++Idx)
+		for (int32 Idx = 1; Idx < NextEventIdx; ++Idx)
 		{
 			const FRealtimeGPUProfilerEvent& Event = GpuProfilerEvents[Idx];
 			const FGPUEventTimeAggregate IncExcTime = EventAggregates[Idx];
@@ -556,7 +531,8 @@ public:
 
 #if GPUPROFILERTRACE_ENABLED
 		TArray<TArray<int32>> GpuProfilerEventChildrenIndices;
-		GpuProfilerEventChildrenIndices.AddDefaulted(GpuProfilerEvents.Num());
+		GpuProfilerEventChildrenIndices.AddDefaulted(NextEventIdx);
+
 		for (int32 EventIdx = 1; EventIdx < GpuProfilerEventParentIndices.Num(); ++EventIdx)
 		{
 			const int32 ParentIdx = GpuProfilerEventParentIndices[EventIdx];
@@ -579,14 +555,31 @@ public:
 
 			if (Timestamp.GPUMicroseconds == 0 || Timestamp.CPUMicroseconds == 0) // Unimplemented platforms, or invalid on the first frame
 			{
-				// Fallback to legacy
-				Timestamp = FGPUTiming::GetCalibrationTimestamp();
+				bool bSuccess = false;
+				for (int32 EventIdx = 1; EventIdx < NextEventIdx; ++EventIdx)
+				{
+					if (!GpuProfilerEvents[EventIdx].IsDiscarded())
+					{
+						// Align CPU and GPU frames
+						Timestamp.GPUMicroseconds = GpuProfilerEvents[EventIdx].GetStartResultMicroseconds(GPUIndex);
+						Timestamp.CPUMicroseconds = static_cast<uint64>(FPlatformTime::ToSeconds64(CPUFrameStartTimestamp) * 1000 * 1000);
+						bSuccess = true;
+						break;	
+					}
+				}
+
+				if (!bSuccess)
+				{
+					// Fallback to legacy
+					Timestamp = FGPUTiming::GetCalibrationTimestamp();
+				}
 			}
 		}		
 
 		// Sanitize event start/end times
-		TArray<TStaticArray<uint64, MAX_NUM_GPUS>> lastEndTimes;
-		lastEndTimes.AddZeroed(GpuProfilerEvents.Num());
+		TArray<TStaticArray<uint64, MAX_NUM_GPUS>> LastEndTimes;
+		LastEndTimes.AddZeroed(NextEventIdx);
+
 		for (int32 EventIdx = 1; EventIdx < GpuProfilerEventParentIndices.Num(); ++EventIdx)
 		{
 			const int32 ParentIdx = GpuProfilerEventParentIndices[EventIdx];
@@ -595,7 +588,7 @@ public:
 			for (uint32 GPUIndex : Event.GetGPUMask())
 			{
 				// Start time must be >= last end time
-				Event.StartResultMicroseconds[GPUIndex] = FMath::Max(Event.StartResultMicroseconds[GPUIndex], lastEndTimes[ParentIdx][GPUIndex]);
+				Event.StartResultMicroseconds[GPUIndex] = FMath::Max(Event.StartResultMicroseconds[GPUIndex], LastEndTimes[ParentIdx][GPUIndex]);
 				// End time must be >= start time
 				Event.EndResultMicroseconds[GPUIndex] = FMath::Max(Event.StartResultMicroseconds[GPUIndex], Event.EndResultMicroseconds[GPUIndex]);
 
@@ -613,7 +606,7 @@ public:
 				}
 
 				// Update last end time for this parent
-				lastEndTimes[ParentIdx][GPUIndex] = Event.EndResultMicroseconds[GPUIndex];
+				LastEndTimes[ParentIdx][GPUIndex] = Event.EndResultMicroseconds[GPUIndex];
 			}
 		}
 
@@ -700,6 +693,7 @@ public:
 		return true;
 	}
 
+	uint64 CPUFrameStartTimestamp;
 	FTimestampCalibrationQueryRHIRef TimestampCalibrationQuery;
 
 private:
@@ -817,10 +811,9 @@ private:
 		FName* SecondaryStore;
 	};
 
-	int32 NextEventIdx;
-	int32 NextResultPendingEventIdx;
+	int32 NextEventIdx = 1;
+	int32 NextResultPendingEventIdx = 1;
 
-	uint32& QueryCount;
 	FRenderQueryPoolRHIRef RenderQueryPool;
 
 	TIndirectArray<FRealtimeGPUProfilerEvent, TInlineAllocator<GPredictedMaxNumEvents>> GpuProfilerEvents;
@@ -854,29 +847,15 @@ void FRealtimeGPUProfiler::SafeRelease()
 
 
 FRealtimeGPUProfiler::FRealtimeGPUProfiler()
-	: WriteBufferIndex(0)
-	, ReadBufferIndex(1) 
-	, WriteFrameNumber(-1)
-	, bStatGatheringPaused(false)
-	, bInBeginEndBlock(false)
 {
 	if (GSupportsTimestampRenderQueries)
 	{
 		RenderQueryPool = RHICreateRenderQueryPool(RQT_AbsoluteTime);
-		for (int Index = 0; Index < NumGPUProfilerBufferedFrames; Index++)
-		{
-			Frames.Add(new FRealtimeGPUProfilerFrame(RenderQueryPool, QueryCount));
-		}
 	}
 }
 
 void FRealtimeGPUProfiler::Cleanup()
 {
-	for (int Index = 0; Index < Frames.Num(); Index++)
-	{
-		delete Frames[Index];
-	}
-	Frames.Empty();
 	RenderQueryPool.SafeRelease();
 	
 	TRACE_GPUPROFILER_DEINITIALIZE();
@@ -945,11 +924,13 @@ void FRealtimeGPUProfiler::BeginFrame(FRHICommandListImmediate& RHICmdList)
 		return;
 	}
 
-	check(bInBeginEndBlock == false);
-	bInBeginEndBlock = true;
+	check(!ActiveFrame.IsValid());
+	ActiveFrame = MakeUnique<FRealtimeGPUProfilerFrame>(RenderQueryPool);
+
+	ActiveFrame->TimestampCalibrationQuery = new FRHITimestampCalibrationQuery();
+	RHICmdList.CalibrateTimers(ActiveFrame->TimestampCalibrationQuery);
 
-	Frames[WriteBufferIndex]->TimestampCalibrationQuery = new FRHITimestampCalibrationQuery();
-	RHICmdList.CalibrateTimers(Frames[WriteBufferIndex]->TimestampCalibrationQuery);
+	ActiveFrame->CPUFrameStartTimestamp = FPlatformTime::Cycles64();
 }
 
 void FRealtimeGPUProfiler::EndFrame(FRHICommandListImmediate& RHICmdList)
@@ -959,31 +940,21 @@ void FRealtimeGPUProfiler::EndFrame(FRHICommandListImmediate& RHICmdList)
 		return;
 	}
 
-	// This is called at the end of the renderthread frame. Note that the RHI thread may still be processing commands for the frame at this point, however
-	// The read buffer index is always 3 frames beind the write buffer index in order to prevent us reading from the frame the RHI thread is still processing. 
-	// This should also ensure the GPU is done with the queries before we try to read them
-	check(Frames.Num() > 0);
-	check(IsInRenderingThread());
-	check(bInBeginEndBlock == true);
-	bInBeginEndBlock = false;
+	check(ActiveFrame.IsValid());
+	PendingFrames.Enqueue(MoveTemp(ActiveFrame));
 
-	if (Frames[ReadBufferIndex]->UpdateStats(RHICmdList
+	while (TUniquePtr<FRealtimeGPUProfilerFrame>* FramePtr = PendingFrames.Peek())
+	{
+		if (!(*FramePtr)->UpdateStats(RHICmdList
 #if GPUPROFILERTRACE_ENABLED
-		, HistoryByDescription
+			, HistoryByDescription
 #endif
 		))
-	{
-		// On a successful read, advance the ReadBufferIndex and WriteBufferIndex and clear the frame we just read
-		Frames[ReadBufferIndex]->Clear(&RHICmdList);
-		WriteFrameNumber = GFrameNumberRenderThread;
-		WriteBufferIndex = (WriteBufferIndex + 1) % Frames.Num();
-		ReadBufferIndex = (ReadBufferIndex + 1) % Frames.Num();
-		bStatGatheringPaused = false;
-	}
-	else
-	{
-		// The stats weren't ready; skip the next frame and don't advance the indices. We'll try to read the stats again next frame
-		bStatGatheringPaused = true;
+		{
+			break;
+		}
+
+		PendingFrames.Pop();
 	}
 }
 
@@ -993,38 +964,32 @@ void FRealtimeGPUProfiler::SuspendFrame()
 	{
 		return;
 	}
-
-	bStatGatheringPaused = true;
 }
 
 FRealtimeGPUProfilerQuery FRealtimeGPUProfiler::PushEvent(FRHIGPUMask GPUMask, const FName& Name, const TStatId& Stat, const TCHAR* Description)
 {
-	check(IsInParallelRenderingThread());
-	if (bStatGatheringPaused || !bInBeginEndBlock)
+	check(IsInRenderingThread());
+	if (ActiveFrame.IsValid())
 	{
-		return {};
+		return ActiveFrame->PushEvent(GPUMask, Name, Stat, Description);
 	}
-	check(Frames.Num() > 0);
-	if (WriteBufferIndex >= 0)
+	else
 	{
-		return Frames[WriteBufferIndex]->PushEvent(GPUMask, Name, Stat, Description);
+		return {};
 	}
-	return {};
 }
 
 FRealtimeGPUProfilerQuery FRealtimeGPUProfiler::PopEvent()
 {
-	check(IsInParallelRenderingThread());
-	if (bStatGatheringPaused || !bInBeginEndBlock)
+	check(IsInRenderingThread());
+	if (ActiveFrame.IsValid())
 	{
-		return {};
+		return ActiveFrame->PopEvent();
 	}
-	check(Frames.Num() > 0);
-	if (WriteBufferIndex >= 0)
+	else
 	{
-		return Frames[WriteBufferIndex]->PopEvent();
+		return {};
 	}
-	return {};
 }
 
 void FRealtimeGPUProfiler::PushStat(FRHICommandListImmediate& RHICmdList, const FName& Name, const TStatId& Stat, const TCHAR* Description)
diff --git a/Engine/Source/Runtime/RenderCore/Public/ProfilingDebugging/RealtimeGPUProfiler.h b/Engine/Source/Runtime/RenderCore/Public/ProfilingDebugging/RealtimeGPUProfiler.h
--- a/Engine/Source/Runtime/RenderCore/Public/ProfilingDebugging/RealtimeGPUProfiler.h
+++ b/Engine/Source/Runtime/RenderCore/Public/ProfilingDebugging/RealtimeGPUProfiler.h
@@ -322,17 +322,10 @@ private:
 	/** Deinitialize of the object*/
 	void Cleanup();
 
+	TUniquePtr<FRealtimeGPUProfilerFrame> ActiveFrame;
+	TQueue<TUniquePtr<FRealtimeGPUProfilerFrame>> PendingFrames;
 
-	/** Ringbuffer of profiler frames */
-	TArray<FRealtimeGPUProfilerFrame*> Frames;
-
-	int32 WriteBufferIndex;
-	int32 ReadBufferIndex;
-	uint32 WriteFrameNumber;
-	uint32 QueryCount = 0;
 	FRenderQueryPoolRHIRef RenderQueryPool;
-	bool bStatGatheringPaused;
-	bool bInBeginEndBlock;
 	bool bLocked = false;
 
 #if GPUPROFILERTRACE_ENABLED
