Backout CL 35490520 as the CSV GPU stats are now broken

diff --git a/Engine/Source/Runtime/RenderCore/Private/ProfilingDebugging/RealtimeGPUProfiler.cpp b/Engine/Source/Runtime/RenderCore/Private/ProfilingDebugging/RealtimeGPUProfiler.cpp
--- a/Engine/Source/Runtime/RenderCore/Private/ProfilingDebugging/RealtimeGPUProfiler.cpp
+++ b/Engine/Source/Runtime/RenderCore/Private/ProfilingDebugging/RealtimeGPUProfiler.cpp
@@ -43,6 +43,7 @@ static TAutoConsoleVariable<int> CVarGPUStatsChildTimesIncluded(
 	TEXT("to the total GPU time, so we probably want this disabled.\n")
 );
 
+static const int32 NumGPUProfilerBufferedFrames = 4;
 static const uint64 InvalidQueryResult = 0xFFFFFFFFFFFFFFFFull;
 
 /*-----------------------------------------------------------------------------
@@ -101,19 +102,18 @@ public:
 	bool GatherQueryResults(FRHICommandListImmediate& RHICmdList)
 	{
 		//QUICK_SCOPE_CYCLE_COUNTER(STAT_SceneUtils_GatherQueryResults);
-		bool bUsed = bStarted && bEnded;
-		if (bDiscarded || !bUsed)
+		if (bDiscarded)
 		{
 			for (uint32 GPUIndex : GPUMask)
 			{
 				StartResultMicroseconds[GPUIndex] = 0;
 				EndResultMicroseconds[GPUIndex] = 0;
 			}
-
-			return bDiscarded;
+			return true;
 		}
 
 		// Get the query results which are still outstanding
+		check(GFrameNumberRenderThread != FrameNumber);
 		check(StartQuery.IsValid() && EndQuery.IsValid());
 
 		for (uint32 GPUIndex : GPUMask)
@@ -233,11 +233,6 @@ public:
 		return 2u;
 	}
 
-	bool IsDiscarded() const
-	{
-		return bDiscarded;
-	}
-
 	TStaticArray<uint64, MAX_NUM_GPUS> StartResultMicroseconds;
 	TStaticArray<uint64, MAX_NUM_GPUS> EndResultMicroseconds;
 
@@ -383,8 +378,11 @@ Container for a single frame's GPU stats
 class FRealtimeGPUProfilerFrame
 {
 public:
-	FRealtimeGPUProfilerFrame(FRHIRenderQueryPool* InRenderQueryPool)
-		: RenderQueryPool(InRenderQueryPool)
+	FRealtimeGPUProfilerFrame(FRenderQueryPoolRHIRef InRenderQueryPool, uint32& InQueryCount)
+		: NextEventIdx(1)
+		, NextResultPendingEventIdx(1)
+		, QueryCount(InQueryCount)
+		, RenderQueryPool(InRenderQueryPool)
 	{
 		GpuProfilerEvents.Empty(GPredictedMaxNumEvents);
 		for (uint32 Idx = 0; Idx < GPredictedMaxNumEvents; ++Idx)
@@ -392,6 +390,8 @@ public:
 			GpuProfilerEvents.Add(new FRealtimeGPUProfilerEvent(*RenderQueryPool));
 		}
 
+		QueryCount += (GPredictedMaxNumEvents - 1u) * FRealtimeGPUProfilerEvent::GetNumRHIQueriesPerEvent();
+
 		GpuProfilerEventParentIndices.Empty(GPredictedMaxNumEvents);
 		GpuProfilerEventParentIndices.AddUninitialized();
 
@@ -402,17 +402,36 @@ public:
 		EventAggregates.AddUninitialized();
 
 		DescriptionStringBuffer.Empty();
+	}
 
-		CPUFrameStartTimestamp = FPlatformTime::Cycles64();
+	~FRealtimeGPUProfilerFrame()
+	{
+		QueryCount -= (GpuProfilerEvents.Num() - 1) * FRealtimeGPUProfilerEvent::GetNumRHIQueriesPerEvent();
 	}
 
-	~FRealtimeGPUProfilerFrame() = default;
+	void Clear(void* Dummy)
+	{
+		NextEventIdx = 1;
+		NextResultPendingEventIdx = 1;
+
+		GpuProfilerEventParentIndices.Reset();
+		GpuProfilerEventParentIndices.AddUninitialized();
+
+		EventStack.Reset();
+		EventStack.Add(0);
+
+		EventAggregates.Reset();
+		EventAggregates.AddUninitialized();
+
+		DescriptionStringBuffer.Empty();
+	}
 
 	FRealtimeGPUProfilerQuery PushEvent(FRHIGPUMask GPUMask, const FName& Name, const TStatId& Stat, const TCHAR* Description)
 	{
 		if (NextEventIdx >= GpuProfilerEvents.Num())
 		{
 			GpuProfilerEvents.Add(new FRealtimeGPUProfilerEvent(*RenderQueryPool));
+			QueryCount += FRealtimeGPUProfilerEvent::GetNumRHIQueriesPerEvent();
 		}
 
 		const int32 EventIdx = NextEventIdx++;
@@ -444,12 +463,18 @@ public:
 		)
 	{
 		// Gather any remaining results and check all the results are ready
-		for (; NextResultPendingEventIdx < NextEventIdx; ++NextResultPendingEventIdx)
+		const int32 NumEventsThisFramePlusOne = NextEventIdx;
+
+		for (; NextResultPendingEventIdx < NumEventsThisFramePlusOne; ++NextResultPendingEventIdx)
 		{
 			FRealtimeGPUProfilerEvent& Event = GpuProfilerEvents[NextResultPendingEventIdx];
 
 			if (!Event.GatherQueryResults(RHICmdList))
 			{
+                // TODO: clloyd - Commmented out for now to stop spam on Mac
+//#if !(UE_BUILD_SHIPPING || UE_BUILD_TEST)
+				//UE_LOG(LogRendererCore, Warning, TEXT("Query '%s' not ready."), *Event.GetName().ToString());
+//#endif
 				// The frame isn't ready yet. Don't update stats - we'll try again next frame. 
 				return false;
 			}
@@ -478,7 +503,7 @@ public:
 		uint64 TotalUs = 0llu;
 		FNameSet StatSeenSet;
 
-		for (int32 Idx = 1; Idx < NextEventIdx; ++Idx)
+		for (int32 Idx = 1; Idx < NumEventsThisFramePlusOne; ++Idx)
 		{
 			const FRealtimeGPUProfilerEvent& Event = GpuProfilerEvents[Idx];
 			const FGPUEventTimeAggregate IncExcTime = EventAggregates[Idx];
@@ -531,8 +556,7 @@ public:
 
 #if GPUPROFILERTRACE_ENABLED
 		TArray<TArray<int32>> GpuProfilerEventChildrenIndices;
-		GpuProfilerEventChildrenIndices.AddDefaulted(NextEventIdx);
-
+		GpuProfilerEventChildrenIndices.AddDefaulted(GpuProfilerEvents.Num());
 		for (int32 EventIdx = 1; EventIdx < GpuProfilerEventParentIndices.Num(); ++EventIdx)
 		{
 			const int32 ParentIdx = GpuProfilerEventParentIndices[EventIdx];
@@ -555,31 +579,14 @@ public:
 
 			if (Timestamp.GPUMicroseconds == 0 || Timestamp.CPUMicroseconds == 0) // Unimplemented platforms, or invalid on the first frame
 			{
-				bool bSuccess = false;
-				for (int32 EventIdx = 1; EventIdx < NextEventIdx; ++EventIdx)
-				{
-					if (!GpuProfilerEvents[EventIdx].IsDiscarded())
-					{
-						// Align CPU and GPU frames
-						Timestamp.GPUMicroseconds = GpuProfilerEvents[EventIdx].GetStartResultMicroseconds(GPUIndex);
-						Timestamp.CPUMicroseconds = static_cast<uint64>(FPlatformTime::ToSeconds64(CPUFrameStartTimestamp) * 1000 * 1000);
-						bSuccess = true;
-						break;	
-					}
-				}
-
-				if (!bSuccess)
-				{
-					// Fallback to legacy
-					Timestamp = FGPUTiming::GetCalibrationTimestamp();
-				}
+				// Fallback to legacy
+				Timestamp = FGPUTiming::GetCalibrationTimestamp();
 			}
 		}		
 
 		// Sanitize event start/end times
-		TArray<TStaticArray<uint64, MAX_NUM_GPUS>> LastEndTimes;
-		LastEndTimes.AddZeroed(NextEventIdx);
-
+		TArray<TStaticArray<uint64, MAX_NUM_GPUS>> lastEndTimes;
+		lastEndTimes.AddZeroed(GpuProfilerEvents.Num());
 		for (int32 EventIdx = 1; EventIdx < GpuProfilerEventParentIndices.Num(); ++EventIdx)
 		{
 			const int32 ParentIdx = GpuProfilerEventParentIndices[EventIdx];
@@ -588,7 +595,7 @@ public:
 			for (uint32 GPUIndex : Event.GetGPUMask())
 			{
 				// Start time must be >= last end time
-				Event.StartResultMicroseconds[GPUIndex] = FMath::Max(Event.StartResultMicroseconds[GPUIndex], LastEndTimes[ParentIdx][GPUIndex]);
+				Event.StartResultMicroseconds[GPUIndex] = FMath::Max(Event.StartResultMicroseconds[GPUIndex], lastEndTimes[ParentIdx][GPUIndex]);
 				// End time must be >= start time
 				Event.EndResultMicroseconds[GPUIndex] = FMath::Max(Event.StartResultMicroseconds[GPUIndex], Event.EndResultMicroseconds[GPUIndex]);
 
@@ -606,7 +613,7 @@ public:
 				}
 
 				// Update last end time for this parent
-				LastEndTimes[ParentIdx][GPUIndex] = Event.EndResultMicroseconds[GPUIndex];
+				lastEndTimes[ParentIdx][GPUIndex] = Event.EndResultMicroseconds[GPUIndex];
 			}
 		}
 
@@ -693,7 +700,6 @@ public:
 		return true;
 	}
 
-	uint64 CPUFrameStartTimestamp;
 	FTimestampCalibrationQueryRHIRef TimestampCalibrationQuery;
 
 private:
@@ -811,9 +817,10 @@ private:
 		FName* SecondaryStore;
 	};
 
-	int32 NextEventIdx = 1;
-	int32 NextResultPendingEventIdx = 1;
+	int32 NextEventIdx;
+	int32 NextResultPendingEventIdx;
 
+	uint32& QueryCount;
 	FRenderQueryPoolRHIRef RenderQueryPool;
 
 	TIndirectArray<FRealtimeGPUProfilerEvent, TInlineAllocator<GPredictedMaxNumEvents>> GpuProfilerEvents;
@@ -847,15 +854,29 @@ void FRealtimeGPUProfiler::SafeRelease()
 
 
 FRealtimeGPUProfiler::FRealtimeGPUProfiler()
+	: WriteBufferIndex(0)
+	, ReadBufferIndex(1) 
+	, WriteFrameNumber(-1)
+	, bStatGatheringPaused(false)
+	, bInBeginEndBlock(false)
 {
 	if (GSupportsTimestampRenderQueries)
 	{
 		RenderQueryPool = RHICreateRenderQueryPool(RQT_AbsoluteTime);
+		for (int Index = 0; Index < NumGPUProfilerBufferedFrames; Index++)
+		{
+			Frames.Add(new FRealtimeGPUProfilerFrame(RenderQueryPool, QueryCount));
+		}
 	}
 }
 
 void FRealtimeGPUProfiler::Cleanup()
 {
+	for (int Index = 0; Index < Frames.Num(); Index++)
+	{
+		delete Frames[Index];
+	}
+	Frames.Empty();
 	RenderQueryPool.SafeRelease();
 	
 	TRACE_GPUPROFILER_DEINITIALIZE();
@@ -924,13 +945,11 @@ void FRealtimeGPUProfiler::BeginFrame(FRHICommandListImmediate& RHICmdList)
 		return;
 	}
 
-	check(!ActiveFrame.IsValid());
-	ActiveFrame = MakeUnique<FRealtimeGPUProfilerFrame>(RenderQueryPool);
-
-	ActiveFrame->TimestampCalibrationQuery = new FRHITimestampCalibrationQuery();
-	RHICmdList.CalibrateTimers(ActiveFrame->TimestampCalibrationQuery);
+	check(bInBeginEndBlock == false);
+	bInBeginEndBlock = true;
 
-	ActiveFrame->CPUFrameStartTimestamp = FPlatformTime::Cycles64();
+	Frames[WriteBufferIndex]->TimestampCalibrationQuery = new FRHITimestampCalibrationQuery();
+	RHICmdList.CalibrateTimers(Frames[WriteBufferIndex]->TimestampCalibrationQuery);
 }
 
 void FRealtimeGPUProfiler::EndFrame(FRHICommandListImmediate& RHICmdList)
@@ -940,21 +959,31 @@ void FRealtimeGPUProfiler::EndFrame(FRHICommandListImmediate& RHICmdList)
 		return;
 	}
 
-	check(ActiveFrame.IsValid());
-	PendingFrames.Enqueue(MoveTemp(ActiveFrame));
+	// This is called at the end of the renderthread frame. Note that the RHI thread may still be processing commands for the frame at this point, however
+	// The read buffer index is always 3 frames beind the write buffer index in order to prevent us reading from the frame the RHI thread is still processing. 
+	// This should also ensure the GPU is done with the queries before we try to read them
+	check(Frames.Num() > 0);
+	check(IsInRenderingThread());
+	check(bInBeginEndBlock == true);
+	bInBeginEndBlock = false;
 
-	while (TUniquePtr<FRealtimeGPUProfilerFrame>* FramePtr = PendingFrames.Peek())
-	{
-		if (!(*FramePtr)->UpdateStats(RHICmdList
+	if (Frames[ReadBufferIndex]->UpdateStats(RHICmdList
 #if GPUPROFILERTRACE_ENABLED
-			, HistoryByDescription
+		, HistoryByDescription
 #endif
 		))
-		{
-			break;
-		}
-
-		PendingFrames.Pop();
+	{
+		// On a successful read, advance the ReadBufferIndex and WriteBufferIndex and clear the frame we just read
+		Frames[ReadBufferIndex]->Clear(&RHICmdList);
+		WriteFrameNumber = GFrameNumberRenderThread;
+		WriteBufferIndex = (WriteBufferIndex + 1) % Frames.Num();
+		ReadBufferIndex = (ReadBufferIndex + 1) % Frames.Num();
+		bStatGatheringPaused = false;
+	}
+	else
+	{
+		// The stats weren't ready; skip the next frame and don't advance the indices. We'll try to read the stats again next frame
+		bStatGatheringPaused = true;
 	}
 }
 
@@ -964,32 +993,38 @@ void FRealtimeGPUProfiler::SuspendFrame()
 	{
 		return;
 	}
+
+	bStatGatheringPaused = true;
 }
 
 FRealtimeGPUProfilerQuery FRealtimeGPUProfiler::PushEvent(FRHIGPUMask GPUMask, const FName& Name, const TStatId& Stat, const TCHAR* Description)
 {
-	check(IsInRenderingThread());
-	if (ActiveFrame.IsValid())
+	check(IsInParallelRenderingThread());
+	if (bStatGatheringPaused || !bInBeginEndBlock)
 	{
-		return ActiveFrame->PushEvent(GPUMask, Name, Stat, Description);
+		return {};
 	}
-	else
+	check(Frames.Num() > 0);
+	if (WriteBufferIndex >= 0)
 	{
-		return {};
+		return Frames[WriteBufferIndex]->PushEvent(GPUMask, Name, Stat, Description);
 	}
+	return {};
 }
 
 FRealtimeGPUProfilerQuery FRealtimeGPUProfiler::PopEvent()
 {
-	check(IsInRenderingThread());
-	if (ActiveFrame.IsValid())
+	check(IsInParallelRenderingThread());
+	if (bStatGatheringPaused || !bInBeginEndBlock)
 	{
-		return ActiveFrame->PopEvent();
+		return {};
 	}
-	else
+	check(Frames.Num() > 0);
+	if (WriteBufferIndex >= 0)
 	{
-		return {};
+		return Frames[WriteBufferIndex]->PopEvent();
 	}
+	return {};
 }
 
 void FRealtimeGPUProfiler::PushStat(FRHICommandListImmediate& RHICmdList, const FName& Name, const TStatId& Stat, const TCHAR* Description)
diff --git a/Engine/Source/Runtime/RenderCore/Public/ProfilingDebugging/RealtimeGPUProfiler.h b/Engine/Source/Runtime/RenderCore/Public/ProfilingDebugging/RealtimeGPUProfiler.h
--- a/Engine/Source/Runtime/RenderCore/Public/ProfilingDebugging/RealtimeGPUProfiler.h
+++ b/Engine/Source/Runtime/RenderCore/Public/ProfilingDebugging/RealtimeGPUProfiler.h
@@ -322,10 +322,17 @@ private:
 	/** Deinitialize of the object*/
 	void Cleanup();
 
-	TUniquePtr<FRealtimeGPUProfilerFrame> ActiveFrame;
-	TQueue<TUniquePtr<FRealtimeGPUProfilerFrame>> PendingFrames;
 
+	/** Ringbuffer of profiler frames */
+	TArray<FRealtimeGPUProfilerFrame*> Frames;
+
+	int32 WriteBufferIndex;
+	int32 ReadBufferIndex;
+	uint32 WriteFrameNumber;
+	uint32 QueryCount = 0;
 	FRenderQueryPoolRHIRef RenderQueryPool;
+	bool bStatGatheringPaused;
+	bool bInBeginEndBlock;
 	bool bLocked = false;
 
 #if GPUPROFILERTRACE_ENABLED
