Refactored reserved buffer commit logic in RDG to flush commits on the next use of the buffer in a pass.

diff --git a/Engine/Source/Runtime/RenderCore/Private/RenderGraphBuilder.cpp b/Engine/Source/Runtime/RenderCore/Private/RenderGraphBuilder.cpp
--- a/Engine/Source/Runtime/RenderCore/Private/RenderGraphBuilder.cpp
+++ b/Engine/Source/Runtime/RenderCore/Private/RenderGraphBuilder.cpp
@@ -1594,7 +1594,7 @@ void FRDGBuilder::ProcessAsyncSetupQueue()
 
 		for (FAsyncSetupOp Op : PoppedOps)
 		{
-			switch (Op.Type)
+			switch (Op.GetType())
 			{
 			case FAsyncSetupOp::EType::SetupPassResources:
 				SetupPassResources(Op.Pass);
@@ -1607,6 +1607,11 @@ void FRDGBuilder::ProcessAsyncSetupQueue()
 			case FAsyncSetupOp::EType::CullRootTexture:
 				AddCullRootTexture(Op.Texture);
 				break;
+
+			case FAsyncSetupOp::EType::ReservedBufferCommit:
+				ensureMsgf(!Op.Buffer->AccessModeState.IsExternalAccess(), TEXT("Buffer %s has a pending reserved commit of %u bytes but is marked for external access! The commit will be ignored!"), Op.Buffer->Name, Op.Payload);
+				Op.Buffer->PendingCommitSize = Op.Payload;
+				break;
 			}
 		}
 	}
@@ -2306,6 +2311,8 @@ void FRDGBuilder::SetupPassResources(FRDGPass* Pass)
 			Buffer->PassStateIndex = Pass->BufferStates.Num();
 
 			PassState = &Pass->BufferStates.Emplace_GetRef(Buffer);
+			PassState->State.ReservedCommitHandle = AcquireReservedCommitHandle(Buffer);
+			PassState->State.SetPass(PassPipeline, PassHandle);
 		}
 		else
 		{
@@ -2317,7 +2324,6 @@ void FRDGBuilder::SetupPassResources(FRDGPass* Pass)
 		PassState->ReferenceCount++;
 		PassState->State.Access = MakeValidAccess(PassState->State.Access, Access);
 		PassState->State.NoUAVBarrierFilter.AddHandle(NoUAVBarrierHandle);
-		PassState->State.SetPass(PassPipeline, PassHandle);
 
 		if (IsWritableAccess(Access))
 		{
@@ -3382,6 +3388,12 @@ void FRDGBuilder::CompilePassBarriers()
 				// Merge the pass state into the merged state.
 				ResourceMergeState->Access |= PassState->Access;
 
+				// If multiple reserved commits were requested, take the latest.
+				if (PassState->ReservedCommitHandle.IsValid())
+				{
+					ResourceMergeState->ReservedCommitHandle = PassState->ReservedCommitHandle;
+				}
+
 				FRDGPassHandle& FirstPassHandle = ResourceMergeState->FirstPass[PassPipeline];
 
 				if (FirstPassHandle.IsNull())
@@ -3478,7 +3490,6 @@ void FRDGBuilder::CollectPassBarriers(FRDGPassHandle PassHandle)
 			if (!Buffer->FirstState)
 			{
 				Buffer->FirstState = StateAfter;
-				Buffer->FirstState->bReservedCommit = Buffer->PendingCommitSize > 0;
 				return IsImmediateMode();
 			}
 			return true;
@@ -3502,28 +3513,28 @@ void FRDGBuilder::CreatePassBarriers()
 		for (FRDGTransitionInfo InfoRDG : BeginBatch->Transitions)
 		{
 			FRHITransitionInfo& InfoRHI = Context.Transitions.Emplace_GetRef();
-			InfoRHI.AccessBefore = InfoRDG.AccessBefore;
-			InfoRHI.AccessAfter  = InfoRDG.AccessAfter;
-			InfoRHI.ArraySlice   = InfoRDG.ArraySlice;
-			InfoRHI.MipIndex     = InfoRDG.MipIndex;
-			InfoRHI.PlaneSlice   = InfoRDG.PlaneSlice;
-			InfoRHI.Flags        = InfoRDG.Flags;
+			InfoRHI.AccessBefore = (ERHIAccess)InfoRDG.AccessBefore;
+			InfoRHI.AccessAfter  = (ERHIAccess)InfoRDG.AccessAfter;
+			InfoRHI.Flags        = (EResourceTransitionFlags)InfoRDG.ResourceTransitionFlags;
 
-			if (InfoRDG.Type == ERDGViewableResourceType::Texture)
+			if ((ERDGViewableResourceType)InfoRDG.ResourceType == ERDGViewableResourceType::Texture)
 			{
-				InfoRHI.Resource = Textures[FRDGTextureHandle(InfoRDG.Handle)]->ResourceRHI;
-				InfoRHI.Type = FRHITransitionInfo::EType::Texture;
+				InfoRHI.Resource   = Textures[FRDGTextureHandle(InfoRDG.ResourceHandle)]->ResourceRHI;
+				InfoRHI.Type       = FRHITransitionInfo::EType::Texture;
+				InfoRHI.ArraySlice = InfoRDG.Texture.ArraySlice;
+				InfoRHI.MipIndex   = InfoRDG.Texture.MipIndex;
+				InfoRHI.PlaneSlice = InfoRDG.Texture.PlaneSlice;
 			}
 			else
 			{
-				FRDGBuffer* Buffer = Buffers[FRDGBufferHandle(InfoRDG.Handle)];
+				FRDGBuffer* Buffer = Buffers[FRDGBufferHandle(InfoRDG.ResourceHandle)];
 
 				InfoRHI.Resource = Buffer->ResourceRHI;
 				InfoRHI.Type = FRHITransitionInfo::EType::Buffer;
 
-				if (InfoRDG.bReservedCommit)
+				if (InfoRDG.Buffer.CommitSize > 0)
 				{
-					InfoRHI.CommitInfo.Emplace(Buffer->PendingCommitSize);
+					InfoRHI.CommitInfo.Emplace(InfoRDG.Buffer.CommitSize);
 				}
 			}
 		}
@@ -3742,7 +3753,7 @@ void FRDGBuilder::AddLastBufferTransition(FRDGBuffer* Buffer)
 	check(IsImmediateMode() || Buffer->bExtracted || Buffer->ReferenceCount == FRDGViewableResource::DeallocatedReferenceCount);
 	check(Buffer->HasRHI());
 
-	if (Buffer->AccessModeState.ActiveMode == FRDGViewableResource::EAccessMode::External)
+	if (Buffer->AccessModeState.IsExternalAccess())
 	{
 		// Assign the final state that was enqueued by the external access pass, which may include merged states.
 		EpilogueResourceAccesses.Emplace(Buffer->GetRHI(), Buffer->State->Access);
@@ -3764,6 +3775,7 @@ void FRDGBuilder::AddLastBufferTransition(FRDGBuffer* Buffer)
 	{
 		StateAfter->SetPass(ERHIPipeline::Graphics, EpiloguePassHandle);
 		StateAfter->Access = Buffer->EpilogueAccess;
+		StateAfter->ReservedCommitHandle = AcquireReservedCommitHandle(Buffer);
 
 		EpilogueResourceAccesses.Emplace(Buffer->GetRHI(), StateAfter->Access);
 	}
@@ -3866,22 +3878,23 @@ void FRDGBuilder::AddTextureTransition(FRDGTexture* Texture, FRDGTextureSubresou
 			{
 				const FRDGTextureSubresource Subresource = Layout.GetSubresource(SubresourceIndex);
 
-				FRDGTransitionInfo Info;
-				Info.AccessBefore = SubresourceStateBefore->Access;
-				Info.AccessAfter = SubresourceStateAfter->Access;
-				Info.Handle = Texture->Handle.GetIndex();
-				Info.Type = ERDGViewableResourceType::Texture;
-				Info.Flags = SubresourceStateAfter->Flags;
-				Info.ArraySlice = Subresource.ArraySlice;
-				Info.MipIndex = Subresource.MipIndex;
-				Info.PlaneSlice = Subresource.PlaneSlice;
-				Info.bReservedCommit = 0;
-
-				if (Info.AccessBefore == ERHIAccess::Discard)
+				EResourceTransitionFlags Flags = SubresourceStateAfter->Flags;
+
+				if (SubresourceStateBefore->Access == ERHIAccess::Discard)
 				{
-					Info.Flags |= EResourceTransitionFlags::Discard;
+					Flags |= EResourceTransitionFlags::Discard;
 				}
 
+				FRDGTransitionInfo Info;
+				Info.AccessBefore            = (uint64)SubresourceStateBefore->Access;
+				Info.AccessAfter             = (uint64)SubresourceStateAfter->Access;
+				Info.ResourceHandle          = (uint64)Texture->Handle.GetIndex();
+				Info.ResourceType            = (uint64)ERDGViewableResourceType::Texture;
+				Info.ResourceTransitionFlags = (uint64)Flags;
+				Info.Texture.ArraySlice      = Subresource.ArraySlice;
+				Info.Texture.MipIndex        = Subresource.MipIndex;
+				Info.Texture.PlaneSlice      = Subresource.PlaneSlice;
+
 				AddTransition(Texture, *SubresourceStateBefore, *SubresourceStateAfter, Info);
 			}
 		}
@@ -3903,15 +3916,12 @@ void FRDGBuilder::AddBufferTransition(FRDGBufferRef Buffer, FRDGSubresourceState
 		if (FRDGSubresourceState::IsTransitionRequired(*StateBefore, *StateAfter))
 		{
 			FRDGTransitionInfo Info;
-			Info.AccessBefore = StateBefore->Access;
-			Info.AccessAfter = StateAfter->Access;
-			Info.Handle = Buffer->Handle.GetIndex();
-			Info.Type = ERDGViewableResourceType::Buffer;
-			Info.Flags = StateAfter->Flags;
-			Info.ArraySlice = 0;
-			Info.MipIndex = 0;
-			Info.PlaneSlice = 0;
-			Info.bReservedCommit = StateAfter->bReservedCommit;
+			Info.AccessBefore            = (uint64)StateBefore->Access;
+			Info.AccessAfter             = (uint64)StateAfter->Access;
+			Info.ResourceHandle          = (uint64)Buffer->Handle.GetIndex();
+			Info.ResourceType            = (uint64)ERDGViewableResourceType::Buffer;
+			Info.ResourceTransitionFlags = (uint64)StateAfter->Flags;
+			Info.Buffer.CommitSize       = GetReservedCommitSize(StateAfter->ReservedCommitHandle);
 
 			AddTransition(Buffer, *StateBefore, *StateAfter, Info);
 		}
diff --git a/Engine/Source/Runtime/RenderCore/Private/RenderGraphResources.cpp b/Engine/Source/Runtime/RenderCore/Private/RenderGraphResources.cpp
--- a/Engine/Source/Runtime/RenderCore/Private/RenderGraphResources.cpp
+++ b/Engine/Source/Runtime/RenderCore/Private/RenderGraphResources.cpp
@@ -45,7 +45,6 @@ FRDGViewableResource::FRDGViewableResource(const TCHAR* InName, const ERDGViewab
 	, bSplitFirstTransition(bInSplitFirstTransition)
 	, bQueuedForUpload(0)
 	, bCollectForAllocate(1)
-	, bQueuedForReservedCommit(0)
 #if !(UE_BUILD_SHIPPING || UE_BUILD_TEST)
 	, bHeapAllocatedDebugName(0)
 #endif
diff --git a/Engine/Source/Runtime/RenderCore/Private/RenderGraphValidation.cpp b/Engine/Source/Runtime/RenderCore/Private/RenderGraphValidation.cpp
--- a/Engine/Source/Runtime/RenderCore/Private/RenderGraphValidation.cpp
+++ b/Engine/Source/Runtime/RenderCore/Private/RenderGraphValidation.cpp
@@ -1422,10 +1422,10 @@ void FRDGBarrierValidation::ValidateBarrierBatchBegin(const FRDGPass* Pass, cons
 		for (const FRDGTransitionInfo& Transition : Transitions)
 		{
 			UE_LOG(LogRDG, Display, TEXT("\t\tMip(%d), Array(%d), Slice(%d): [%s, %s] -> [%s, %s]"),
-				Transition.MipIndex, Transition.ArraySlice, Transition.PlaneSlice,
-				*GetRHIAccessName(Transition.AccessBefore),
+				Transition.Texture.MipIndex, Transition.Texture.ArraySlice, Transition.Texture.PlaneSlice,
+				*GetRHIAccessName((ERHIAccess)Transition.AccessBefore),
 				*GetRHIPipelineName(Batch.PipelinesToBegin),
-				*GetRHIAccessName(Transition.AccessAfter),
+				*GetRHIAccessName((ERHIAccess)Transition.AccessAfter),
 				*GetRHIPipelineName(Batch.PipelinesToEnd));
 		}
 	}
@@ -1446,9 +1446,9 @@ void FRDGBarrierValidation::ValidateBarrierBatchBegin(const FRDGPass* Pass, cons
 			Buffer,
 			Buffer->GetRHIUnchecked(),
 			Buffer->Name,
-			*GetRHIAccessName(Transition.AccessBefore),
+			*GetRHIAccessName((ERHIAccess)Transition.AccessBefore),
 			*GetRHIPipelineName(Batch.PipelinesToBegin),
-			*GetRHIAccessName(Transition.AccessAfter),
+			*GetRHIAccessName((ERHIAccess)Transition.AccessAfter),
 			*GetRHIPipelineName(Batch.PipelinesToEnd));
 	}
 }
diff --git a/Engine/Source/Runtime/RenderCore/Public/RenderGraphBuilder.h b/Engine/Source/Runtime/RenderCore/Public/RenderGraphBuilder.h
--- a/Engine/Source/Runtime/RenderCore/Public/RenderGraphBuilder.h
+++ b/Engine/Source/Runtime/RenderCore/Public/RenderGraphBuilder.h
@@ -859,38 +859,56 @@ private:
 
 	struct FAsyncSetupOp
 	{
-		enum class EType
+		enum class EType : uint8
 		{
 			SetupPassResources,
 			CullRootBuffer,
-			CullRootTexture
+			CullRootTexture,
+			ReservedBufferCommit
 		};
 
 		static FAsyncSetupOp SetupPassResources(FRDGPass* Pass)
 		{
-			FAsyncSetupOp Op;
-			Op.Type = EType::SetupPassResources;
+			FAsyncSetupOp Op(EType::SetupPassResources);
 			Op.Pass = Pass;
 			return Op;
 		}
 
 		static FAsyncSetupOp CullRootBuffer(FRDGBuffer* Buffer)
 		{
-			FAsyncSetupOp Op;
-			Op.Type = EType::CullRootBuffer;
+			FAsyncSetupOp Op(EType::CullRootBuffer);
 			Op.Buffer = Buffer;
 			return Op;
 		}
 
 		static FAsyncSetupOp CullRootTexture(FRDGTexture* Texture)
 		{
-			FAsyncSetupOp Op;
-			Op.Type = EType::CullRootTexture;
+			FAsyncSetupOp Op(EType::CullRootTexture);
 			Op.Texture = Texture;
 			return Op;
 		}
 
-		EType Type;
+		static FAsyncSetupOp ReservedBufferCommit(FRDGBuffer* Buffer, uint64 CommitSizeInBytes)
+		{
+			FAsyncSetupOp Op(EType::ReservedBufferCommit, CommitSizeInBytes);
+			Op.Buffer = Buffer;
+			return Op;
+		}
+
+		EType GetType() const
+		{
+			return (EType)Type;
+		}
+
+		uint64 Type : 8;
+		uint64 Payload : 48;
+
+		FAsyncSetupOp(EType InType, uint64 InPayload = 0)
+			: Type((uint8)InType)
+			, Payload(InPayload)
+		{
+			check(InPayload < (1ull << 48ull));
+		}
 
 		union
 		{
@@ -918,6 +936,30 @@ private:
 	void LaunchAsyncSetupQueueTask();
 	void ProcessAsyncSetupQueue();
 
+	//////////////////////////////////////////////////////////////////////////////
+	// Reserved Buffer Commits
+
+	FRDGBufferReservedCommitHandle AcquireReservedCommitHandle(FRDGBuffer* Buffer)
+	{
+		FRDGBufferReservedCommitHandle Handle;
+
+		if (Buffer->PendingCommitSize > 0)
+		{
+			Handle = FRDGBufferReservedCommitHandle(ReservedBufferCommitSizes.Num());
+			ReservedBufferCommitSizes.Emplace(Buffer->PendingCommitSize);
+			Buffer->PendingCommitSize = 0;
+		}
+
+		return Handle;
+	}
+
+	uint64 GetReservedCommitSize(FRDGBufferReservedCommitHandle Handle)
+	{
+		return Handle.IsValid() ? ReservedBufferCommitSizes[Handle.GetIndex()] : 0;
+	}
+
+	TArray<uint64, FRDGArrayAllocator> ReservedBufferCommitSizes;
+
 	//////////////////////////////////////////////////////////////////////////////
 	// Culling
 
diff --git a/Engine/Source/Runtime/RenderCore/Public/RenderGraphBuilder.inl b/Engine/Source/Runtime/RenderCore/Public/RenderGraphBuilder.inl
--- a/Engine/Source/Runtime/RenderCore/Public/RenderGraphBuilder.inl
+++ b/Engine/Source/Runtime/RenderCore/Public/RenderGraphBuilder.inl
@@ -387,9 +387,8 @@ inline void FRDGBuilder::QueueCommitReservedBuffer(FRDGBufferRef Buffer, uint64
 {
 	IF_RDG_ENABLE_DEBUG(UserValidation.ValidateCommitBuffer(Buffer, CommitSizeInBytes));
 
-	Buffer->PendingCommitSize = FMath::Max<uint64>(CommitSizeInBytes, Buffer->PendingCommitSize);
-	Buffer->PooledBuffer->SetCommittedSize(Buffer->PendingCommitSize);
-	Buffer->bQueuedForReservedCommit = 1;
+	AsyncSetupQueue.Push(FAsyncSetupOp::ReservedBufferCommit(Buffer, CommitSizeInBytes));
+	Buffer->PooledBuffer->SetCommittedSize(CommitSizeInBytes);
 }
 
 inline void FRDGBuilder::QueueTextureExtraction(FRDGTextureRef Texture, TRefCountPtr<IPooledRenderTarget>* OutTexturePtr, ERHIAccess AccessFinal, ERDGResourceExtractionFlags Flags)
diff --git a/Engine/Source/Runtime/RenderCore/Public/RenderGraphDefinitions.h b/Engine/Source/Runtime/RenderCore/Public/RenderGraphDefinitions.h
--- a/Engine/Source/Runtime/RenderCore/Public/RenderGraphDefinitions.h
+++ b/Engine/Source/Runtime/RenderCore/Public/RenderGraphDefinitions.h
@@ -554,31 +554,34 @@ public:
 	void Reset()
 	{
 		Handle = HandleType::Null;
-		bUnique = false;
 	}
 
 	void AddHandle(HandleType InHandle)
 	{
+		checkf(InHandle != NotUniqueHandle, TEXT("Overflowed TRDGHandleUniqueFilter"));
+
 		if (Handle != InHandle && InHandle.IsValid())
 		{
-			bUnique = Handle.IsNull();
-			Handle = InHandle;
+			Handle = Handle.IsNull() ? InHandle : NotUniqueHandle;
 		}
 	}
 
 	HandleType GetUniqueHandle() const
 	{
-		return bUnique ? Handle : HandleType::Null;
+		return Handle != NotUniqueHandle ? Handle : HandleType::Null;
 	}
 
 private:
+	static const HandleType NotUniqueHandle;
 	HandleType Handle;
-	bool bUnique = false;
 };
 
 template <typename ObjectType, typename IndexType>
 const TRDGHandle<ObjectType, IndexType> TRDGHandle<ObjectType, IndexType>::Null;
 
+template <typename HandleType>
+const HandleType TRDGHandleUniqueFilter<HandleType>::NotUniqueHandle(TNumericLimits<typename HandleType::IndexType>::Max() - 1);
+
 struct FRDGTextureDesc : public FRHITextureDesc
 {
 	static FRDGTextureDesc Create2D(
@@ -716,6 +719,7 @@ using FRDGTextureRegistry = TRDGHandleRegistry<FRDGTextureHandle, ERDGHandleRegi
 using FRDGTextureBitArray = TRDGHandleBitArray<FRDGTextureHandle>;
 
 using FRDGBufferHandle = TRDGHandle<FRDGBuffer, uint32>;
+using FRDGBufferReservedCommitHandle = TRDGHandle<FRDGBuffer, uint16>;
 using FRDGBufferRegistry = TRDGHandleRegistry<FRDGBufferHandle, ERDGHandleRegistryDestructPolicy::Registry>;
 using FRDGBufferBitArray = TRDGHandleBitArray<FRDGBufferHandle>;
 
diff --git a/Engine/Source/Runtime/RenderCore/Public/RenderGraphPass.h b/Engine/Source/Runtime/RenderCore/Public/RenderGraphPass.h
--- a/Engine/Source/Runtime/RenderCore/Public/RenderGraphPass.h
+++ b/Engine/Source/Runtime/RenderCore/Public/RenderGraphPass.h
@@ -69,17 +69,30 @@ enum class ERDGBarrierLocation : uint8
 
 struct FRDGTransitionInfo
 {
-	FRDGTransitionInfo() = default;
-
-	ERHIAccess AccessBefore;
-	ERHIAccess AccessAfter;
-	uint16 Handle;
-	ERDGViewableResourceType Type;
-	EResourceTransitionFlags Flags;
-	uint32 ArraySlice : 16;
-	uint32 MipIndex   : 8;
-	uint32 PlaneSlice : 2;
-	uint32 bReservedCommit : 1;
+	static_assert((int32)ERHIAccess::Last <= (1 << 20) && (int32)ERDGViewableResourceType::MAX <= 3 && (int32)EResourceTransitionFlags::Last == (1 << 2), "FRDGTransitionInfo packing is no longer correct.");
+
+	uint64 AccessBefore            : 21; // 21
+	uint64 AccessAfter             : 21; // 42
+	uint64 ResourceHandle          : 16; // 58
+	uint64 ResourceType            : 3;  // 61
+	uint64 ResourceTransitionFlags : 3;  // 64
+
+	union
+	{
+		struct
+		{
+			uint16 ArraySlice;
+			uint8  MipIndex;
+			uint8  PlaneSlice;
+
+		} Texture;
+
+		struct
+		{
+			uint64 CommitSize;
+
+		} Buffer;
+	};
 };
 
 struct FRDGBarrierBatchEndId
diff --git a/Engine/Source/Runtime/RenderCore/Public/RenderGraphResources.h b/Engine/Source/Runtime/RenderCore/Public/RenderGraphResources.h
--- a/Engine/Source/Runtime/RenderCore/Public/RenderGraphResources.h
+++ b/Engine/Source/Runtime/RenderCore/Public/RenderGraphResources.h
@@ -63,17 +63,13 @@ struct FRDGSubresourceState
 	/** Given a before and after state, returns whether they can be merged into a single state. */
 	static bool IsMergeAllowed(ERDGViewableResourceType ResourceType, const FRDGSubresourceState& Previous, const FRDGSubresourceState& Next);
 
-	FRDGSubresourceState()
-		: bReservedCommit(0)
-	{}
+	FRDGSubresourceState() = default;
 
 	explicit FRDGSubresourceState(ERHIAccess InAccess)
 		: Access(InAccess)
-		, bReservedCommit(0)
 	{}
 
 	explicit FRDGSubresourceState(ERHIPipeline Pipeline, FRDGPassHandle PassHandle)
-		: bReservedCommit(0)
 	{
 		SetPass(Pipeline, PassHandle);
 	}
@@ -108,11 +104,11 @@ struct FRDGSubresourceState
 	/** The last no-UAV barrier to be used by this subresource. */
 	FRDGViewUniqueFilter NoUAVBarrierFilter;
 
+	/** Whether this subresource state represents a commit operation for a reserved resource. */
+	FRDGBufferReservedCommitHandle ReservedCommitHandle;
+
 	/** The last used transition flags on the pass. */
 	EResourceTransitionFlags Flags = EResourceTransitionFlags::None;
-
-	/** Whether this subresource state represents a commit operation for a reserved resource. */
-	uint8 bReservedCommit : 1;
 };
 
 using FRDGTextureSubresourceState = TRDGTextureSubresourceArray<FRDGSubresourceState*, FRDGArrayAllocator>;
@@ -412,9 +408,6 @@ protected:
 	/** If false, the resource needs to be collected. */
 	uint8 bCollectForAllocate : 1;
 
-	/** If true, the reserved resource is having tiles committed. */
-	uint8 bQueuedForReservedCommit : 1;
-
 #if !(UE_BUILD_SHIPPING || UE_BUILD_TEST)
 	/** If true, the debug name passed to the constructor was heap allocated, and needs to be freed in the destructor. */
 	uint8 bHeapAllocatedDebugName : 1;
