Discard the active and pending GPU profiler frames on shutdown
 - RHI resources must be cleaned up before calling RHIExit()

diff --git a/Engine/Source/Runtime/RenderCore/Private/ProfilingDebugging/RealtimeGPUProfiler.cpp b/Engine/Source/Runtime/RenderCore/Private/ProfilingDebugging/RealtimeGPUProfiler.cpp
--- a/Engine/Source/Runtime/RenderCore/Private/ProfilingDebugging/RealtimeGPUProfiler.cpp
+++ b/Engine/Source/Runtime/RenderCore/Private/ProfilingDebugging/RealtimeGPUProfiler.cpp
@@ -856,6 +856,9 @@ FRealtimeGPUProfiler::FRealtimeGPUProfiler()
 
 void FRealtimeGPUProfiler::Cleanup()
 {
+	ActiveFrame.Reset();
+	PendingFrames.Empty();
+
 	RenderQueryPool.SafeRelease();
 	
 	TRACE_GPUPROFILER_DEINITIALIZE();
