Refresh registered ray tracing geometry when FRayTracingGeometry::ReleaseRHI() is called.
- otherwise the geometry stays in EvictableGeometries and causes a fail to check.

diff --git a/Engine/Source/Runtime/RenderCore/Private/RayTracingGeometry.cpp b/Engine/Source/Runtime/RenderCore/Private/RayTracingGeometry.cpp
--- a/Engine/Source/Runtime/RenderCore/Private/RayTracingGeometry.cpp
+++ b/Engine/Source/Runtime/RenderCore/Private/RayTracingGeometry.cpp
@@ -240,6 +240,8 @@ void FRayTracingGeometry::ReleaseRHI()
 	RemoveBuildRequest();
 	RayTracingGeometryRHI.SafeRelease();
 	GeometryState = EGeometryStateFlags::Invalid;
+
+	GRayTracingGeometryManager->RefreshRegisteredGeometry(RayTracingGeometryHandle);
 }
 
 void FRayTracingGeometry::RemoveBuildRequest()
