Reintroduced flush of RHI command lists if bDispatchAfterExecute is true on a pass. Without it, the RHI thread spends too much time processing large batches of work, causing GPU bubbles.

diff --git a/Engine/Source/Runtime/RenderCore/Private/RenderGraphBuilder.cpp b/Engine/Source/Runtime/RenderCore/Private/RenderGraphBuilder.cpp
--- a/Engine/Source/Runtime/RenderCore/Private/RenderGraphBuilder.cpp
+++ b/Engine/Source/Runtime/RenderCore/Private/RenderGraphBuilder.cpp
@@ -1975,18 +1975,30 @@ void FRDGBuilder::Execute()
 					PrevSerialPass = nullptr;
 				}
 
+				bool bDispatchAfterExecute = false;
+
 				if (Pass->bDispatchPass)
 				{
 					FRDGDispatchPass* DispatchPass = static_cast<FRDGDispatchPass*>(Pass);
 					DispatchPass->CommandListsEvent.Wait();
 					QueuedCmdLists.Append(MoveTemp(DispatchPass->CommandLists));
+
+					bDispatchAfterExecute = Pass->bDispatchAfterExecute;
 				}
 				else if (Pass->bParallelExecuteBegin)
 				{
 					FParallelPassSet& ParallelPassSet = ParallelExecute.ParallelPassSets[Pass->ParallelPassSetIndex];
 					check(ParallelPassSet.CmdList != nullptr);
 					QueuedCmdLists.Add(ParallelPassSet);
+
+					bDispatchAfterExecute = ParallelPassSet.bDispatchAfterExecute;
 				}
+
+				if (bDispatchAfterExecute)
+				{
+					FlushParallel();
+					RHICmdList.ImmediateFlush(EImmediateFlushType::DispatchToRHIThread);
+				};
 			}
 			else
 			{
