Removed legacy dispatch after execute calls before / after async compute.

diff --git a/Engine/Source/Runtime/RenderCore/Private/RenderGraphBuilder.cpp b/Engine/Source/Runtime/RenderCore/Private/RenderGraphBuilder.cpp
--- a/Engine/Source/Runtime/RenderCore/Private/RenderGraphBuilder.cpp
+++ b/Engine/Source/Runtime/RenderCore/Private/RenderGraphBuilder.cpp
@@ -1497,12 +1497,6 @@ void FRDGBuilder::Compile()
 
 				AsyncComputePass->bAsyncComputeBegin = 1;
 				AsyncComputePass->GetPrologueBarriersToEnd(Allocators.Transition).AddDependency(&EpilogueBarriersToBeginForAsyncCompute);
-
-				// Since we are fencing the graphics pipe to some new async compute work, make sure to flush any prior work.
-				if (AsyncComputePassBeforeFork)
-				{
-					AsyncComputePassBeforeFork->bDispatchAfterExecute = 1;
-				}
 			}
 
 			AsyncComputePassBeforeFork = AsyncComputePass;
@@ -1551,7 +1545,6 @@ void FRDGBuilder::Compile()
 				FRDGBarrierBatchBegin& EpilogueBarriersToBeginForGraphics = AsyncComputePass->GetEpilogueBarriersToBeginForGraphics(Allocators.Transition, TransitionCreateQueue);
 
 				AsyncComputePass->bAsyncComputeEnd = 1;
-				AsyncComputePass->bDispatchAfterExecute = 1;
 				EpilogueBarriersToBeginForGraphics.SetUseCrossPipelineFence();
 
 				GraphicsJoinPass->bGraphicsJoin = 1;
