Fixed crash due to D3D12 recording commands into multiple swap chains. Dispatch work to the RHI after each window renders.

diff --git a/Engine/Source/Runtime/RenderCore/Private/RenderGraphBuilder.cpp b/Engine/Source/Runtime/RenderCore/Private/RenderGraphBuilder.cpp
--- a/Engine/Source/Runtime/RenderCore/Private/RenderGraphBuilder.cpp
+++ b/Engine/Source/Runtime/RenderCore/Private/RenderGraphBuilder.cpp
@@ -1497,12 +1497,6 @@ void FRDGBuilder::Compile()
 
 				AsyncComputePass->bAsyncComputeBegin = 1;
 				AsyncComputePass->GetPrologueBarriersToEnd(Allocators.Transition).AddDependency(&EpilogueBarriersToBeginForAsyncCompute);
-
-				// Since we are fencing the graphics pipe to some new async compute work, make sure to flush any prior work.
-				if (AsyncComputePassBeforeFork)
-				{
-					AsyncComputePassBeforeFork->bDispatchAfterExecute = 1;
-				}
 			}
 
 			AsyncComputePassBeforeFork = AsyncComputePass;
@@ -1551,7 +1545,6 @@ void FRDGBuilder::Compile()
 				FRDGBarrierBatchBegin& EpilogueBarriersToBeginForGraphics = AsyncComputePass->GetEpilogueBarriersToBeginForGraphics(Allocators.Transition, TransitionCreateQueue);
 
 				AsyncComputePass->bAsyncComputeEnd = 1;
-				AsyncComputePass->bDispatchAfterExecute = 1;
 				EpilogueBarriersToBeginForGraphics.SetUseCrossPipelineFence();
 
 				GraphicsJoinPass->bGraphicsJoin = 1;
diff --git a/Engine/Source/Runtime/RenderCore/Public/RenderGraphBuilder.inl b/Engine/Source/Runtime/RenderCore/Public/RenderGraphBuilder.inl
--- a/Engine/Source/Runtime/RenderCore/Public/RenderGraphBuilder.inl
+++ b/Engine/Source/Runtime/RenderCore/Public/RenderGraphBuilder.inl
@@ -493,7 +493,11 @@ inline void FRDGBuilder::QueueBufferExtraction(FRDGBufferRef Buffer, TRefCountPt
 
 inline void FRDGBuilder::AddDispatchHint()
 {
-	if (Passes.Num() > 0)
+	if (IsImmediateMode())
+	{
+		RHICmdList.ImmediateFlush(EImmediateFlushType::DispatchToRHIThread);
+	}
+	else if (Passes.Num() > 0)
 	{
 		Passes[Passes.Last()]->bDispatchAfterExecute = 1;
 	}
diff --git a/Engine/Source/Runtime/SlateRHIRenderer/Private/SlateRHIRenderer.cpp b/Engine/Source/Runtime/SlateRHIRenderer/Private/SlateRHIRenderer.cpp
--- a/Engine/Source/Runtime/SlateRHIRenderer/Private/SlateRHIRenderer.cpp
+++ b/Engine/Source/Runtime/SlateRHIRenderer/Private/SlateRHIRenderer.cpp
@@ -920,6 +920,9 @@ FSlateDrawWindowPassOutputs FSlateRHIRenderer::DrawWindow_RenderThread(FRDGBuild
 		}
 	}
 
+	// Dispatch work after each window since certain RHI's can't record commands for multiple swap chains.
+	GraphBuilder.AddDispatchHint();
+
 	FSlateDrawWindowPassOutputs Outputs;
 	Outputs.ViewportRHI = ViewportInfo.ViewportRHI;
 	Outputs.ViewportTextureRHI = ViewportTextureRHI;
