[shaders]
- add inputhash-based validation when reading & writing worker output files; clear the field on the output struct after validating so it's not serialized to DDC
- ensure values required by validation are set on skipped pipeline jobs when an earlier pipeline phase fails

diff --git a/Engine/Shaders/Public/ShaderVersion.ush b/Engine/Shaders/Public/ShaderVersion.ush
--- a/Engine/Shaders/Public/ShaderVersion.ush
+++ b/Engine/Shaders/Public/ShaderVersion.ush
@@ -3,4 +3,4 @@
 // Each time the console command is executed it generates a new GUID. As this file is included
 // in Platform.ush (which should be included in any shader) it allows to invalidate the shader DDC.
 // 
-#pragma message("UESHADERMETADATA_VERSION E153BDBF-C281-40FF-BCD1-3A139BC52C71")
\ No newline at end of file
+#pragma message("UESHADERMETADATA_VERSION 10C02C57-1251-46A8-A27D-B5C80CD940EF")
\ No newline at end of file
diff --git a/Engine/Source/Runtime/RenderCore/Private/ShaderCore.cpp b/Engine/Source/Runtime/RenderCore/Private/ShaderCore.cpp
--- a/Engine/Source/Runtime/RenderCore/Private/ShaderCore.cpp
+++ b/Engine/Source/Runtime/RenderCore/Private/ShaderCore.cpp
@@ -1691,6 +1691,8 @@ public:
 				Job.Output.bSucceeded = Job.Output.bSucceeded && Job.SecondaryOutput->bSucceeded;
 				// ensure the target field is set on the job output struct as we use it for validation during serialization
 				Job.SecondaryOutput->Target = Job.Input.Target;
+				// also set the job input hash on the output struct for validation purposes
+				Job.SecondaryOutput->ValidateInputHash = Job.Input.Hash;
 				if (Job.Output.bSucceeded)
 				{
 					Job.SecondaryOutput->GenerateOutputHash();
@@ -1705,6 +1707,8 @@ public:
 
 		// ensure the target field is set on the job output struct as we use it for validation during serialization
 		Job.Output.Target = Job.Input.Target;
+		// also set the job input hash on the output struct for validation purposes
+		Job.Output.ValidateInputHash = Job.Input.Hash;
 		if (Job.Output.bSucceeded)
 		{
 			Job.Output.GenerateOutputHash();
@@ -1820,6 +1824,13 @@ void CompileShaderPipeline(const TArray<const IShaderFormat*>& ShaderFormats, FS
 	if (!CurrentJob->Output.bSucceeded)
 	{
 		// Can't carry on compiling the pipeline
+		// Set values used for validation on the pipeline jobs that we're skipping before returning
+		for (int32 Index = 1; Index < PipelineJob->StageJobs.Num(); ++Index)
+		{
+			FShaderCompileJob& Job = *PipelineJob->StageJobs[Index];
+			Job.Output.Target = Job.Input.Target;
+			Job.Output.ValidateInputHash = Job.Input.Hash;
+		}
 		return;
 	}
 
@@ -3948,8 +3959,8 @@ FShaderCompilerInputHash FShaderCompileJob::GetInputHash()
 
 	FMemoryHasherBlake3 Hasher;
 
-	int32 FShaderCompilerOutputStructVersionLocal = FShaderCompilerOutputStructVersion;
-	Hasher << FShaderCompilerOutputStructVersionLocal;
+	FGuid ShaderCacheVersionLocal = UE_SHADER_CACHE_VERSION;
+	Hasher << ShaderCacheVersionLocal;
 
 	uint32 FormatVersion = GetTargetPlatformManagerRef().ShaderFormatVersion(Input.ShaderFormat);
 	Hasher << FormatVersion;
@@ -4090,6 +4101,10 @@ void FShaderCompileJob::SerializeWorkerOutput(FArchive& Ar)
 	bool bSecondaryOutput = SecondaryOutput.IsValid();
 	Ar << bSecondaryOutput;
 
+	checkf(Output.ValidateInputHash == Input.Hash, TEXT("Output.ValidateInputHash does not match Input.Hash; incorrect results associated with job?"));
+	// empty validation hash after running validation to avoid impacting output deduplication
+	Output.ValidateInputHash = FShaderCompilerInputHash();
+
 	if (bSecondaryOutput)
 	{
 		if (Ar.IsLoading())
@@ -4097,6 +4112,9 @@ void FShaderCompileJob::SerializeWorkerOutput(FArchive& Ar)
 			SecondaryOutput = MakeUnique<FShaderCompilerOutput>();
 		}
 		Ar << *SecondaryOutput;
+		checkf(SecondaryOutput->ValidateInputHash == Input.Hash, TEXT("SecondaryOutput.ValidateInputHash does not match Input.Hash; incorrect results associated with job?"));
+		// empty validation hash after running validation to avoid impacting output deduplication
+		SecondaryOutput->ValidateInputHash = FShaderCompilerInputHash();
 	}
 
 	bool bSucceededTemp = (bool)bSucceeded;
@@ -4106,7 +4124,7 @@ void FShaderCompileJob::SerializeWorkerOutput(FArchive& Ar)
 	checkf(Output.Target == Input.Target, TEXT("Output FShaderTarget does not match the input struct; incorrect results associated with job?"));
 	if (bSecondaryOutput)
 	{
-		checkf(SecondaryOutput->Target == Input.Target, TEXT("Output FShaderTarget does not match the input struct; incorrect results associated with job?"));
+		checkf(SecondaryOutput->Target == Input.Target, TEXT("SecondaryOutput FShaderTarget does not match the input struct; incorrect results associated with job?"));
 	}
 }
 
diff --git a/Engine/Source/Runtime/RenderCore/Public/ShaderCompilerCore.h b/Engine/Source/Runtime/RenderCore/Public/ShaderCompilerCore.h
--- a/Engine/Source/Runtime/RenderCore/Public/ShaderCompilerCore.h
+++ b/Engine/Source/Runtime/RenderCore/Public/ShaderCompilerCore.h
@@ -33,14 +33,15 @@ typedef TSharedPtr<TArray<ANSICHAR>, ESPMode::ThreadSafe> FShaderSharedAnsiStrin
 // this is for the protocol, not the data, bump if FShaderCompilerInput/FShaderPreprocessOutput serialization, SerializeWorkerInput or ProcessInputFromArchive changes.
 inline const int32 ShaderCompileWorkerInputVersion = 28;
 // this is for the protocol, not the data, bump if FShaderCompilerOutput or WriteToOutputArchive changes.
-inline const int32 ShaderCompileWorkerOutputVersion = 22;
+inline const int32 ShaderCompileWorkerOutputVersion = 23;
 // this is for the protocol, not the data.
 inline const int32 ShaderCompileWorkerSingleJobHeader = 'S';
 // this is for the protocol, not the data.
 inline const int32 ShaderCompileWorkerPipelineJobHeader = 'P';
 
-// modify this for changes to the FShaderCompilerOutput data structure (in addition to ShaderCompileWorkerOutputVersion)
-inline const int32 FShaderCompilerOutputStructVersion = 4;
+// Modify this to invalidate _just_ the cache/DDC entries for individual shaders (will not cause shadermaps to rebuild if they are not otherwise out-of-date).
+// This should be bumped for changes to the FShaderCompilerOutput data structure (in addition to ShaderCompileWorkerOutputVersion)
+inline static const FGuid UE_SHADER_CACHE_VERSION = FGuid("DB7CA032-63BC-4D95-8F65-64C6B39EF882");
 
 namespace UE::ShaderCompiler
 {
@@ -553,7 +554,7 @@ struct FShaderCompilerError
 
 /**
  *	The output of the shader compiler.
- *	Bump FShaderCompilerOutputStructVersion and ShaderCompileWorkerOutputVersion if FShaderCompilerOutput changes
+ *	Bump UE_SHADER_CACHE_VERSION and ShaderCompileWorkerOutputVersion if FShaderCompilerOutput changes
  */
 struct FShaderCompilerOutput
 {
@@ -574,6 +575,7 @@ struct FShaderCompilerOutput
 	FShaderTarget Target;
 	FShaderCode ShaderCode;
 	FSHAHash OutputHash;
+	FShaderCompilerInputHash ValidateInputHash;
 	uint32 NumInstructions;
 	uint32 NumTextureSamplers;
 	double CompileTime;
@@ -656,6 +658,7 @@ struct FShaderCompilerOutput
 		}
 		Ar << Output.PlatformDebugData;
 		Ar << Output.ShaderStatistics;
+		Ar << Output.ValidateInputHash;
 
 		return Ar;
 	}
