[shaders] minor refactor of pipeline job compile function to reduce code duplication; this has the side effect of correctly setting validation fields for all cases where pipeline jobs could be skipped (instead of just the case where the first job failed which would have been the case with the previous code)

diff --git a/Engine/Source/Runtime/RenderCore/Private/ShaderCore.cpp b/Engine/Source/Runtime/RenderCore/Private/ShaderCore.cpp
--- a/Engine/Source/Runtime/RenderCore/Private/ShaderCore.cpp
+++ b/Engine/Source/Runtime/RenderCore/Private/ShaderCore.cpp
@@ -1809,30 +1809,6 @@ void CompileShader(const TArray<const IShaderFormat*>& ShaderFormats, FShaderCom
 void CompileShaderPipeline(const TArray<const IShaderFormat*>& ShaderFormats, FShaderPipelineCompileJob* PipelineJob, const FString& Dir, int32* CompileCount)
 {
 	checkf(PipelineJob->StageJobs.Num() > 0, TEXT("Pipeline %s has zero jobs!"), PipelineJob->Key.ShaderPipeline->GetName());
-	FShaderCompileJob* CurrentJob = PipelineJob->StageJobs[0]->GetSingleShaderJob();
-
-	// Flag should be set on the first job when the FShaderPipelineCompileJob was constructed, to ensure the flag is included when computing the input hash.
-	check(CurrentJob->Input.bCompilingForShaderPipeline == true);
-
-	// First job doesn't have to trim outputs
-	CurrentJob->Input.bIncludeUsedOutputs = false;
-
-	// Compile the shader directly through the platform dll (directly from the shader dir as the working directory)
-	CompileShader(ShaderFormats, *CurrentJob, Dir, CompileCount);
-
-	CurrentJob->bSucceeded = CurrentJob->Output.bSucceeded;
-	if (!CurrentJob->Output.bSucceeded)
-	{
-		// Can't carry on compiling the pipeline
-		// Set values used for validation on the pipeline jobs that we're skipping before returning
-		for (int32 Index = 1; Index < PipelineJob->StageJobs.Num(); ++Index)
-		{
-			FShaderCompileJob& Job = *PipelineJob->StageJobs[Index];
-			Job.Output.Target = Job.Input.Target;
-			Job.Output.ValidateInputHash = Job.Input.Hash;
-		}
-		return;
-	}
 
 	// This tells the shader compiler we do want to remove unused outputs
 	bool bEnableRemovingUnused = true;
@@ -1848,12 +1824,21 @@ void CompileShaderPipeline(const TArray<const IShaderFormat*>& ShaderFormats, FS
 		}
 	}
 
-	for (int32 Index = 1; Index < PipelineJob->StageJobs.Num(); ++Index)
+	FShaderCompileJob* CurrentJob = nullptr;
+	for (int32 Index = 0; Index < PipelineJob->StageJobs.Num(); ++Index)
 	{
 		auto* PreviousJob = CurrentJob;
 		CurrentJob = PipelineJob->StageJobs[Index]->GetSingleShaderJob();
-		bEnableRemovingUnused = bEnableRemovingUnused && PreviousJob->Output.bSupportsQueryingUsedAttributes;
-		if (bEnableRemovingUnused)
+
+		bool bFirstJob = PreviousJob == nullptr;
+		if (bFirstJob)
+		{
+			CurrentJob->Input.bIncludeUsedOutputs = false;
+			// Flag should be set on the first job when the FShaderPipelineCompileJob was constructed, to ensure the flag is included when computing the input hash.
+			check(CurrentJob->Input.bCompilingForShaderPipeline == true);
+		}
+
+		if (bEnableRemovingUnused && !bFirstJob && PreviousJob->Output.bSupportsQueryingUsedAttributes)
 		{
 			CurrentJob->Input.bIncludeUsedOutputs = true;
 			CurrentJob->Input.bCompilingForShaderPipeline = true;
@@ -1867,6 +1852,13 @@ void CompileShaderPipeline(const TArray<const IShaderFormat*>& ShaderFormats, FS
 		if (!CurrentJob->Output.bSucceeded)
 		{
 			// Can't carry on compiling the pipeline
+			// Set values used for validation on the pipeline jobs that we're skipping before returning
+			for (int32 SkipIndex = Index + 1; SkipIndex < PipelineJob->StageJobs.Num(); SkipIndex++)
+			{
+				FShaderCompileJob& Job = *PipelineJob->StageJobs[SkipIndex];
+				Job.Output.Target = Job.Input.Target;
+				Job.Output.ValidateInputHash = Job.Input.Hash;
+			}
 			return;
 		}
 	}
