Fix constness of FRDGBuilder::AllocParameters copy source argument

#rnx

diff --git a/Engine/Source/Runtime/RenderCore/Public/RenderGraphBuilder.h b/Engine/Source/Runtime/RenderCore/Public/RenderGraphBuilder.h
--- a/Engine/Source/Runtime/RenderCore/Public/RenderGraphBuilder.h
+++ b/Engine/Source/Runtime/RenderCore/Public/RenderGraphBuilder.h
@@ -195,7 +195,7 @@ public:
 
 	/** Allocates a parameter struct with a lifetime tied to graph execution, and copies contents from an existing parameters struct. */
 	template <typename ParameterStructType>
-	ParameterStructType* AllocParameters(ParameterStructType* StructToCopy);
+	ParameterStructType* AllocParameters(const ParameterStructType* StructToCopy);
 
 	/** Allocates a data-driven parameter struct with a lifetime tied to graph execution. */
 	template <typename BaseParameterStructType>
diff --git a/Engine/Source/Runtime/RenderCore/Public/RenderGraphBuilder.inl b/Engine/Source/Runtime/RenderCore/Public/RenderGraphBuilder.inl
--- a/Engine/Source/Runtime/RenderCore/Public/RenderGraphBuilder.inl
+++ b/Engine/Source/Runtime/RenderCore/Public/RenderGraphBuilder.inl
@@ -169,7 +169,7 @@ FORCEINLINE ParameterStructType* FRDGBuilder::AllocParameters()
 }
 
 template <typename ParameterStructType>
-FORCEINLINE ParameterStructType* FRDGBuilder::AllocParameters(ParameterStructType* StructToCopy)
+FORCEINLINE ParameterStructType* FRDGBuilder::AllocParameters(const ParameterStructType* StructToCopy)
 {
 	ParameterStructType* Struct = Allocators.Root.Alloc<ParameterStructType>();
 	*Struct = *StructToCopy;
