Prepare RHI breadcrumbs and GPU event stream for "stat gpu" support
 - Added a TStatId field on breadcrumbs. When set, this indicates the breadcrumb should write its computed GPU duration to the given stat.
 - Implemented "stat gpu" alongside the "stat unit" GPU event stream sink. Times in "stat gpu" are now taken as the union of busy time across all GPU queues, in the same way we compute the "stat unit" GPU time.
 - RDG support is handled via the new RHI_EVENT_SCOPE_STAT macro, allowing us to tag RDG scopes with GPU stats. When the new GPU profiler is enabled, RDG_GPU_STAT_SCOPE and SCOPED_GPU_STAT are empty and will eventually be deprecated and removed.

Cleanup of breadcrumb macros
 - Remove unnecessary "Name" arguments.
 - Require the user to wrap the format string in quotes, rather than stringizing the format arg with the preprocessor.
 - Removed "F" version of macros. Both string literal and formatted strings can be handled with the varargs macro, since the varargs are simply empty when using only a string literal.

Added GetTypeHash function for TStatId
 - This didn't exist before, and allows use of TStatId in TMaps etc.

diff --git a/Engine/Source/Runtime/Core/Public/Stats/Stats2.h b/Engine/Source/Runtime/Core/Public/Stats/Stats2.h
--- a/Engine/Source/Runtime/Core/Public/Stats/Stats2.h
+++ b/Engine/Source/Runtime/Core/Public/Stats/Stats2.h
@@ -279,6 +279,11 @@ struct TStatId
 		return StatIdPtr != Other.StatIdPtr;
 	}
 
+	friend uint32 GetTypeHash(TStatId StatId)
+	{
+		return GetTypeHash(StatId.StatIdPtr);
+	}
+
 private:
 	/** NAME_None. */
 	CORE_API static TStatIdData TStatId_NAME_None;
diff --git a/Engine/Source/Runtime/D3D12RHI/Private/D3D12Allocation.cpp b/Engine/Source/Runtime/D3D12RHI/Private/D3D12Allocation.cpp
--- a/Engine/Source/Runtime/D3D12RHI/Private/D3D12Allocation.cpp
+++ b/Engine/Source/Runtime/D3D12RHI/Private/D3D12Allocation.cpp
@@ -1623,7 +1623,7 @@ void FD3D12DefaultBufferAllocator::BeginFrame(FD3D12ContextArray const& Contexts
 		FD3D12CommandContext& CommandContext = *Contexts[ERHIPipeline::Graphics]->GetSingleDeviceContext(GetParentDevice()->GetGPUIndex());
 
 		TRACE_CPUPROFILER_EVENT_SCOPE(FlushPendingBufferCopyOps);
-		RHI_BREADCRUMB_EVENT(CommandContext, BufferPoolCopyOps);
+		RHI_BREADCRUMB_EVENT(CommandContext, "BufferPoolCopyOps");
 
 		for (FD3D12BufferPool* DefaultBufferPool : DefaultBufferPools)
 		{
@@ -1823,7 +1823,7 @@ void FD3D12TextureAllocatorPool::BeginFrame(FD3D12ContextArray const& Contexts)
 		FD3D12CommandContext& CommandContext = *Contexts[ERHIPipeline::Graphics]->GetSingleDeviceContext(GetParentDevice()->GetGPUIndex());
 
 		TRACE_CPUPROFILER_EVENT_SCOPE(FlushPendingTextureCopyOps);
-		RHI_BREADCRUMB_EVENT(CommandContext, TexturePoolCopyOps);
+		RHI_BREADCRUMB_EVENT(CommandContext, "TexturePoolCopyOps");
 		
 		for (uint32 PoolIndex = 0; PoolIndex < (uint32)EPoolType::Count; ++PoolIndex)
 		{
diff --git a/Engine/Source/Runtime/D3D12RHI/Private/D3D12Texture.cpp b/Engine/Source/Runtime/D3D12RHI/Private/D3D12Texture.cpp
--- a/Engine/Source/Runtime/D3D12RHI/Private/D3D12Texture.cpp
+++ b/Engine/Source/Runtime/D3D12RHI/Private/D3D12Texture.cpp
@@ -2117,7 +2117,7 @@ public:
 				FD3D12Resource* UploadBuffer = UpdateInfo.SrcResourceLocation->GetResource();
 				CD3DX12_TEXTURE_COPY_LOCATION SourceCopyLocation(UploadBuffer->GetResource(), UpdateInfo.PlacedSubresourceFootprint);
 
-				RHI_BREADCRUMB_EVENT(Context, EndMultiUpdateTexture3D);
+				RHI_BREADCRUMB_EVENT(Context, "EndMultiUpdateTexture3D");
 
 				Context.GraphicsCommandList()->CopyTextureRegion(
 					&DestCopyLocation,
@@ -2327,7 +2327,7 @@ void FD3D12DynamicRHI::EndUpdateTexture3D_Internal(FRHICommandListBase& RHICmdLi
 			for (FD3D12Texture& Texture : *FD3D12DynamicRHI::ResourceCast(RootTexture))
 			{
 				FD3D12CommandContext& Context = FD3D12CommandContext::Get(ExecutingCmdList, Texture.GetParentDevice()->GetGPUIndex());
-				RHI_BREADCRUMB_EVENT(Context, EndUpdateTexture3D);
+				RHI_BREADCRUMB_EVENT(Context, "EndUpdateTexture3D");
 
 				CD3DX12_TEXTURE_COPY_LOCATION DestCopyLocation(Texture.GetResource()->GetResource(), MipIdx);
 				CD3DX12_TEXTURE_COPY_LOCATION SourceCopyLocation(SrcResourceLocation->GetResource()->GetResource(), PlacedSubresourceFootprint);
diff --git a/Engine/Source/Runtime/Engine/Private/KismetRenderingLibrary.cpp b/Engine/Source/Runtime/Engine/Private/KismetRenderingLibrary.cpp
--- a/Engine/Source/Runtime/Engine/Private/KismetRenderingLibrary.cpp
+++ b/Engine/Source/Runtime/Engine/Private/KismetRenderingLibrary.cpp
@@ -803,7 +803,7 @@ void UKismetRenderingLibrary::BeginDrawCanvasToRenderTarget(UObject* WorldContex
 				](FRHICommandListImmediate& RHICmdList)
 			{
 #if  WANTS_DRAW_MESH_EVENTS
-				Breadcrumb->Emplace(RHICmdList, TEXT("DrawCanvasToTarget: %s"), Name);
+				Breadcrumb->Emplace(RHICmdList, TStatId(), TEXT("DrawCanvasToTarget: %s"), Name);
 #endif
 				RenderTargetResource->FlushDeferredResourceUpdate(RHICmdList);
 			});
diff --git a/Engine/Source/Runtime/Engine/Private/UnrealEngine.cpp b/Engine/Source/Runtime/Engine/Private/UnrealEngine.cpp
--- a/Engine/Source/Runtime/Engine/Private/UnrealEngine.cpp
+++ b/Engine/Source/Runtime/Engine/Private/UnrealEngine.cpp
@@ -10709,7 +10709,7 @@ bool UEngine::PerformError(const TCHAR* Cmd, FOutputDevice& Ar)
 			static void RHIThreadCrash(FRHICommandList& CmdList)
 			{
 				// Test breadcrumb system
-				RHI_BREADCRUMB_EVENTF(CmdList, Event, TEXT("Debug RHI Thread Crash"));
+				RHI_BREADCRUMB_EVENT(CmdList, "Debug RHI Thread Crash");
 
 				CmdList.EnqueueLambda([](FRHICommandListBase& ExecutingCmdList)
 				{
diff --git a/Engine/Source/Runtime/Launch/Private/LaunchEngineLoop.cpp b/Engine/Source/Runtime/Launch/Private/LaunchEngineLoop.cpp
--- a/Engine/Source/Runtime/Launch/Private/LaunchEngineLoop.cpp
+++ b/Engine/Source/Runtime/Launch/Private/LaunchEngineLoop.cpp
@@ -5493,12 +5493,12 @@ static inline void BeginFrameRenderThread(FRHICommandListImmediate& RHICmdList,
 #if CSV_PROFILER
 	if (FCsvProfiler::Get()->IsCapturing_Renderthread())
 	{
-		GRHIFrameBreadcrumb.Emplace(RHICmdList, TEXT("CsvFrame %d"), FCsvProfiler::Get()->GetCaptureFrameNumberRT());
+		GRHIFrameBreadcrumb.Emplace(RHICmdList, TStatId(), TEXT("CsvFrame %d"), FCsvProfiler::Get()->GetCaptureFrameNumberRT());
 	}
 	else
 #endif
 	{
-		GRHIFrameBreadcrumb.Emplace(RHICmdList, TEXT("Frame %d"), CurrentFrameCounter);
+		GRHIFrameBreadcrumb.Emplace(RHICmdList, TStatId(), TEXT("Frame %d"), CurrentFrameCounter);
 	}
 #endif
 
diff --git a/Engine/Source/Runtime/RHI/Private/GPUProfiler.cpp b/Engine/Source/Runtime/RHI/Private/GPUProfiler.cpp
--- a/Engine/Source/Runtime/RHI/Private/GPUProfiler.cpp
+++ b/Engine/Source/Runtime/RHI/Private/GPUProfiler.cpp
@@ -706,29 +706,128 @@ namespace UE::RHI::GPUProfiler
 		}
 	}
 
-	struct FGPUProfilerSink_StatUnit final : public FEventSink
+	// Handles computing the "stat unit" GPU time, and "stat gpu" stats.
+	struct FGPUProfilerSink_StatSystem final : public FEventSink
 	{
-		struct FQueueTimestamps : public TArray<uint64>
+		class FTimestampStream
 		{
-			int32 TimestampIndex = 0;
-			uint64 BusyCycles = 0;
+		private:
+			TArray<uint64> Values;
+
+		public:
+			struct FState
+			{
+				FTimestampStream const& Stream;
+				int32 TimestampIndex = 0;
+				uint64 BusyCycles = 0;
+
+				FState(FTimestampStream const& Stream)
+					: Stream(Stream)
+				{}
+
+				uint64 GetCurrentTimestamp (uint64 Anchor) const { return Stream.Values[TimestampIndex] - Anchor; }
+				uint64 GetPreviousTimestamp(uint64 Anchor) const { return Stream.Values[TimestampIndex - 1] - Anchor; }
+
+				bool HasMoreTimestamps() const { return TimestampIndex < Stream.Values.Num(); }
+				bool IsStartingWork   () const { return (TimestampIndex & 0x01) == 0x00; }
+				void AdvanceTimestamp () { TimestampIndex++; }
+			};
+
+			void AddTimestamp(uint64 Value, bool bBegin)
+			{
+				if (bBegin)
+				{
+					if (!Values.IsEmpty() && Value <= Values.Last())
+					{
+						//
+						// The Begin TOP event is sooner than the last End BOP event.
+						// The markers overlap, and the GPU was not idle.
+						// 
+						// Remove the previous End event, and discard this Begin event.
+						//
+						Values.RemoveAt(Values.Num() - 1, EAllowShrinking::No);
+					}
+					else
+					{
+						// GPU was idle. Keep this timestamp.
+						Values.Add(Value);
+					}
+				}
+				else
+				{
+					Values.Add(Value);
+				}
+			}
+
+			static uint64 ComputeUnion(TArrayView<FTimestampStream::FState> Streams)
+			{
+				// The total number of cycles where at least one GPU pipe was busy.
+				uint64 UnionBusyCycles = 0;
+
+				uint64 LastMinCycles = 0;
+				int32 BusyPipes = 0;
+				bool bFirst = true;
+
+				uint64 Anchor = 0; // @todo - handle possible timestamp wraparound
+
+				// Process the time ranges from each pipe.
+				while (true)
+				{
+					// Find the next minimum timestamp
+					FTimestampStream::FState* NextMin = nullptr;
+					for (auto& Current : Streams)
+					{
+						if (Current.HasMoreTimestamps() && (!NextMin || Current.GetCurrentTimestamp(Anchor) < NextMin->GetCurrentTimestamp(Anchor)))
+						{
+							NextMin = &Current;
+						}
+					}
+
+					if (!NextMin)
+						break; // No more timestamps to process
+
+					if (!bFirst)
+					{
+						if (BusyPipes > 0 && NextMin->GetCurrentTimestamp(Anchor) > LastMinCycles)
+						{
+							// Accumulate the union busy time across all pipes
+							UnionBusyCycles += NextMin->GetCurrentTimestamp(Anchor) - LastMinCycles;
+						}
+
+						if (!NextMin->IsStartingWork())
+						{
+							// Accumulate the busy time for this pipe specifically.
+							NextMin->BusyCycles += NextMin->GetCurrentTimestamp(Anchor) - NextMin->GetPreviousTimestamp(Anchor);
+						}
+					}
+
+					LastMinCycles = NextMin->GetCurrentTimestamp(Anchor);
+
+					BusyPipes += NextMin->IsStartingWork() ? 1 : -1;
+					check(BusyPipes >= 0);
 
-			uint64 GetCurrentTimestamp (uint64 Anchor) const { return (*this)[TimestampIndex] - Anchor; }
-			uint64 GetPreviousTimestamp(uint64 Anchor) const { return (*this)[TimestampIndex - 1] - Anchor; }
+					NextMin->AdvanceTimestamp();
+					bFirst = false;
+				}
 
-			bool HasMoreTimestamps() const { return TimestampIndex < Num(); }
-			bool IsStartingWork()    const { return (TimestampIndex & 0x01) == 0x00; }
+				check(BusyPipes == 0);
 
-			void AdvanceTimestamp() { TimestampIndex++; }
+				return UnionBusyCycles;
+			}
 		};
 
 		struct FQueueState
 		{
 			bool bBusy = false;
-			FQueueTimestamps Timestamps;
+			FTimestampStream QueueTimestamps;
+
+		#if WITH_RHI_BREADCRUMBS && HAS_GPU_STATS
+			TMap<TStatId, FTimestampStream> StatTimestamps;
+			TMap<TStatId, int32> ActiveStats;
+		#endif
 		};
 
-		using FFrameState = TMap<FQueue, FQueueTimestamps>;
+		using FFrameState = TMap<FQueue, FQueueState>;
 
 		TMap<FQueue, FQueueState> QueueStates;
 		TMap<uint32, FFrameState> Frames;
@@ -755,22 +854,15 @@ namespace UE::RHI::GPUProfiler
 						check(!QueueState.bBusy);
 						QueueState.bBusy = true;
 						uint64 Value = Event->Value.Get<FEvent::FBeginWork>().GPUTimestampTOP;
+						QueueState.QueueTimestamps.AddTimestamp(Value, true);
 
-						if (!QueueState.Timestamps.IsEmpty() && Value <= QueueState.Timestamps.Last())
-						{
-							//
-							// The Begin TOP event is sooner than the last End BOP event.
-							// The markers overlap, and the GPU was not idle.
-							// 
-							// Remove the previous End event, and discard this Begin event.
-							//
-							QueueState.Timestamps.RemoveAt(QueueState.Timestamps.Num() - 1, EAllowShrinking::No);
-						}
-						else
+					#if WITH_RHI_BREADCRUMBS && HAS_GPU_STATS
+						// Apply the timestamp to all active stats
+						for (auto const& [StatId, RefCount] : QueueState.ActiveStats)
 						{
-							// GPU was idle. Keep this timestamp.
-							QueueState.Timestamps.Add(Value);
+							QueueState.StatTimestamps.FindChecked(StatId).AddTimestamp(Value, true);
 						}
+					#endif
 					}
 					break;
 
@@ -780,7 +872,15 @@ namespace UE::RHI::GPUProfiler
 						QueueState.bBusy = false;
 
 						uint64 Value = Event->Value.Get<FEvent::FEndWork>().GPUTimestampBOP;
-						QueueState.Timestamps.Add(Value);
+						QueueState.QueueTimestamps.AddTimestamp(Value, false);
+
+					#if WITH_RHI_BREADCRUMBS && HAS_GPU_STATS
+						// Apply the timestamp to all active stats
+						for (auto const& [StatId, RefCount] : QueueState.ActiveStats)
+						{
+							QueueState.StatTimestamps.FindChecked(StatId).AddTimestamp(Value, false);
+						}
+					#endif
 					}
 					break;
 
@@ -790,7 +890,7 @@ namespace UE::RHI::GPUProfiler
 						FEvent::FFrameBoundary const& FrameBoundary = Event->Value.Get<FEvent::FFrameBoundary>();
 
 						FFrameState& FrameState = Frames.FindOrAdd(FrameBoundary.FrameNumber);
-						FrameState.Emplace(Queue, MoveTemp(QueueState.Timestamps));
+						FrameState.Emplace(Queue, MoveTemp(QueueState));
 
 						if (FrameState.Num() == QueueStates.Num())
 						{
@@ -802,68 +902,96 @@ namespace UE::RHI::GPUProfiler
 						}
 					}
 					break;
+
+			#if WITH_RHI_BREADCRUMBS && HAS_GPU_STATS
+				case FEvent::EType::BeginBreadcrumb:
+					{
+						FEvent::FBeginBreadcrumb const& BeginBreadcrumb = Event->Value.Get<FEvent::FBeginBreadcrumb>();
+						TStatId StatId = BeginBreadcrumb.Breadcrumb->Name.StatId;
+
+						if (StatId.IsValidStat())
+						{
+							// Disregard the stat if it is nested within itself (i.e. its already in the ActiveStats map with a non-zero ref count).
+							// Only the outermost stat will count the busy time, otherwise we'd be double-counting the nested time.
+							int32 RefCount = QueueState.ActiveStats.FindOrAdd(StatId)++;
+							if (RefCount == 0)
+							{
+								QueueState.StatTimestamps.FindOrAdd(StatId).AddTimestamp(BeginBreadcrumb.GPUTimestampTOP, true);
+							}
+						}
+					}
+					break;
+
+				case FEvent::EType::EndBreadcrumb:
+					{
+						FEvent::FEndBreadcrumb const& EndBreadcrumb = Event->Value.Get<FEvent::FEndBreadcrumb>();
+						TStatId StatId = EndBreadcrumb.Breadcrumb->Name.StatId;
+
+						if (StatId.IsValidStat())
+						{
+							// Pop the stat when the refcount hits zero.
+							int32 RefCount = --QueueState.ActiveStats.FindChecked(StatId);
+							if (RefCount == 0)
+							{
+								QueueState.StatTimestamps.FindChecked(StatId).AddTimestamp(EndBreadcrumb.GPUTimestampBOP, false);
+								QueueState.ActiveStats.FindAndRemoveChecked(StatId);
+							}							
+						}
+					}
+					break;
+			#endif // WITH_RHI_BREADCRUMBS && HAS_GPU_STATS
+
 				}
 			}
 		}
 
 		void ProcessFrame(FFrameState& FrameState)
 		{
-			// The total number of cycles where at least one GPU pipe was busy during the frame.
-			uint64 UnionBusyCycles = 0;
-
-			uint64 LastMinCycles = 0;
-			int32 BusyPipes = 0;
-			bool bFirst = true;
+			TArray<FTimestampStream::FState, TInlineAllocator<GetRHIPipelineCount() * MAX_NUM_GPUS>> StreamPointers;
 
-			uint64 Anchor = 0; // @todo - handle possible timestamp wraparound
-
-			// Process the time ranges from each pipe.
-			while (true)
+			// Compute the individual GPU stats
+		#if WITH_RHI_BREADCRUMBS && HAS_GPU_STATS
+			TSet<TStatId> UniqueStats;
+			for (auto const& [Queue, State] : FrameState)
 			{
-				// Find the next minimum timestamp
-				FQueueTimestamps* NextMin = nullptr;
-				for (auto& [Queue, Current] : FrameState)
+				for (auto const& [StatId, Timestamps] : State.StatTimestamps)
 				{
-					if (Current.HasMoreTimestamps() && (!NextMin || Current.GetCurrentTimestamp(Anchor) < NextMin->GetCurrentTimestamp(Anchor)))
-					{
-						NextMin = &Current;
-					}
+					UniqueStats.Add(StatId);
 				}
+			}
 
-				if (!NextMin)
-					break; // No more timestamps to process
-
-				if (!bFirst)
+			for (TStatId StatId : UniqueStats)
+			{
+				StreamPointers.Reset();
+				for (auto const& [Queue, State] : FrameState)
 				{
-					if (BusyPipes > 0 && NextMin->GetCurrentTimestamp(Anchor) > LastMinCycles)
-					{
-						// Accumulate the union busy time across all pipes
-						UnionBusyCycles += NextMin->GetCurrentTimestamp(Anchor) - LastMinCycles;
-					}
-
-					if (!NextMin->IsStartingWork())
+					FTimestampStream const* Stream = State.StatTimestamps.Find(StatId);
+					if (Stream)
 					{
-						// Accumulate the busy time for this pipe specifically.
-						NextMin->BusyCycles += NextMin->GetCurrentTimestamp(Anchor) - NextMin->GetPreviousTimestamp(Anchor);
+						StreamPointers.Emplace(*Stream);
 					}
 				}
 
-				LastMinCycles = NextMin->GetCurrentTimestamp(Anchor);
-
-				BusyPipes += NextMin->IsStartingWork() ? 1 : -1;
-				check(BusyPipes >= 0);
+				uint64 Union = FTimestampStream::ComputeUnion(StreamPointers);
+				double Time = FPlatformTime::ToMilliseconds64(Union);
 
-				NextMin->AdvanceTimestamp();
-				bFirst = false;
+				SET_FLOAT_STAT_FName(StatId.GetName(), Time);
 			}
+		#endif // WITH_RHI_BREADCRUMBS && HAS_GPU_STATS
 
-			check(BusyPipes == 0);
+			// Compute the whole-frame total GPU time.
+			StreamPointers.Reset();
+			for (auto const& [Queue, State] : FrameState)
+			{
+				StreamPointers.Emplace(State.QueueTimestamps);
+			}
+			uint64 WholeFrameUnion = FTimestampStream::ComputeUnion(StreamPointers);
 
 			// Update the global GPU frame time stats - need to convert to Cycles32 rather than Cycles64.
-			GGPUFrameTime = FPlatformMath::TruncToInt(FPlatformTime::ToSeconds64(UnionBusyCycles) / FPlatformTime::GetSecondsPerCycle());
+			GGPUFrameTime = FPlatformMath::TruncToInt(FPlatformTime::ToSeconds64(WholeFrameUnion) / FPlatformTime::GetSecondsPerCycle());
 		}
 
-	} GGPUProfilerSink_StatUnit;
+	} GGPUProfilerSink_StatSystem;
 
 	struct FGPUProfilerSink_ProfileGPU final : public FEventSink
 	{
diff --git a/Engine/Source/Runtime/RHI/Private/RHIBreadcrumbs.cpp b/Engine/Source/Runtime/RHI/Private/RHIBreadcrumbs.cpp
--- a/Engine/Source/Runtime/RHI/Private/RHIBreadcrumbs.cpp
+++ b/Engine/Source/Runtime/RHI/Private/RHIBreadcrumbs.cpp
@@ -12,7 +12,7 @@
 FRHIBreadcrumbNode::FRHIBreadcrumbNode()
 	: Name([]() -> FRHIBreadcrumb const&
 	{
-		static TRHIBreadcrumb<> Breadcrumb(TEXT("Sentinel"));
+		static TRHIBreadcrumb<> Breadcrumb(TStatId(), TEXT("Sentinel"));
 		return Breadcrumb;
 	}())
 {}
diff --git a/Engine/Source/Runtime/RHI/Public/RHIBreadcrumbs.h b/Engine/Source/Runtime/RHI/Public/RHIBreadcrumbs.h
--- a/Engine/Source/Runtime/RHI/Public/RHIBreadcrumbs.h
+++ b/Engine/Source/Runtime/RHI/Public/RHIBreadcrumbs.h
@@ -156,7 +156,7 @@
 		FRHIBreadcrumbAllocatorArray const& GetParents() const { return Parents; }
 
 		template<size_t N, typename... TArgs>
-		inline FRHIBreadcrumbNode* AllocBreadcrumb(TCHAR const(&FormatString)[N], TArgs&&... Args);
+		inline FRHIBreadcrumbNode* AllocBreadcrumb(TStatId StatId, TCHAR const(&FormatString)[N], TArgs&&... Args);
 
 		inline void* Alloc(uint32 Size, uint32 Align)
 		{
@@ -433,8 +433,9 @@
 		RHI_API static std::atomic<uint32> NextID;
 
 	protected:
-		FRHIBreadcrumb()
+		FRHIBreadcrumb(TStatId StatId)
 			: ID(NextID.fetch_add(1, std::memory_order_relaxed) | 0x80000000) // Set the top bit to avoid collision with zero (i.e. "no breadcrumb")
+			, StatId(StatId)
 		{}
 
 		FRHIBreadcrumb(FRHIBreadcrumb const&) = delete;
@@ -442,6 +443,7 @@
 
 	public:
 		uint32 const ID;
+		TStatId const StatId;
 
 #if RHI_BREADCRUMBS_EMIT_CPU
 		uint32 const CPUTraceMarkerID = 0;
@@ -655,8 +657,9 @@
 		TRHIBreadcrumb(TRHIBreadcrumb const& Other) = delete;
 
 	public:
-		TRHIBreadcrumb(TCHAR const* FormatString, TArgs const&... Args)
-			: FormatString(FormatString)
+		TRHIBreadcrumb(TStatId StatId, TCHAR const* FormatString, TArgs const&... Args)
+			: FRHIBreadcrumb(StatId)
+			, FormatString(FormatString)
 			, Values(Args...)
 		{
 			CreateTraceMarkers();
@@ -685,8 +688,9 @@
 		TRHIBreadcrumb(TRHIBreadcrumb const& Other) = delete;
 
 	public:
-		TRHIBreadcrumb(TCHAR const* StringLiteral)
-			: StringLiteral(StringLiteral)
+		TRHIBreadcrumb(TStatId StatId, TCHAR const* StringLiteral)
+			: FRHIBreadcrumb(StatId)
+			, StringLiteral(StringLiteral)
 		{
 			CreateTraceMarkers();
 		}
@@ -703,20 +707,20 @@
 	};
 
 	template<size_t N, typename... TArgs>
-	inline FRHIBreadcrumbNode* FRHIBreadcrumbAllocator::AllocBreadcrumb(TCHAR const(&FormatString)[N], TArgs&&... Args)
+	inline FRHIBreadcrumbNode* FRHIBreadcrumbAllocator::AllocBreadcrumb(TStatId StatId, TCHAR const(&FormatString)[N], TArgs&&... Args)
 	{
 		struct FStorage
 		{
 			TRHIBreadcrumb<std::decay_t<TArgs>...> Name;
 			FRHIBreadcrumbNode Node;
 
-			FStorage(FRHIBreadcrumbAllocator& Allocator, TCHAR const(&FormatString)[N], TArgs&&... Args)
-				: Name(FormatString, Forward<TArgs>(Args)...)
+			FStorage(FRHIBreadcrumbAllocator& Allocator, TStatId StatId, TCHAR const(&FormatString)[N], TArgs&&... Args)
+				: Name(StatId, FormatString, Forward<TArgs>(Args)...)
 				, Node(Allocator, Name)
 			{}
 		};
 
-		return &Alloc<FStorage>(*this, FormatString, Forward<TArgs>(Args)...)->Node;
+		return &Alloc<FStorage>(*this, StatId, FormatString, Forward<TArgs>(Args)...)->Node;
 	}
 
 	class FRHIBreadcrumbNodeRef
@@ -759,7 +763,7 @@
 
 	public:
 		template<size_t N, typename... TArgs>
-		inline FRHIBreadcrumbEventManual(FRHIComputeCommandList& RHICmdList, TCHAR const(&FormatString)[N], TArgs&&... Args);
+		inline FRHIBreadcrumbEventManual(FRHIComputeCommandList& RHICmdList, TStatId StatId, TCHAR const(&FormatString)[N], TArgs&&... Args);
 
 		inline void End(FRHIComputeCommandList& RHICmdList);
 
@@ -779,31 +783,43 @@
 		FRHIBreadcrumbEventScope(FRHIBreadcrumbEventScope&&) = delete;
 
 		template<size_t N, typename... TArgs>
-		inline FRHIBreadcrumbEventScope(FRHIComputeCommandList& InRHICmdList, ERHIPipeline InPipeline, bool bCondition, TCHAR const(&FormatString)[N], TArgs&&... Args);
+		inline FRHIBreadcrumbEventScope(FRHIComputeCommandList& InRHICmdList, TStatId StatId, ERHIPipeline InPipeline, bool bCondition, TCHAR const(&FormatString)[N], TArgs&&... Args);
 
 	public:
 		// Top-of-pipe breadcrumb event scope for RHI command lists
 		template<size_t N, typename... TArgs>
-		inline FRHIBreadcrumbEventScope(FRHIComputeCommandList& InRHICmdList, bool bCondition, TCHAR const(&FormatString)[N], TArgs&&... Args);
+		inline FRHIBreadcrumbEventScope(FRHIComputeCommandList& InRHICmdList, TStatId StatId, bool bCondition, TCHAR const(&FormatString)[N], TArgs&&... Args);
 
 		// Bottom-of-pipe breadcrumb event scope for RHI contexts
 		template<size_t N, typename... TArgs>
-		inline FRHIBreadcrumbEventScope(IRHIComputeContext& InRHIContext, bool bCondition, TCHAR const(&FormatString)[N], TArgs&&... Args);
+		inline FRHIBreadcrumbEventScope(IRHIComputeContext& InRHIContext, TStatId StatId, bool bCondition, TCHAR const(&FormatString)[N], TArgs&&... Args);
 
 		inline ~FRHIBreadcrumbEventScope();
 	};
 
-	#define RHI_BREADCRUMB_EVENT(             RHICmdList_Or_RHIContext, Name                        ) FRHIBreadcrumbEventScope PREPROCESSOR_JOIN(BreadcrumbEvent_##Name,__LINE__)(RHICmdList_Or_RHIContext, true     , TEXT(#Name)          );
-	#define RHI_BREADCRUMB_EVENT_CONDITIONAL( RHICmdList_Or_RHIContext, Name, Condition             ) FRHIBreadcrumbEventScope PREPROCESSOR_JOIN(BreadcrumbEvent_##Name,__LINE__)(RHICmdList_Or_RHIContext, Condition, TEXT(#Name)          );
-	#define RHI_BREADCRUMB_EVENTF(            RHICmdList_Or_RHIContext, Name           , Format, ...) FRHIBreadcrumbEventScope PREPROCESSOR_JOIN(BreadcrumbEvent_##Name,__LINE__)(RHICmdList_Or_RHIContext, true     , Format, ##__VA_ARGS__);
-	#define RHI_BREADCRUMB_EVENTF_CONDITIONAL(RHICmdList_Or_RHIContext, Name, Condition, Format, ...) FRHIBreadcrumbEventScope PREPROCESSOR_JOIN(BreadcrumbEvent_##Name,__LINE__)(RHICmdList_Or_RHIContext, Condition, Format, ##__VA_ARGS__);
+	#define RHI_BREADCRUMB_EVENT(                 RHICmdList_Or_RHIContext,                  Format, ...) FRHIBreadcrumbEventScope ANONYMOUS_VARIABLE(BreadcrumbEvent)(RHICmdList_Or_RHIContext, TStatId()                  , true     , TEXT(Format), ## __VA_ARGS__)
+	#define RHI_BREADCRUMB_EVENT_CONDITIONAL(     RHICmdList_Or_RHIContext,       Condition, Format, ...) FRHIBreadcrumbEventScope ANONYMOUS_VARIABLE(BreadcrumbEvent)(RHICmdList_Or_RHIContext, TStatId()                  , Condition, TEXT(Format), ## __VA_ARGS__)
+#if HAS_GPU_STATS
+	#define RHI_BREADCRUMB_EVENT_STAT(            RHICmdList_Or_RHIContext, Stat,            Format, ...) FRHIBreadcrumbEventScope ANONYMOUS_VARIABLE(BreadcrumbEvent)(RHICmdList_Or_RHIContext, GET_STATID(Stat_GPU_##Stat), true     , TEXT(Format), ## __VA_ARGS__)
+	#define RHI_BREADCRUMB_EVENT_CONDITIONAL_STAT(RHICmdList_Or_RHIContext, Stat, Condition, Format, ...) FRHIBreadcrumbEventScope ANONYMOUS_VARIABLE(BreadcrumbEvent)(RHICmdList_Or_RHIContext, GET_STATID(Stat_GPU_##Stat), Condition, TEXT(Format), ## __VA_ARGS__)
+#else
+	#define RHI_BREADCRUMB_EVENT_STAT(            RHICmdList_Or_RHIContext, Stat,            Format, ...) FRHIBreadcrumbEventScope ANONYMOUS_VARIABLE(BreadcrumbEvent)(RHICmdList_Or_RHIContext, TStatId()                  , true     , TEXT(Format), ## __VA_ARGS__)
+	#define RHI_BREADCRUMB_EVENT_CONDITIONAL_STAT(RHICmdList_Or_RHIContext, Stat, Condition, Format, ...) FRHIBreadcrumbEventScope ANONYMOUS_VARIABLE(BreadcrumbEvent)(RHICmdList_Or_RHIContext, TStatId()                  , Condition, TEXT(Format), ## __VA_ARGS__)
+#endif
+
+	// Used only for back compat with SCOPED_DRAW_EVENTF
+	#define RHI_BREADCRUMB_EVENT_STR_DEPRECATED(            RHICmdList_Or_RHIContext,            Format, ...) FRHIBreadcrumbEventScope ANONYMOUS_VARIABLE(BreadcrumbEvent)(RHICmdList_Or_RHIContext, TStatId(), true     , Format, ## __VA_ARGS__)
+	#define RHI_BREADCRUMB_EVENT_CONDITIONAL_STR_DEPRECATED(RHICmdList_Or_RHIContext, Condition, Format, ...) FRHIBreadcrumbEventScope ANONYMOUS_VARIABLE(BreadcrumbEvent)(RHICmdList_Or_RHIContext, TStatId(), Condition, Format, ## __VA_ARGS__)
 
 #else // WITH_RHI_BREADCRUMBS == 0
 
-	#define RHI_BREADCRUMB_EVENT(             RHICmdList_Or_RHIContext, Name                        )
-	#define RHI_BREADCRUMB_EVENT_CONDITIONAL( RHICmdList_Or_RHIContext, Name, Condition             )
-	#define RHI_BREADCRUMB_EVENTF(            RHICmdList_Or_RHIContext, Name           , Format, ...)
-	#define RHI_BREADCRUMB_EVENTF_CONDITIONAL(RHICmdList_Or_RHIContext, Name, Condition, Format, ...)
+	#define RHI_BREADCRUMB_EVENT(...)
+	#define RHI_BREADCRUMB_EVENT_STAT(...)
+	#define RHI_BREADCRUMB_EVENT_CONDITIONAL(...)
+	#define RHI_BREADCRUMB_EVENT_CONDITIONAL_STAT(...)
+
+	#define RHI_BREADCRUMB_EVENT_STR_DEPRECATED(...)
+	#define RHI_BREADCRUMB_EVENT_CONDITIONAL_STR_DEPRECATED(...)
 
 #endif
 //
diff --git a/Engine/Source/Runtime/RHI/Public/RHICommandList.h b/Engine/Source/Runtime/RHI/Public/RHICommandList.h
--- a/Engine/Source/Runtime/RHI/Public/RHICommandList.h
+++ b/Engine/Source/Runtime/RHI/Public/RHICommandList.h
@@ -2826,7 +2826,7 @@ public:
 	{
 #if WITH_PROFILEGPU
 		extern RHI_API TAutoConsoleVariable<int32> GProfileGPUTransitions;
-		RHI_BREADCRUMB_EVENT_CONDITIONAL(*this, RHIBeginTransitions, GProfileGPUTransitions.GetValueOnAnyThread() != 0);
+		RHI_BREADCRUMB_EVENT_CONDITIONAL(*this, GProfileGPUTransitions.GetValueOnAnyThread() != 0, "RHIBeginTransitions");
 #endif
 
 		if (Bypass())
@@ -2852,7 +2852,7 @@ public:
 	{
 #if WITH_PROFILEGPU
 		extern RHI_API TAutoConsoleVariable<int32> GProfileGPUTransitions;
-		RHI_BREADCRUMB_EVENT_CONDITIONAL(*this, RHIEndTransitions, GProfileGPUTransitions.GetValueOnAnyThread() != 0);
+		RHI_BREADCRUMB_EVENT_CONDITIONAL(*this, GProfileGPUTransitions.GetValueOnAnyThread() != 0, "RHIEndTransitions");
 #endif
 
 		if (Bypass())
diff --git a/Engine/Source/Runtime/RHI/Public/RHICommandList.inl b/Engine/Source/Runtime/RHI/Public/RHICommandList.inl
--- a/Engine/Source/Runtime/RHI/Public/RHICommandList.inl
+++ b/Engine/Source/Runtime/RHI/Public/RHICommandList.inl
@@ -138,20 +138,20 @@ inline FRHIResourceReplaceBatcher::~FRHIResourceReplaceBatcher()
 
 	// Top-of-pipe breadcrumb event scope for RHI command lists
 	template<size_t N, typename... TArgs>
-	inline FRHIBreadcrumbEventScope::FRHIBreadcrumbEventScope(FRHIComputeCommandList& InRHICmdList, bool bCondition, TCHAR const(&FormatString)[N], TArgs&&... Args)
-		: FRHIBreadcrumbEventScope(InRHICmdList, InRHICmdList.GetPipeline(), bCondition, FormatString, Forward<TArgs>(Args)...)
+	inline FRHIBreadcrumbEventScope::FRHIBreadcrumbEventScope(FRHIComputeCommandList& InRHICmdList, TStatId StatId, bool bCondition, TCHAR const(&FormatString)[N], TArgs&&... Args)
+		: FRHIBreadcrumbEventScope(InRHICmdList, StatId, InRHICmdList.GetPipeline(), bCondition, FormatString, Forward<TArgs>(Args)...)
 	{}
 
 	// Bottom-of-pipe breadcrumb event scope for RHI contexts
 	template<size_t N, typename... TArgs>
-	inline FRHIBreadcrumbEventScope::FRHIBreadcrumbEventScope(IRHIComputeContext& InRHIContext, bool bCondition, TCHAR const(&FormatString)[N], TArgs&&... Args)
-		: FRHIBreadcrumbEventScope(static_cast<FRHIComputeCommandList&>(InRHIContext.GetExecutingCommandList()), InRHIContext.GetPipeline(), bCondition, FormatString, Forward<TArgs>(Args)...)
+	inline FRHIBreadcrumbEventScope::FRHIBreadcrumbEventScope(IRHIComputeContext& InRHIContext, TStatId StatId, bool bCondition, TCHAR const(&FormatString)[N], TArgs&&... Args)
+		: FRHIBreadcrumbEventScope(static_cast<FRHIComputeCommandList&>(InRHIContext.GetExecutingCommandList()), StatId, InRHIContext.GetPipeline(), bCondition, FormatString, Forward<TArgs>(Args)...)
 	{}
 
 	template<size_t N, typename... TArgs>
-	inline FRHIBreadcrumbEventScope::FRHIBreadcrumbEventScope(FRHIComputeCommandList& InRHICmdList, ERHIPipeline InPipeline, bool bCondition, TCHAR const(&FormatString)[N], TArgs&&... Args)
+	inline FRHIBreadcrumbEventScope::FRHIBreadcrumbEventScope(FRHIComputeCommandList& InRHICmdList, TStatId StatId, ERHIPipeline InPipeline, bool bCondition, TCHAR const(&FormatString)[N], TArgs&&... Args)
 		: RHICmdList(InRHICmdList)
-		, Node(bCondition ? RHICmdList.GetBreadcrumbAllocator().AllocBreadcrumb(FormatString, Forward<TArgs>(Args)...) : nullptr)
+		, Node(bCondition ? RHICmdList.GetBreadcrumbAllocator().AllocBreadcrumb(StatId, FormatString, Forward<TArgs>(Args)...) : nullptr)
 		, Pipeline(InPipeline)
 	{
 		if (Node)
@@ -172,8 +172,8 @@ inline FRHIResourceReplaceBatcher::~FRHIResourceReplaceBatcher()
 	}
 
 	template<size_t N, typename... TArgs>
-	inline FRHIBreadcrumbEventManual::FRHIBreadcrumbEventManual(FRHIComputeCommandList& RHICmdList, TCHAR const(&FormatString)[N], TArgs&&... Args)
-		: Node(RHICmdList.GetBreadcrumbAllocator().AllocBreadcrumb(FormatString, Forward<TArgs>(Args)...))
+	inline FRHIBreadcrumbEventManual::FRHIBreadcrumbEventManual(FRHIComputeCommandList& RHICmdList, TStatId StatId, TCHAR const(&FormatString)[N], TArgs&&... Args)
+		: Node(RHICmdList.GetBreadcrumbAllocator().AllocBreadcrumb(StatId, FormatString, Forward<TArgs>(Args)...))
 	#if DO_CHECK
 		, Pipeline(RHICmdList.GetPipeline())
 		, ThreadId(FPlatformTLS::GetCurrentThreadId())
diff --git a/Engine/Source/Runtime/RenderCore/Private/ProfilingDebugging/RealtimeGPUProfiler.cpp b/Engine/Source/Runtime/RenderCore/Private/ProfilingDebugging/RealtimeGPUProfiler.cpp
--- a/Engine/Source/Runtime/RenderCore/Private/ProfilingDebugging/RealtimeGPUProfiler.cpp
+++ b/Engine/Source/Runtime/RenderCore/Private/ProfilingDebugging/RealtimeGPUProfiler.cpp
@@ -11,6 +11,7 @@
 #if RHI_NEW_GPU_PROFILER
 
 // @todo - new gpu profiler
+CSV_DEFINE_CATEGORY_MODULE(RENDERCORE_API, GPU, true);
 
 #elif HAS_GPU_STATS
 
diff --git a/Engine/Source/Runtime/RenderCore/Private/RenderCaptureInterface.cpp b/Engine/Source/Runtime/RenderCore/Private/RenderCaptureInterface.cpp
--- a/Engine/Source/Runtime/RenderCore/Private/RenderCaptureInterface.cpp
+++ b/Engine/Source/Runtime/RenderCore/Private/RenderCaptureInterface.cpp
@@ -35,7 +35,7 @@ namespace RenderCaptureInterface
 #if WITH_RHI_BREADCRUMBS
 				if (bPushEvent)
 				{
-					Breadcrumb->Emplace(RHICommandListLocal, TEXT("%s"), EventName);
+					Breadcrumb->Emplace(RHICommandListLocal, TStatId(), TEXT("%s"), EventName);
 				}
 #endif
 			});
@@ -58,7 +58,7 @@ namespace RenderCaptureInterface
 			if (bEvent)
 			{
 				RHIBreadcrumb = MakeUnique<TOptional<FRHIBreadcrumbEventManual>>();
-				RHIBreadcrumb->Emplace(*RHICommandList, TEXT("%s"), FString(InEventName));
+				RHIBreadcrumb->Emplace(*RHICommandList, TStatId(), TEXT("%s"), FString(InEventName));
 			}
 #endif
 		}
@@ -85,7 +85,7 @@ namespace RenderCaptureInterface
 #if RDG_EVENTS
 			if (bEvent)
 			{
-				RDGEvent.Emplace(*GraphBuilder, ERDGScopeFlags::None, RDG_EVENT_NAME("%s", InEventName));
+				RDGEvent.Emplace(*GraphBuilder, ERDGScopeFlags::None, TStatId(), RDG_EVENT_NAME("%s", InEventName));
 			}
 #endif
 		}
diff --git a/Engine/Source/Runtime/RenderCore/Private/RenderGraphBuilder.cpp b/Engine/Source/Runtime/RenderCore/Private/RenderGraphBuilder.cpp
--- a/Engine/Source/Runtime/RenderCore/Private/RenderGraphBuilder.cpp
+++ b/Engine/Source/Runtime/RenderCore/Private/RenderGraphBuilder.cpp
@@ -2517,7 +2517,7 @@ FRDGPass* FRDGBuilder::SetupParameterPass(FRDGPass* Pass)
 	if (ScopeState.ScopeMode == ERDGScopeMode::AllEventsAndPassNames)
 	{
 		FRDGEventName Name = Pass->GetEventName();
-		PassNameScope.Emplace(*this, ERDGScopeFlags::None, MoveTemp(Name));
+		PassNameScope.Emplace(*this, ERDGScopeFlags::None, TStatId(), MoveTemp(Name));
 	}
 #endif
 
diff --git a/Engine/Source/Runtime/RenderCore/Private/RenderGraphEvent.cpp b/Engine/Source/Runtime/RenderCore/Private/RenderGraphEvent.cpp
--- a/Engine/Source/Runtime/RenderCore/Private/RenderGraphEvent.cpp
+++ b/Engine/Source/Runtime/RenderCore/Private/RenderGraphEvent.cpp
@@ -286,7 +286,7 @@ FRDGEventName::FRDGEventName(const TCHAR* EventFormat, ...)
 }
 
 #if WITH_RHI_BREADCRUMBS
-FRHIBreadcrumbNode* FRDGEventName::AllocBreadcrumb(FRHIBreadcrumbAllocator& Allocator) const
+FRHIBreadcrumbNode* FRDGEventName::AllocBreadcrumb(TStatId StatId, FRHIBreadcrumbAllocator& Allocator) const
 {
 #if RDG_EVENTS == RDG_EVENTS_STRING_COPY
 
@@ -297,7 +297,7 @@ FRHIBreadcrumbNode* FRDGEventName::AllocBreadcrumb(FRHIBreadcrumbAllocator& Allo
 				: TEXT("<unnamed pass>");
 			
 			// Cast hack to force treating EventFormat as a string literal.
-			return Allocator.AllocBreadcrumb(*reinterpret_cast<TCHAR const(*)[1]>(String));
+			return Allocator.AllocBreadcrumb(StatId, *reinterpret_cast<TCHAR const(*)[1]>(String));
 		}
 		else
 		{
@@ -310,13 +310,13 @@ FRHIBreadcrumbNode* FRDGEventName::AllocBreadcrumb(FRHIBreadcrumbAllocator& Allo
 			FMemory::Memcpy(StrCopy, *FormattedEventName, Size);
 
 			// Cast hack to force treating StrCopy as a string literal.
-			return Allocator.AllocBreadcrumb(*reinterpret_cast<TCHAR const(*)[1]>(StrCopy));
+			return Allocator.AllocBreadcrumb(StatId, *reinterpret_cast<TCHAR const(*)[1]>(StrCopy));
 		}
 
 #elif RDG_EVENTS >= RDG_EVENTS_STRING_REF
 
 		// Cast hack to force treating EventFormat as a string literal.
-		return Allocator.AllocBreadcrumb(*reinterpret_cast<TCHAR const(*)[1]>(EventFormat));
+		return Allocator.AllocBreadcrumb(StatId, *reinterpret_cast<TCHAR const(*)[1]>(EventFormat));
 
 #else
 
diff --git a/Engine/Source/Runtime/RenderCore/Public/ProfilingDebugging/RealtimeGPUProfiler.h b/Engine/Source/Runtime/RenderCore/Public/ProfilingDebugging/RealtimeGPUProfiler.h
--- a/Engine/Source/Runtime/RenderCore/Public/ProfilingDebugging/RealtimeGPUProfiler.h
+++ b/Engine/Source/Runtime/RenderCore/Public/ProfilingDebugging/RealtimeGPUProfiler.h
@@ -1,13 +1,5 @@
 // Copyright Epic Games, Inc. All Rights Reserved.
 
-/**
- *
- * This file contains the various draw mesh macros that display draw calls
- * inside of PIX.
- */
-
-// Colors that are defined for a particular mesh type
-// Each event type will be displayed using the defined color
 #pragma once
 
 #include "Containers/Array.h"
@@ -56,7 +48,7 @@
 				{
 					std::apply([Event, &RHICmdList, &FormatString](auto&&... InnerArgs)
 					{
-						Event->Emplace(RHICmdList, true, FormatString, std::forward<TArgs>(InnerArgs)...);
+						Event->Emplace(RHICmdList, TStatId(), true, FormatString, std::forward<TArgs>(InnerArgs)...);
 					}, Values);
 				});
 			}
@@ -75,32 +67,32 @@
 		}
 	};
 
-	#define RHI_BREADCRUMB_EVENT_GAMETHREAD(             Name                        ) FRHIBreadcrumbEvent_GameThread PREPROCESSOR_JOIN(BreadcrumbEvent_GameThread_##Name,__LINE__)(true     , TEXT(#Name)          );
-	#define RHI_BREADCRUMB_EVENT_CONDITIONAL_GAMETHREAD( Name, Condition             ) FRHIBreadcrumbEvent_GameThread PREPROCESSOR_JOIN(BreadcrumbEvent_GameThread_##Name,__LINE__)(Condition, TEXT(#Name)          );
-	#define RHI_BREADCRUMB_EVENTF_GAMETHREAD(            Name,            Format, ...) FRHIBreadcrumbEvent_GameThread PREPROCESSOR_JOIN(BreadcrumbEvent_GameThread_##Name,__LINE__)(true     , Format, ##__VA_ARGS__);
-	#define RHI_BREADCRUMB_EVENTF_CONDITIONAL_GAMETHREAD(Name, Condition, Format, ...) FRHIBreadcrumbEvent_GameThread PREPROCESSOR_JOIN(BreadcrumbEvent_GameThread_##Name,__LINE__)(Condition, Format, ##__VA_ARGS__);
+	#define RHI_BREADCRUMB_EVENT_GAMETHREAD(                       Format, ...) FRHIBreadcrumbEvent_GameThread ANONYMOUS_VARIABLE(BreadcrumbEvent_GameThread)(true     , TEXT(Format), ##__VA_ARGS__)
+	#define RHI_BREADCRUMB_EVENT_CONDITIONAL_GAMETHREAD(Condition, Format, ...) FRHIBreadcrumbEvent_GameThread ANONYMOUS_VARIABLE(BreadcrumbEvent_GameThread)(Condition, TEXT(Format), ##__VA_ARGS__)
+
+	// Used only for back compat with SCOPED_DRAW_EVENTF_GAMETHREAD
+	#define RHI_BREADCRUMB_EVENT_GAMETHREAD_STR_DEPRECATED(Format, ...) FRHIBreadcrumbEvent_GameThread ANONYMOUS_VARIABLE(BreadcrumbEvent_GameThread)(true, Format, ##__VA_ARGS__)
 
 #else
 
-	#define RHI_BREADCRUMB_EVENT_GAMETHREAD(             Name                        ) do { } while(0)
-	#define RHI_BREADCRUMB_EVENT_CONDITIONAL_GAMETHREAD( Name, Condition             ) do { } while(0)
-	#define RHI_BREADCRUMB_EVENTF_GAMETHREAD(            Name,            Format, ...) do { } while(0)
-	#define RHI_BREADCRUMB_EVENTF_CONDITIONAL_GAMETHREAD(Name, Condition, Format, ...) do { } while(0)
+	#define RHI_BREADCRUMB_EVENT_GAMETHREAD(...)                do { } while(0)
+	#define RHI_BREADCRUMB_EVENT_CONDITIONAL_GAMETHREAD(...)    do { } while(0)
+	#define RHI_BREADCRUMB_EVENT_GAMETHREAD_STR_DEPRECATED(...) do { } while(0)
 
 #endif
 
 // Macros to allow for scoping of draw events outside of RHI function implementations
 // Render-thread event macros:
-#define SCOPED_DRAW_EVENT(RHICmdList, Name)                                             RHI_BREADCRUMB_EVENT(             RHICmdList, Name                                  );
-#define SCOPED_DRAW_EVENTF(RHICmdList, Name, Format, ...)                               RHI_BREADCRUMB_EVENTF(            RHICmdList, Name           , Format, ##__VA_ARGS__);
-#define SCOPED_CONDITIONAL_DRAW_EVENT(RHICmdList, Name, Condition)                      RHI_BREADCRUMB_EVENT_CONDITIONAL( RHICmdList, Name, Condition                       );
-#define SCOPED_CONDITIONAL_DRAW_EVENTF(RHICmdList, Name, Condition, Format, ...)        RHI_BREADCRUMB_EVENTF_CONDITIONAL(RHICmdList, Name, Condition, Format, ##__VA_ARGS__);
+#define SCOPED_DRAW_EVENT(RHICmdList, Name)                                      RHI_BREADCRUMB_EVENT(RHICmdList, #Name);
+#define SCOPED_DRAW_EVENTF(RHICmdList, Name, Format, ...)                        RHI_BREADCRUMB_EVENT_STR_DEPRECATED(RHICmdList, Format, ##__VA_ARGS__);
+#define SCOPED_CONDITIONAL_DRAW_EVENT(RHICmdList, Name, Condition)               RHI_BREADCRUMB_EVENT_CONDITIONAL(RHICmdList, Condition, #Name);
+#define SCOPED_CONDITIONAL_DRAW_EVENTF(RHICmdList, Name, Condition, Format, ...) RHI_BREADCRUMB_EVENT_CONDITIONAL_STR_DEPRECATED(RHICmdList, Condition, Format, ##__VA_ARGS__);
 
 // Non-render-thread event macros:
-#define SCOPED_DRAW_EVENT_GAMETHREAD(Name)                                              RHI_BREADCRUMB_EVENT_GAMETHREAD(             Name                                  )
-#define SCOPED_DRAW_EVENTF_GAMETHREAD(Name, Format, ...)                                RHI_BREADCRUMB_EVENTF_GAMETHREAD(            Name           , Format, ##__VA_ARGS__)
-#define SCOPED_CONDITIONAL_DRAW_EVENT_GAMETHREAD(Name, Condition)                       RHI_BREADCRUMB_EVENT_CONDITIONAL_GAMETHREAD( Name, Condition                       )
-#define SCOPED_CONDITIONAL_DRAW_EVENTF_GAMETHREAD(Name, Condition, Format, ...)         RHI_BREADCRUMB_EVENTF_CONDITIONAL_GAMETHREAD(Name, Condition, Format, ##__VA_ARGS__)
+#define SCOPED_DRAW_EVENT_GAMETHREAD(Name)                                       RHI_BREADCRUMB_EVENT_GAMETHREAD(#Name);
+#define SCOPED_DRAW_EVENTF_GAMETHREAD(Name, Format, ...)                         RHI_BREADCRUMB_EVENT_GAMETHREAD_STR_DEPRECATED(Format, ##__VA_ARGS__);
+#define SCOPED_CONDITIONAL_DRAW_EVENT_GAMETHREAD(Name, Condition)                RHI_BREADCRUMB_EVENT_CONDITIONAL_GAMETHREAD(Condition, #Name);
+#define SCOPED_CONDITIONAL_DRAW_EVENTF_GAMETHREAD(Name, Condition, Format, ...)  RHI_BREADCRUMB_EVENT_CONDITIONAL_GAMETHREAD(Condition, Format, ##__VA_ARGS__);
 
 // Deprecated macros
 #define BEGIN_DRAW_EVENTF(RHICmdList, Name, Event, Format, ...)                                      UE_DEPRECATED_MACRO(5.5, "BEGIN_DRAW_EVENTF has been deprecated. Equivalent functionality can be implemented by constructing / destructing an instance of FRHIBreadcrumbEventManual."                      )
@@ -108,34 +100,34 @@
 #define STOP_DRAW_EVENT_GAMETHREAD(...)                                                              UE_DEPRECATED_MACRO(5.5, "STOP_DRAW_EVENT_GAMETHREAD has been deprecated. Equivalent functionality can be implemented by constructing / destructing an instance of FRHIBreadcrumbEvent_GameThread."        )
 #define BEGIN_DRAW_EVENTF_GAMETHREAD(...)                                                            UE_DEPRECATED_MACRO(5.5, "BEGIN_DRAW_EVENTF_GAMETHREAD has been deprecated. Equivalent functionality can be implemented by constructing / destructing an instance of FRHIBreadcrumbEvent_GameThread."      )
 #define BEGIN_DRAW_EVENTF_COLOR_GAMETHREAD(...)                                                      UE_DEPRECATED_MACRO(5.5, "BEGIN_DRAW_EVENTF_COLOR_GAMETHREAD has been deprecated. Equivalent functionality can be implemented by constructing / destructing an instance of FRHIBreadcrumbEvent_GameThread.")
-#define SCOPED_DRAW_EVENT_COLOR(RHICmdList, Color, Name)                                             UE_DEPRECATED_MACRO(5.5, "SCOPED_DRAW_EVENT_COLOR has been deprecated. Use SCOPED_DRAW_EVENT instead."                                                ) SCOPED_DRAW_EVENT(RHICmdList, Name)
-#define SCOPED_GPU_EVENT(RHICmdList, Name)                                                           UE_DEPRECATED_MACRO(5.5, "SCOPED_GPU_EVENT has been deprecated. Use SCOPED_DRAW_EVENT instead."                                                       ) SCOPED_DRAW_EVENT(RHICmdList, Name)
-#define SCOPED_GPU_EVENT_COLOR(RHICmdList, Color, Name)                                              UE_DEPRECATED_MACRO(5.5, "SCOPED_GPU_EVENT_COLOR has been deprecated. Use SCOPED_DRAW_EVENT instead."                                                 ) SCOPED_DRAW_EVENT(RHICmdList, Name)
-#define SCOPED_COMPUTE_EVENT(RHICmdList, Name)                                                       UE_DEPRECATED_MACRO(5.5, "SCOPED_COMPUTE_EVENT has been deprecated. Use SCOPED_DRAW_EVENT instead."                                                   ) SCOPED_DRAW_EVENT(RHICmdList, Name)
-#define SCOPED_COMPUTE_EVENT_COLOR(RHICmdList, Color, Name)                                          UE_DEPRECATED_MACRO(5.5, "SCOPED_COMPUTE_EVENT_COLOR has been deprecated. Use SCOPED_DRAW_EVENT instead."                                             ) SCOPED_DRAW_EVENT(RHICmdList, Name)
-#define SCOPED_DRAW_EVENTF_COLOR(RHICmdList, Color, Name, Format, ...)                               UE_DEPRECATED_MACRO(5.5, "SCOPED_DRAW_EVENTF_COLOR has been deprecated. Use SCOPED_DRAW_EVENTF instead."                                              ) SCOPED_DRAW_EVENTF(RHICmdList, Name, Format, ##__VA_ARGS__)
-#define SCOPED_GPU_EVENTF(RHICmdList, Name, Format, ...)                                             UE_DEPRECATED_MACRO(5.5, "SCOPED_GPU_EVENTF has been deprecated. Use SCOPED_DRAW_EVENTF instead."                                                     ) SCOPED_DRAW_EVENTF(RHICmdList, Name, Format, ##__VA_ARGS__)
-#define SCOPED_GPU_EVENTF_COLOR(RHICmdList, Color, Name, Format, ...)                                UE_DEPRECATED_MACRO(5.5, "SCOPED_GPU_EVENTF_COLOR has been deprecated. Use SCOPED_DRAW_EVENTF instead."                                               ) SCOPED_DRAW_EVENTF(RHICmdList, Name, Format, ##__VA_ARGS__)
-#define SCOPED_COMPUTE_EVENTF(RHICmdList, Name, Format, ...)                                         UE_DEPRECATED_MACRO(5.5, "SCOPED_COMPUTE_EVENTF has been deprecated. Use SCOPED_DRAW_EVENTF instead."                                                 ) SCOPED_DRAW_EVENTF(RHICmdList, Name, Format, ##__VA_ARGS__)
-#define SCOPED_COMPUTE_EVENTF_COLOR(RHICmdList, Color, Name, Format, ...)                            UE_DEPRECATED_MACRO(5.5, "SCOPED_COMPUTE_EVENTF_COLOR has been deprecated. Use SCOPED_DRAW_EVENTF instead."                                           ) SCOPED_DRAW_EVENTF(RHICmdList, Name, Format, ##__VA_ARGS__)
-#define SCOPED_CONDITIONAL_DRAW_EVENT_COLOR(RHICmdList, Name, Color, Condition)                      UE_DEPRECATED_MACRO(5.5, "SCOPED_CONDITIONAL_DRAW_EVENT_COLOR has been deprecated. Use SCOPED_CONDITIONAL_DRAW_EVENT instead."                        ) SCOPED_CONDITIONAL_DRAW_EVENT(RHICmdList, Name, Condition)
-#define SCOPED_CONDITIONAL_GPU_EVENT(RHICmdList, Name, Condition)                                    UE_DEPRECATED_MACRO(5.5, "SCOPED_CONDITIONAL_GPU_EVENT has been deprecated. Use SCOPED_CONDITIONAL_DRAW_EVENT instead."                               ) SCOPED_CONDITIONAL_DRAW_EVENT(RHICmdList, Name, Condition)
-#define SCOPED_CONDITIONAL_GPU_EVENT_COLOR(RHICmdList, Name, Color, Condition)                       UE_DEPRECATED_MACRO(5.5, "SCOPED_CONDITIONAL_GPU_EVENT_COLOR has been deprecated. Use SCOPED_CONDITIONAL_DRAW_EVENT instead."                         ) SCOPED_CONDITIONAL_DRAW_EVENT(RHICmdList, Name, Condition)
-#define SCOPED_CONDITIONAL_COMPUTE_EVENT(RHICmdList, Name, Condition)                                UE_DEPRECATED_MACRO(5.5, "SCOPED_CONDITIONAL_COMPUTE_EVENT has been deprecated. Use SCOPED_CONDITIONAL_DRAW_EVENT instead."                           ) SCOPED_CONDITIONAL_DRAW_EVENT(RHICmdList, Name, Condition)
-#define SCOPED_CONDITIONAL_COMPUTE_EVENT_COLOR(RHICmdList, Name, Color, Condition)                   UE_DEPRECATED_MACRO(5.5, "SCOPED_CONDITIONAL_COMPUTE_EVENT_COLOR has been deprecated. Use SCOPED_CONDITIONAL_DRAW_EVENT instead."                     ) SCOPED_CONDITIONAL_DRAW_EVENT(RHICmdList, Name, Condition)
-#define SCOPED_CONDITIONAL_DRAW_EVENTF_COLOR(RHICmdList, Color, Name, Condition, Format, ...)        UE_DEPRECATED_MACRO(5.5, "SCOPED_CONDITIONAL_DRAW_EVENTF_COLOR has been deprecated. Use SCOPED_CONDITIONAL_DRAW_EVENTF instead."                      ) SCOPED_CONDITIONAL_DRAW_EVENTF(RHICmdList, Name, Condition, Format, ##__VA_ARGS__)
-#define SCOPED_CONDITIONAL_GPU_EVENTF(RHICmdList, Name, Condition, Format, ...)                      UE_DEPRECATED_MACRO(5.5, "SCOPED_CONDITIONAL_GPU_EVENTF has been deprecated. Use SCOPED_CONDITIONAL_DRAW_EVENTF instead."                             ) SCOPED_CONDITIONAL_DRAW_EVENTF(RHICmdList, Name, Condition, Format, ##__VA_ARGS__)
-#define SCOPED_CONDITIONAL_GPU_EVENTF_COLOR(RHICmdList, Color, Name, Condition, Format, ...)         UE_DEPRECATED_MACRO(5.5, "SCOPED_CONDITIONAL_GPU_EVENTF_COLOR has been deprecated. Use SCOPED_CONDITIONAL_DRAW_EVENTF instead."                       ) SCOPED_CONDITIONAL_DRAW_EVENTF(RHICmdList, Name, Condition, Format, ##__VA_ARGS__)
-#define SCOPED_CONDITIONAL_COMPUTE_EVENTF(RHICmdList, Name, Condition, Format, ...)                  UE_DEPRECATED_MACRO(5.5, "SCOPED_CONDITIONAL_COMPUTE_EVENTF has been deprecated. Use SCOPED_CONDITIONAL_DRAW_EVENTF instead."                         ) SCOPED_CONDITIONAL_DRAW_EVENTF(RHICmdList, Name, Condition, Format, ##__VA_ARGS__)
-#define SCOPED_CONDITIONAL_COMPUTE_EVENTF_COLOR(RHICmdList, Color, Name, Condition, Format, ...)     UE_DEPRECATED_MACRO(5.5, "SCOPED_CONDITIONAL_COMPUTE_EVENTF_COLOR has been deprecated. Use SCOPED_CONDITIONAL_DRAW_EVENTF instead."                   ) SCOPED_CONDITIONAL_DRAW_EVENTF(RHICmdList, Name, Condition, Format, ##__VA_ARGS__)
-#define BEGIN_DRAW_EVENTF_COLOR(RHICmdList, Color, Name, Event, Format, ...)                         UE_DEPRECATED_MACRO(5.5, "BEGIN_DRAW_EVENTF_COLOR has been deprecated. Use BEGIN_DRAW_EVENTF instead."                                                ) BEGIN_DRAW_EVENTF(RHICmdList, Name, Event, Format, ##__VA_ARGS__)
-#define BEGIN_GPU_EVENTF(RHICmdList, Name, Event, Format, ...)                                       UE_DEPRECATED_MACRO(5.5, "BEGIN_GPU_EVENTF has been deprecated. Use BEGIN_DRAW_EVENTF instead."                                                       ) BEGIN_DRAW_EVENTF(RHICmdList, Name, Event, Format, ##__VA_ARGS__)
-#define BEGIN_GPU_EVENTF_COLOR(RHICmdList, Color, Name, Event, Format, ...)                          UE_DEPRECATED_MACRO(5.5, "BEGIN_GPU_EVENTF_COLOR has been deprecated. Use BEGIN_DRAW_EVENTF instead."                                                 ) BEGIN_DRAW_EVENTF(RHICmdList, Name, Event, Format, ##__VA_ARGS__)
-#define STOP_GPU_EVENT(Event)                                                                        UE_DEPRECATED_MACRO(5.5, "STOP_GPU_EVENT has been deprecated. Use STOP_DRAW_EVENT instead."                                                           ) STOP_DRAW_EVENT(Event)
-#define SCOPED_DRAW_EVENT_COLOR_GAMETHREAD(Color, Name)                                              UE_DEPRECATED_MACRO(5.5, "SCOPED_DRAW_EVENT_COLOR_GAMETHREAD has been deprecated. Use SCOPED_DRAW_EVENT_GAMETHREAD instead."                          ) SCOPED_DRAW_EVENT_GAMETHREAD(Name)
-#define SCOPED_DRAW_EVENTF_COLOR_GAMETHREAD(Color, Name, Format, ...)                                UE_DEPRECATED_MACRO(5.5, "SCOPED_DRAW_EVENTF_COLOR_GAMETHREAD has been deprecated. Use SCOPED_DRAW_EVENTF_GAMETHREAD instead."                        ) SCOPED_DRAW_EVENTF_GAMETHREAD(Name, Format, ##__VA_ARGS__)
-#define SCOPED_CONDITIONAL_DRAW_EVENT_COLOR_GAMETHREAD(Name, Color, Condition)                       UE_DEPRECATED_MACRO(5.5, "SCOPED_CONDITIONAL_DRAW_EVENT_COLOR_GAMETHREAD has been deprecated. Use SCOPED_CONDITIONAL_DRAW_EVENT_GAMETHREAD instead."  ) SCOPED_CONDITIONAL_DRAW_EVENT_GAMETHREAD(Name, Condition)
-#define SCOPED_CONDITIONAL_DRAW_EVENTF_COLOR_GAMETHREAD(Color, Name, Condition, Format, ...)         UE_DEPRECATED_MACRO(5.5, "SCOPED_CONDITIONAL_DRAW_EVENTF_COLOR_GAMETHREAD has been deprecated. Use SCOPED_CONDITIONAL_DRAW_EVENTF_GAMETHREAD instead.") SCOPED_CONDITIONAL_DRAW_EVENTF_GAMETHREAD(Name, Condition, Format, ##__VA_ARGS__)
+#define SCOPED_DRAW_EVENT_COLOR(RHICmdList, Color, Name)                                             UE_DEPRECATED_MACRO(5.5, "SCOPED_DRAW_EVENT_COLOR has been deprecated. Use RHI_BREADCRUMB_EVENT instead."                                               ) RHI_BREADCRUMB_EVENT(RHICmdList, #Name)
+#define SCOPED_GPU_EVENT(RHICmdList, Name)                                                           UE_DEPRECATED_MACRO(5.5, "SCOPED_GPU_EVENT has been deprecated. Use RHI_BREADCRUMB_EVENT instead."                                                      ) RHI_BREADCRUMB_EVENT(RHICmdList, #Name)
+#define SCOPED_GPU_EVENT_COLOR(RHICmdList, Color, Name)                                              UE_DEPRECATED_MACRO(5.5, "SCOPED_GPU_EVENT_COLOR has been deprecated. Use RHI_BREADCRUMB_EVENT instead."                                                ) RHI_BREADCRUMB_EVENT(RHICmdList, #Name)
+#define SCOPED_COMPUTE_EVENT(RHICmdList, Name)                                                       UE_DEPRECATED_MACRO(5.5, "SCOPED_COMPUTE_EVENT has been deprecated. Use RHI_BREADCRUMB_EVENT instead."                                                  ) RHI_BREADCRUMB_EVENT(RHICmdList, #Name)
+#define SCOPED_COMPUTE_EVENT_COLOR(RHICmdList, Color, Name)                                          UE_DEPRECATED_MACRO(5.5, "SCOPED_COMPUTE_EVENT_COLOR has been deprecated. Use RHI_BREADCRUMB_EVENT instead."                                            ) RHI_BREADCRUMB_EVENT(RHICmdList, #Name)
+#define SCOPED_DRAW_EVENTF_COLOR(RHICmdList, Color, Name, Format, ...)                               UE_DEPRECATED_MACRO(5.5, "SCOPED_DRAW_EVENTF_COLOR has been deprecated. Use RHI_BREADCRUMB_EVENT instead."                                              ) RHI_BREADCRUMB_EVENT(RHICmdList, Format, ##__VA_ARGS__)
+#define SCOPED_GPU_EVENTF(RHICmdList, Name, Format, ...)                                             UE_DEPRECATED_MACRO(5.5, "SCOPED_GPU_EVENTF has been deprecated. Use RHI_BREADCRUMB_EVENT instead."                                                     ) RHI_BREADCRUMB_EVENT(RHICmdList, Format, ##__VA_ARGS__)
+#define SCOPED_GPU_EVENTF_COLOR(RHICmdList, Color, Name, Format, ...)                                UE_DEPRECATED_MACRO(5.5, "SCOPED_GPU_EVENTF_COLOR has been deprecated. Use RHI_BREADCRUMB_EVENT instead."                                               ) RHI_BREADCRUMB_EVENT(RHICmdList, Format, ##__VA_ARGS__)
+#define SCOPED_COMPUTE_EVENTF(RHICmdList, Name, Format, ...)                                         UE_DEPRECATED_MACRO(5.5, "SCOPED_COMPUTE_EVENTF has been deprecated. Use RHI_BREADCRUMB_EVENT instead."                                                 ) RHI_BREADCRUMB_EVENT(RHICmdList, Format, ##__VA_ARGS__)
+#define SCOPED_COMPUTE_EVENTF_COLOR(RHICmdList, Color, Name, Format, ...)                            UE_DEPRECATED_MACRO(5.5, "SCOPED_COMPUTE_EVENTF_COLOR has been deprecated. Use RHI_BREADCRUMB_EVENT instead."                                           ) RHI_BREADCRUMB_EVENT(RHICmdList, Format, ##__VA_ARGS__)
+#define SCOPED_CONDITIONAL_DRAW_EVENT_COLOR(RHICmdList, Name, Color, Condition)                      UE_DEPRECATED_MACRO(5.5, "SCOPED_CONDITIONAL_DRAW_EVENT_COLOR has been deprecated. Use RHI_BREADCRUMB_EVENT_CONDITIONAL instead."                       ) RHI_BREADCRUMB_EVENT_CONDITIONAL(RHICmdList, Condition, #Name)
+#define SCOPED_CONDITIONAL_GPU_EVENT(RHICmdList, Name, Condition)                                    UE_DEPRECATED_MACRO(5.5, "SCOPED_CONDITIONAL_GPU_EVENT has been deprecated. Use RHI_BREADCRUMB_EVENT_CONDITIONAL instead."                              ) RHI_BREADCRUMB_EVENT_CONDITIONAL(RHICmdList, Condition, #Name)
+#define SCOPED_CONDITIONAL_GPU_EVENT_COLOR(RHICmdList, Name, Color, Condition)                       UE_DEPRECATED_MACRO(5.5, "SCOPED_CONDITIONAL_GPU_EVENT_COLOR has been deprecated. Use RHI_BREADCRUMB_EVENT_CONDITIONAL instead."                        ) RHI_BREADCRUMB_EVENT_CONDITIONAL(RHICmdList, Condition, #Name)
+#define SCOPED_CONDITIONAL_COMPUTE_EVENT(RHICmdList, Name, Condition)                                UE_DEPRECATED_MACRO(5.5, "SCOPED_CONDITIONAL_COMPUTE_EVENT has been deprecated. Use RHI_BREADCRUMB_EVENT_CONDITIONAL instead."                          ) RHI_BREADCRUMB_EVENT_CONDITIONAL(RHICmdList, Condition, #Name)
+#define SCOPED_CONDITIONAL_COMPUTE_EVENT_COLOR(RHICmdList, Name, Color, Condition)                   UE_DEPRECATED_MACRO(5.5, "SCOPED_CONDITIONAL_COMPUTE_EVENT_COLOR has been deprecated. Use RHI_BREADCRUMB_EVENT_CONDITIONAL instead."                    ) RHI_BREADCRUMB_EVENT_CONDITIONAL(RHICmdList, Condition, #Name)
+#define SCOPED_CONDITIONAL_DRAW_EVENTF_COLOR(RHICmdList, Color, Name, Condition, Format, ...)        UE_DEPRECATED_MACRO(5.5, "SCOPED_CONDITIONAL_DRAW_EVENTF_COLOR has been deprecated. Use RHI_BREADCRUMB_EVENT_CONDITIONAL instead."                      ) RHI_BREADCRUMB_EVENT_CONDITIONAL(RHICmdList, Condition, Format, ##__VA_ARGS__)
+#define SCOPED_CONDITIONAL_GPU_EVENTF(RHICmdList, Name, Condition, Format, ...)                      UE_DEPRECATED_MACRO(5.5, "SCOPED_CONDITIONAL_GPU_EVENTF has been deprecated. Use RHI_BREADCRUMB_EVENT_CONDITIONAL instead."                             ) RHI_BREADCRUMB_EVENT_CONDITIONAL(RHICmdList, Condition, Format, ##__VA_ARGS__)
+#define SCOPED_CONDITIONAL_GPU_EVENTF_COLOR(RHICmdList, Color, Name, Condition, Format, ...)         UE_DEPRECATED_MACRO(5.5, "SCOPED_CONDITIONAL_GPU_EVENTF_COLOR has been deprecated. Use RHI_BREADCRUMB_EVENT_CONDITIONAL instead."                       ) RHI_BREADCRUMB_EVENT_CONDITIONAL(RHICmdList, Condition, Format, ##__VA_ARGS__)
+#define SCOPED_CONDITIONAL_COMPUTE_EVENTF(RHICmdList, Name, Condition, Format, ...)                  UE_DEPRECATED_MACRO(5.5, "SCOPED_CONDITIONAL_COMPUTE_EVENTF has been deprecated. Use RHI_BREADCRUMB_EVENT_CONDITIONAL instead."                         ) RHI_BREADCRUMB_EVENT_CONDITIONAL(RHICmdList, Condition, Format, ##__VA_ARGS__)
+#define SCOPED_CONDITIONAL_COMPUTE_EVENTF_COLOR(RHICmdList, Color, Name, Condition, Format, ...)     UE_DEPRECATED_MACRO(5.5, "SCOPED_CONDITIONAL_COMPUTE_EVENTF_COLOR has been deprecated. Use RHI_BREADCRUMB_EVENT_CONDITIONAL instead."                   ) RHI_BREADCRUMB_EVENT_CONDITIONAL(RHICmdList, Condition, Format, ##__VA_ARGS__)
+#define BEGIN_DRAW_EVENTF_COLOR(RHICmdList, Color, Name, Event, Format, ...)                         UE_DEPRECATED_MACRO(5.5, "BEGIN_DRAW_EVENTF_COLOR has been deprecated. Use BEGIN_DRAW_EVENTF instead."                                                  ) BEGIN_DRAW_EVENTF(RHICmdList, Name, Event, Format, ##__VA_ARGS__)
+#define BEGIN_GPU_EVENTF(RHICmdList, Name, Event, Format, ...)                                       UE_DEPRECATED_MACRO(5.5, "BEGIN_GPU_EVENTF has been deprecated. Use BEGIN_DRAW_EVENTF instead."                                                         ) BEGIN_DRAW_EVENTF(RHICmdList, Name, Event, Format, ##__VA_ARGS__)
+#define BEGIN_GPU_EVENTF_COLOR(RHICmdList, Color, Name, Event, Format, ...)                          UE_DEPRECATED_MACRO(5.5, "BEGIN_GPU_EVENTF_COLOR has been deprecated. Use BEGIN_DRAW_EVENTF instead."                                                   ) BEGIN_DRAW_EVENTF(RHICmdList, Name, Event, Format, ##__VA_ARGS__)
+#define STOP_GPU_EVENT(Event)                                                                        UE_DEPRECATED_MACRO(5.5, "STOP_GPU_EVENT has been deprecated. Use STOP_DRAW_EVENT instead."                                                             ) STOP_DRAW_EVENT(Event)
+#define SCOPED_DRAW_EVENT_COLOR_GAMETHREAD(Color, Name)                                              UE_DEPRECATED_MACRO(5.5, "SCOPED_DRAW_EVENT_COLOR_GAMETHREAD has been deprecated. Use RHI_BREADCRUMB_EVENT_GAMETHREAD instead."                         ) RHI_BREADCRUMB_EVENT_GAMETHREAD(#Name)
+#define SCOPED_DRAW_EVENTF_COLOR_GAMETHREAD(Color, Name, Format, ...)                                UE_DEPRECATED_MACRO(5.5, "SCOPED_DRAW_EVENTF_COLOR_GAMETHREAD has been deprecated. Use RHI_BREADCRUMB_EVENT_GAMETHREAD instead."                        ) RHI_BREADCRUMB_EVENT_GAMETHREAD(Format, ##__VA_ARGS__)
+#define SCOPED_CONDITIONAL_DRAW_EVENT_COLOR_GAMETHREAD(Name, Color, Condition)                       UE_DEPRECATED_MACRO(5.5, "SCOPED_CONDITIONAL_DRAW_EVENT_COLOR_GAMETHREAD has been deprecated. Use RHI_BREADCRUMB_EVENT_CONDITIONAL_GAMETHREAD instead." ) RHI_BREADCRUMB_EVENT_CONDITIONAL_GAMETHREAD(Condition, #Name)
+#define SCOPED_CONDITIONAL_DRAW_EVENTF_COLOR_GAMETHREAD(Color, Name, Condition, Format, ...)         UE_DEPRECATED_MACRO(5.5, "SCOPED_CONDITIONAL_DRAW_EVENTF_COLOR_GAMETHREAD has been deprecated. Use RHI_BREADCRUMB_EVENT_CONDITIONAL_GAMETHREAD instead.") RHI_BREADCRUMB_EVENT_CONDITIONAL_GAMETHREAD(Condition, Format, ##__VA_ARGS__)
 #define SCOPED_RHI_DRAW_EVENT(RHICmdContext, Name)                                                   UE_DEPRECATED_MACRO(5.5, "SCOPED_RHI_DRAW_EVENT has been deprecated. Use standard RHI breadcrumb events instead."                   )
 #define SCOPED_RHI_DRAW_EVENTF(RHICmdContext, Name, Format, ...)                                     UE_DEPRECATED_MACRO(5.5, "SCOPED_RHI_DRAW_EVENTF has been deprecated. Use standard RHI breadcrumb events instead."                  )
 #define SCOPED_RHI_CONDITIONAL_DRAW_EVENT(RHICmdContext, Name, Condition)                            UE_DEPRECATED_MACRO(5.5, "SCOPED_RHI_CONDITIONAL_DRAW_EVENT has been deprecated. Use standard RHI breadcrumb events instead."       )
@@ -147,23 +139,42 @@
 
 #if RHI_NEW_GPU_PROFILER
 
-	// @todo
-	#define DECLARE_GPU_STAT(StatName)
-	#define DECLARE_GPU_STAT_NAMED(StatName, NameString)
-	#define DECLARE_GPU_DRAWCALL_STAT(StatName)
-	#define DECLARE_GPU_DRAWCALL_STAT_NAMED(StatName, NameString)
-	#define DECLARE_GPU_DRAWCALL_STAT_EXTERN(StatName)
-
-	#define DECLARE_GPU_STAT_NAMED_EXTERN(StatName, NameString)
-	#define DEFINE_GPU_STAT(StatName)
-	#define DEFINE_GPU_DRAWCALL_STAT(StatName)
-
-	#define SCOPED_GPU_STAT_VERBOSE(RHICmdList, StatName, Description)
-	#define SCOPED_GPU_STAT(RHICmdList, StatName) 
-
-	#define GPU_STATS_BEGINFRAME(RHICmdList) 
-	#define GPU_STATS_ENDFRAME(RHICmdList) 
-	#define GPU_STATS_SUSPENDFRAME()
+	#if HAS_GPU_STATS
+
+		CSV_DECLARE_CATEGORY_MODULE_EXTERN(RENDERCORE_API, GPU);
+
+		// @todo
+		#define DECLARE_GPU_STAT(                StatName            ) DECLARE_FLOAT_COUNTER_STAT(TEXT(#StatName)       , Stat_GPU_##StatName, STATGROUP_GPU  ); CSV_DEFINE_STAT(GPU,StatName);
+		#define DECLARE_GPU_STAT_NAMED(          StatName, NameString) DECLARE_FLOAT_COUNTER_STAT(NameString            , Stat_GPU_##StatName, STATGROUP_GPU  ); CSV_DEFINE_STAT(GPU,StatName);
+		#define DECLARE_GPU_DRAWCALL_STAT(       StatName            ) DECLARE_FLOAT_COUNTER_STAT(TEXT(#StatName)       , Stat_GPU_##StatName, STATGROUP_GPU  ); CSV_DEFINE_STAT(GPU,StatName);
+		#define DECLARE_GPU_DRAWCALL_STAT_NAMED( StatName, NameString) DECLARE_FLOAT_COUNTER_STAT(NameString            , Stat_GPU_##StatName, STATGROUP_GPU  ); CSV_DEFINE_STAT(GPU,StatName);
+		#define DECLARE_GPU_DRAWCALL_STAT_EXTERN(StatName            ) DECLARE_FLOAT_COUNTER_STAT_EXTERN(TEXT(#StatName), Stat_GPU_##StatName, STATGROUP_GPU, ); CSV_DECLARE_STAT_EXTERN(GPU,StatName);
+
+		// Extern GPU stats are needed where a stat is used in multiple CPPs. Use the DECLARE_GPU_STAT_NAMED_EXTERN in the header and DEFINE_GPU_STAT in the CPPs
+		#define DECLARE_GPU_STAT_NAMED_EXTERN(StatName, NameString) DECLARE_FLOAT_COUNTER_STAT_EXTERN(NameString, Stat_GPU_##StatName, STATGROUP_GPU, ); CSV_DECLARE_STAT_EXTERN(GPU,StatName);
+		#define DEFINE_GPU_STAT(              StatName            ) DEFINE_STAT(Stat_GPU_##StatName);                                                    CSV_DEFINE_STAT(GPU,StatName);
+		#define DEFINE_GPU_DRAWCALL_STAT(     StatName            ) DEFINE_STAT(Stat_GPU_##StatName);                                                    CSV_DEFINE_STAT(GPU,StatName);
+
+	#else
+
+		#define DECLARE_GPU_STAT(...)
+		#define DECLARE_GPU_STAT_NAMED(...)
+		#define DECLARE_GPU_DRAWCALL_STAT(...)
+		#define DECLARE_GPU_DRAWCALL_STAT_NAMED(...)
+		#define DECLARE_GPU_DRAWCALL_STAT_EXTERN(...)
+		#define DECLARE_GPU_STAT_NAMED_EXTERN(...)
+		#define DEFINE_GPU_STAT(...)
+		#define DEFINE_GPU_DRAWCALL_STAT(...)
+
+	#endif
+
+	// Empty when using the new RHI GPU profiler. GPU stats are handled via RHI_BREADCRUMB_EVENT_STAT etc.
+	// @todo deprecate
+	#define SCOPED_GPU_STAT_VERBOSE(...)
+	#define SCOPED_GPU_STAT(...)
+	#define GPU_STATS_BEGINFRAME(...) 
+	#define GPU_STATS_ENDFRAME(...) 
+	#define GPU_STATS_SUSPENDFRAME(...)
 
 #else
 
diff --git a/Engine/Source/Runtime/RenderCore/Public/RenderGraphEvent.h b/Engine/Source/Runtime/RenderCore/Public/RenderGraphEvent.h
--- a/Engine/Source/Runtime/RenderCore/Public/RenderGraphEvent.h
+++ b/Engine/Source/Runtime/RenderCore/Public/RenderGraphEvent.h
@@ -49,7 +49,7 @@ public:
 	const TCHAR* GetTCHAR() const;
 
 #if WITH_RHI_BREADCRUMBS
-	FRHIBreadcrumbNode* AllocBreadcrumb(FRHIBreadcrumbAllocator& Allocator) const;
+	FRHIBreadcrumbNode* AllocBreadcrumb(TStatId StatId, FRHIBreadcrumbAllocator& Allocator) const;
 #endif // WITH_RHI_BREADCRUMBS
 
 private:
@@ -74,6 +74,9 @@ enum class ERDGScopeFlags : uint8
 
 	// Ensures the scope is always emitted (ignores cvars that disable scopes)
 	AlwaysEnable = 1 << 1,
+
+	// The scope includes a GPU stat, so may need to be enabled even when cvars are disabling scopes.
+	Stat = 1 << 2,
 };
 ENUM_CLASS_FLAGS(ERDGScopeFlags);
 
@@ -168,7 +171,7 @@ struct FRDGScope_Budget
 		FRHIBreadcrumbNode* Node = nullptr;
 	#endif
 
-		inline FRDGScope_RHI(FRDGScopeState& State, FRDGEventName&& Name);
+		inline FRDGScope_RHI(FRDGScopeState& State, TStatId StatId, FRDGEventName&& Name);
 		inline void ImmediateEnd(FRDGScopeState& State);
 
 	#if WITH_RHI_BREADCRUMBS
@@ -341,10 +344,19 @@ private:
 	// Skip expensive string formatting for the relatively common case of no varargs.  We detect this by stringizing the varargs and checking if the string is non-empty (more than just a null terminator).
 	#define RDG_EVENT_NAME(Format, ...) (sizeof(#__VA_ARGS__ "") > 1 ? FRDGEventName(TEXT(Format), ##__VA_ARGS__) : FRDGEventName(1, TEXT(Format)))
 
-	#define RDG_EVENT_SCOPE(GraphBuilder, Format, ...)                                  \
+	#define RDG_EVENT_SCOPE(GraphBuilder, Format, ...) \
 		TRDGEventScopeGuard<FRDGScope_RHI> PREPROCESSOR_JOIN(__RDG_ScopeRef_,__LINE__)( \
 			(GraphBuilder)                                                              \
 			, ERDGScopeFlags::None                                                      \
+			, TStatId()                                                                 \
+			, RDG_EVENT_NAME(Format, ##__VA_ARGS__)                                     \
+		)
+
+	#define RDG_EVENT_SCOPE_STAT(GraphBuilder, StatName, Format, ...)                   \
+		TRDGEventScopeGuard<FRDGScope_RHI> PREPROCESSOR_JOIN(__RDG_ScopeRef_,__LINE__)( \
+			(GraphBuilder)                                                              \
+			, ERDGScopeFlags::Stat                                                      \
+			, GET_STATID(Stat_GPU_##StatName)                                           \
 			, RDG_EVENT_NAME(Format, ##__VA_ARGS__)                                     \
 		)
 
@@ -357,6 +369,22 @@ private:
 				PREPROCESSOR_JOIN(__RDG_ScopeRef_,__LINE__).Emplace(                               \
 					(GraphBuilder)                                                                 \
 					, ERDGScopeFlags::None                                                         \
+					, TStatId()                                                                    \
+					, RDG_EVENT_NAME(Format, ##__VA_ARGS__)                                        \
+				);                                                                                 \
+			}   \
+		} while (false)
+
+	#define RDG_EVENT_SCOPE_CONDITIONAL_STAT(GraphBuilder, Condition, StatName, Format, ...)       \
+		TOptional<TRDGEventScopeGuard<FRDGScope_RHI>> PREPROCESSOR_JOIN(__RDG_ScopeRef_,__LINE__); \
+		do                                                                                         \
+		{                                                                                          \
+			if (Condition)                                                                         \
+			{                                                                                      \
+				PREPROCESSOR_JOIN(__RDG_ScopeRef_,__LINE__).Emplace(                               \
+					(GraphBuilder)                                                                 \
+					, ERDGScopeFlags::Stat                                                         \
+					, GET_STATID(Stat_GPU_##StatName)                                              \
 					, RDG_EVENT_NAME(Format, ##__VA_ARGS__)                                        \
 				);                                                                                 \
 			}   \
@@ -371,6 +399,7 @@ private:
 			PREPROCESSOR_JOIN(__RDG_ScopeRef_,__LINE__).Emplace(                                   \
 				(GraphBuilder)                                                                     \
 				, ERDGScopeFlags::Final                                                            \
+				, TStatId()                                                                        \
 				, RDG_EVENT_NAME(Format, ##__VA_ARGS__)                                            \
 			);                                                                                     \
 		} while (false)
@@ -400,8 +429,9 @@ private:
 #endif
 
  /** Injects a scope onto both the RDG and RHI timeline. */
-#define RDG_RHI_EVENT_SCOPE(GraphBuilder, Name) RDG_EVENT_SCOPE(GraphBuilder, #Name); SCOPED_DRAW_EVENT(GraphBuilder.RHICmdList, Name);
-#define RDG_RHI_GPU_STAT_SCOPE(GraphBuilder, StatName) RDG_GPU_STAT_SCOPE(GraphBuilder, StatName); SCOPED_GPU_STAT(GraphBuilder.RHICmdList, StatName);
+#define RDG_RHI_EVENT_SCOPE(     GraphBuilder, Name)       RDG_EVENT_SCOPE(GraphBuilder, #Name);            RHI_BREADCRUMB_EVENT(GraphBuilder.RHICmdList, #Name)
+#define RDG_RHI_EVENT_SCOPE_STAT(GraphBuilder, Stat, Name) RDG_EVENT_SCOPE_STAT(GraphBuilder, Stat, #Name); RHI_BREADCRUMB_EVENT_STAT(GraphBuilder.RHICmdList, Stat, #Name)
+#define RDG_RHI_GPU_STAT_SCOPE(  GraphBuilder, StatName)   RDG_GPU_STAT_SCOPE(GraphBuilder, StatName);      SCOPED_GPU_STAT(GraphBuilder.RHICmdList, StatName);
 
 namespace DynamicRenderScaling
 {
diff --git a/Engine/Source/Runtime/RenderCore/Public/RenderGraphEvent.inl b/Engine/Source/Runtime/RenderCore/Public/RenderGraphEvent.inl
--- a/Engine/Source/Runtime/RenderCore/Public/RenderGraphEvent.inl
+++ b/Engine/Source/Runtime/RenderCore/Public/RenderGraphEvent.inl
@@ -100,10 +100,10 @@ inline void FRDGScope_GPU::EndGPU(FRHIComputeCommandList& RHICmdList)
 
 #if RDG_EVENTS
 
-	inline FRDGScope_RHI::FRDGScope_RHI(FRDGScopeState& State, FRDGEventName&& Name)
+	inline FRDGScope_RHI::FRDGScope_RHI(FRDGScopeState& State, TStatId StatId, FRDGEventName&& Name)
 		: Name(MoveTemp(Name))
 	#if WITH_RHI_BREADCRUMBS
-		, Node(Name.AllocBreadcrumb(State.GetBreadcrumbAllocator()))
+		, Node(Name.AllocBreadcrumb(StatId, State.GetBreadcrumbAllocator()))
 	#endif
 	{
 	#if WITH_RHI_BREADCRUMBS
diff --git a/Engine/Source/Runtime/Renderer/Private/MeshPassProcessor.cpp b/Engine/Source/Runtime/Renderer/Private/MeshPassProcessor.cpp
--- a/Engine/Source/Runtime/Renderer/Private/MeshPassProcessor.cpp
+++ b/Engine/Source/Runtime/Renderer/Private/MeshPassProcessor.cpp
@@ -2255,6 +2255,7 @@ void FPassProcessorManager::SetPassFlags(EShadingPath ShadingPath, EMeshPass::Ty
 FMeshDrawCommand::FMeshDrawEvent::FMeshDrawEvent(const FMeshDrawCommand& MeshDrawCommand, const uint32 InstanceFactor, FRHICommandList& RHICmdList)
 	: Breadcrumb(
 		  RHICmdList
+		, TStatId()
 		, GShowMaterialDrawEvents != 0
 		, TEXT("%s %s (%u instances)")
 		, MeshDrawCommand.DebugData.MaterialRenderProxy->GetMaterialName()
diff --git a/Engine/Source/Runtime/Renderer/Private/ShadowDepthRendering.cpp b/Engine/Source/Runtime/Renderer/Private/ShadowDepthRendering.cpp
--- a/Engine/Source/Runtime/Renderer/Private/ShadowDepthRendering.cpp
+++ b/Engine/Source/Runtime/Renderer/Private/ShadowDepthRendering.cpp
@@ -1586,7 +1586,7 @@ void FSceneRenderer::RenderShadowDepthMapAtlases(FRDGBuilder& GraphBuilder)
 				FString LightNameWithLevel;
 				GetLightNameForDrawEvent(CurrentLightForDrawEvent, LightNameWithLevel);
 
-				RDGScope.Emplace(GraphBuilder, ERDGScopeFlags::None, RDG_EVENT_NAME("%s", *LightNameWithLevel));
+				RDGScope.Emplace(GraphBuilder, ERDGScopeFlags::None, TStatId(), RDG_EVENT_NAME("%s", *LightNameWithLevel));
 			}
 		#endif
 		};
diff --git a/Engine/Source/Runtime/VulkanRHI/Private/VulkanViewport.cpp b/Engine/Source/Runtime/VulkanRHI/Private/VulkanViewport.cpp
--- a/Engine/Source/Runtime/VulkanRHI/Private/VulkanViewport.cpp
+++ b/Engine/Source/Runtime/VulkanRHI/Private/VulkanViewport.cpp
@@ -878,7 +878,7 @@ void FVulkanViewport::DestroySwapchain(FVulkanSwapChainRecreateInfo* RecreateInf
 
 inline static void CopyImageToBackBuffer(FVulkanCommandListContext* Context, FVulkanCmdBuffer* CmdBuffer, FVulkanTexture& SrcSurface, FVulkanTexture& DstSurface, int32 SizeX, int32 SizeY, int32 WindowSizeX, int32 WindowSizeY, VkSurfaceTransformFlagBitsKHR CachedSurfaceTransform)
 {
-	RHI_BREADCRUMB_EVENT(*Context, CopyImageToBackBuffer);
+	RHI_BREADCRUMB_EVENT(*Context, "CopyImageToBackBuffer");
 	const bool bNeedsVulkanPreTransform = CachedSurfaceTransform != VK_SURFACE_TRANSFORM_IDENTITY_BIT_KHR;
 
 	FVulkanLayoutManager& LayoutManager = CmdBuffer->GetLayoutManager();
