Add missing RDG_EVENT_SCOPE_STAT and RDG_EVENT_SCOPE_CONDITIONAL_STAT macros when RDG_EVENTS is 0

diff --git a/Engine/Source/Runtime/RenderCore/Public/RenderGraphEvent.h b/Engine/Source/Runtime/RenderCore/Public/RenderGraphEvent.h
--- a/Engine/Source/Runtime/RenderCore/Public/RenderGraphEvent.h
+++ b/Engine/Source/Runtime/RenderCore/Public/RenderGraphEvent.h
@@ -406,11 +406,13 @@ private:
 
 #else
 
-	#define RDG_EVENT_NAME(Format, ...) FRDGEventName()
+	#define RDG_EVENT_NAME(...) FRDGEventName()
 
-	#define RDG_EVENT_SCOPE(            GraphBuilder,            Format, ...) do { } while (false)
-	#define RDG_EVENT_SCOPE_CONDITIONAL(GraphBuilder, Condition, Format, ...) do { } while (false)
-	#define RDG_EVENT_SCOPE_FINAL(      GraphBuilder,            Format, ...) do { } while (false)
+	#define RDG_EVENT_SCOPE(...)                  do { } while (false)
+	#define RDG_EVENT_SCOPE_STAT(...)             do { } while (false)
+	#define RDG_EVENT_SCOPE_CONDITIONAL(...)      do { } while (false)
+	#define RDG_EVENT_SCOPE_CONDITIONAL_STAT(...) do { } while (false)
+	#define RDG_EVENT_SCOPE_FINAL(...)            do { } while (false)
 
 #endif
 
