Fix build errors when RDG events are enabled but GPU stats are disabled
 - The new RDG_EVENT_SCOPE_STAT macros need to drop the StatName field when HAS_GPU_STATS is zero.

diff --git a/Engine/Source/Runtime/RenderCore/Public/RenderGraphEvent.h b/Engine/Source/Runtime/RenderCore/Public/RenderGraphEvent.h
--- a/Engine/Source/Runtime/RenderCore/Public/RenderGraphEvent.h
+++ b/Engine/Source/Runtime/RenderCore/Public/RenderGraphEvent.h
@@ -352,13 +352,18 @@ private:
 			, RDG_EVENT_NAME(Format, ##__VA_ARGS__)                                     \
 		)
 
-	#define RDG_EVENT_SCOPE_STAT(GraphBuilder, StatName, Format, ...)                   \
-		TRDGEventScopeGuard<FRDGScope_RHI> PREPROCESSOR_JOIN(__RDG_ScopeRef_,__LINE__)( \
-			(GraphBuilder)                                                              \
-			, ERDGScopeFlags::Stat                                                      \
-			, GET_STATID(Stat_GPU_##StatName)                                           \
-			, RDG_EVENT_NAME(Format, ##__VA_ARGS__)                                     \
-		)
+	#if HAS_GPU_STATS
+		#define RDG_EVENT_SCOPE_STAT(GraphBuilder, StatName, Format, ...)                   \
+			TRDGEventScopeGuard<FRDGScope_RHI> PREPROCESSOR_JOIN(__RDG_ScopeRef_,__LINE__)( \
+				(GraphBuilder)                                                              \
+				, ERDGScopeFlags::Stat                                                      \
+				, GET_STATID(Stat_GPU_##StatName)                                           \
+				, RDG_EVENT_NAME(Format, ##__VA_ARGS__)                                     \
+			)
+	#else
+		#define RDG_EVENT_SCOPE_STAT(GraphBuilder, StatName, Format, ...) \
+			    RDG_EVENT_SCOPE(GraphBuilder, Format, ##__VA_ARGS__)
+	#endif
 
 	#define RDG_EVENT_SCOPE_CONDITIONAL(GraphBuilder, Condition, Format, ...)                      \
 		TOptional<TRDGEventScopeGuard<FRDGScope_RHI>> PREPROCESSOR_JOIN(__RDG_ScopeRef_,__LINE__); \
@@ -375,20 +380,25 @@ private:
 			}   \
 		} while (false)
 
-	#define RDG_EVENT_SCOPE_CONDITIONAL_STAT(GraphBuilder, Condition, StatName, Format, ...)       \
-		TOptional<TRDGEventScopeGuard<FRDGScope_RHI>> PREPROCESSOR_JOIN(__RDG_ScopeRef_,__LINE__); \
-		do                                                                                         \
-		{                                                                                          \
-			if (Condition)                                                                         \
-			{                                                                                      \
-				PREPROCESSOR_JOIN(__RDG_ScopeRef_,__LINE__).Emplace(                               \
-					(GraphBuilder)                                                                 \
-					, ERDGScopeFlags::Stat                                                         \
-					, GET_STATID(Stat_GPU_##StatName)                                              \
-					, RDG_EVENT_NAME(Format, ##__VA_ARGS__)                                        \
-				);                                                                                 \
-			}   \
-		} while (false)
+	#if HAS_GPU_STATS
+		#define RDG_EVENT_SCOPE_CONDITIONAL_STAT(GraphBuilder, Condition, StatName, Format, ...)       \
+			TOptional<TRDGEventScopeGuard<FRDGScope_RHI>> PREPROCESSOR_JOIN(__RDG_ScopeRef_,__LINE__); \
+			do                                                                                         \
+			{                                                                                          \
+				if (Condition)                                                                         \
+				{                                                                                      \
+					PREPROCESSOR_JOIN(__RDG_ScopeRef_,__LINE__).Emplace(                               \
+						(GraphBuilder)                                                                 \
+						, ERDGScopeFlags::Stat                                                         \
+						, GET_STATID(Stat_GPU_##StatName)                                              \
+						, RDG_EVENT_NAME(Format, ##__VA_ARGS__)                                        \
+					);                                                                                 \
+				}   \
+			} while (false)
+	#else
+		#define RDG_EVENT_SCOPE_CONDITIONAL_STAT(GraphBuilder, Condition, StatName, Format, ...) \
+			    RDG_EVENT_SCOPE_CONDITIONAL(GraphBuilder, Condition, Format, ##__VA_ARGS__)
+	#endif
 
 	// The 'Final' version disables any further child scopes or pass events. It is intended
 	// to group overlapping passes as events can disable overlap on certain GPUs.
