Fixed warning about uninitialized member.

diff --git a/Engine/Source/Runtime/RenderCore/Private/RenderGraphPass.cpp b/Engine/Source/Runtime/RenderCore/Private/RenderGraphPass.cpp
--- a/Engine/Source/Runtime/RenderCore/Private/RenderGraphPass.cpp
+++ b/Engine/Source/Runtime/RenderCore/Private/RenderGraphPass.cpp
@@ -4,6 +4,8 @@
 #include "RenderGraphPrivate.h"
 #include "RenderGraphBuilder.h"
 
+FEmptyShaderParameters FRDGSentinelPass::EmptyShaderParameters;
+
 FUniformBufferStaticBindings FRDGParameterStruct::GetStaticUniformBuffers() const
 {
 	FUniformBufferStaticBindings GlobalUniformBuffers;
diff --git a/Engine/Source/Runtime/RenderCore/Public/RenderGraphPass.h b/Engine/Source/Runtime/RenderCore/Public/RenderGraphPass.h
--- a/Engine/Source/Runtime/RenderCore/Public/RenderGraphPass.h
+++ b/Engine/Source/Runtime/RenderCore/Public/RenderGraphPass.h
@@ -793,10 +793,13 @@ public:
 	{}
 
 private:
-	FEmptyShaderParameters EmptyShaderParameters;
+	static FEmptyShaderParameters EmptyShaderParameters;
 	friend class FRDGBuilder;
 };
 
+template <typename ExecuteLambdaType>
+FEmptyShaderParameters TRDGEmptyLambdaPass<ExecuteLambdaType>::EmptyShaderParameters;
+
 /** Render graph pass used for the prologue / epilogue passes. */
 class FRDGSentinelPass final
 	: public FRDGPass
@@ -809,7 +812,7 @@ public:
 	}
 
 private:
-	FEmptyShaderParameters EmptyShaderParameters;
+	static FEmptyShaderParameters EmptyShaderParameters;
 };
 
 #include "RenderGraphParameters.inl" // IWYU pragma: export
