Disabled async RDG tasks to investigate Slate crash.

diff --git a/Engine/Source/Runtime/RenderCore/Private/RenderGraphPrivate.cpp b/Engine/Source/Runtime/RenderCore/Private/RenderGraphPrivate.cpp
--- a/Engine/Source/Runtime/RenderCore/Private/RenderGraphPrivate.cpp
+++ b/Engine/Source/Runtime/RenderCore/Private/RenderGraphPrivate.cpp
@@ -342,7 +342,7 @@ FAutoConsoleVariableRef CVarRDGParallelSetup(
 	TEXT(" 1: pass setup is done asynchronously (default);"),
 	ECVF_RenderThreadSafe);
 
-int32 GRDGParallelExecute = 2;
+int32 GRDGParallelExecute = 1;
 FAutoConsoleVariableRef CVarRDGParallelExecute(
 	TEXT("r.RDG.ParallelExecute"), GRDGParallelExecute,
 	TEXT("Whether to enable parallel execution of passes when supported.")
