Shader library fix to support loading chunks scattered over different folders.

Previously, it was looking in the base directory only.
However, when shader library is split between different install bundles, chunks might be scattered over multiple directories.

It actually happens with .metallibs, which are located in different UFS folders now (which is fine).
Refer to CL 35423329 to see why.

The fix is to look through the mounted chunk folder, so shader library picks that.

#rnx

diff --git a/Engine/Source/Runtime/RenderCore/Private/ShaderCodeLibrary.cpp b/Engine/Source/Runtime/RenderCore/Private/ShaderCodeLibrary.cpp
--- a/Engine/Source/Runtime/RenderCore/Private/ShaderCodeLibrary.cpp
+++ b/Engine/Source/Runtime/RenderCore/Private/ShaderCodeLibrary.cpp
@@ -230,17 +230,15 @@ namespace UE
 					return Components.Num();
 				}
 
-				void OnPakFileMounted(const FMountedPakFileInfo& MountInfo)
+				void OnPakFileMounted(const FMountedPakFileInfo& MountInfo, const FString& Directory)
 				{
 					if (!PresentChunks.Contains(MountInfo.ChunkId))
 					{
 						FString ChunkLibraryName = GetShaderLibraryNameForChunk(LogicalName, MountInfo.ChunkId);
 
-						// Ignore chunk mount point as it's useless in locating the actual library directory. For instance, chunks can
-						// have mount points like ../../../ProjectName, while the actual library file is still stored in Content subdirectory.
-						// Just use the base directory always and expect the library to be placed in the same location for all chunks
-						// (which is the current behavior).
-						if (OpenShaderCode(BaseDirectory, ChunkLibraryName))
+						// Parts of shader library might be in the UFS (e.g. .metallibs), hence we need to look for them in the
+						// appropriate directory.
+						if (OpenShaderCode(Directory, ChunkLibraryName))
 						{
 							PresentChunks.Add(MountInfo.ChunkId);
 						}
@@ -2704,7 +2702,7 @@ public:
 					FScopeLock KnownPakFilesLocker(&FMountedPakFileInfo::KnownPakFilesAccessLock);
 					for (TSet<FMountedPakFileInfo>::TConstIterator Iter(FMountedPakFileInfo::KnownPakFiles); Iter; ++Iter)
 					{
-						Library->OnPakFileMounted(*Iter);
+						Library->OnPakFileMounted(*Iter, Directory);
 					}
 				}
 				NewComponentIDs = Library->PresentChunks.Difference(PrevComponentSet);
@@ -2744,7 +2742,7 @@ public:
 									{
 										// create a fake FPakFileMountedInfo
 										FMountedPakFileInfo PakFileInfo(Directory, ChunkID);
-										Library->OnPakFileMounted(PakFileInfo);
+										Library->OnPakFileMounted(PakFileInfo, Directory);
 									}
 								}
 							}
@@ -2896,7 +2894,7 @@ public:
 				for (TTuple<FString, TUniquePtr<UE::ShaderLibrary::Private::FNamedShaderLibrary>>& NamedLibraryPair : NamedLibrariesStack)
 				{
 					TSet<int32> PrevComponentSet = NamedLibraryPair.Value->PresentChunks;
-					NamedLibraryPair.Value->OnPakFileMounted(MountInfo);
+					NamedLibraryPair.Value->OnPakFileMounted(MountInfo, NamedLibraryPair.Value->BaseDirectory);
 					for (int32 ComponentID : NamedLibraryPair.Value->PresentChunks.Difference(PrevComponentSet))
 					{
 						// Defer these for outside of the lock, ShaderPipelineCache may want to inspect shader library.
