Fixed dubious reference parameter.

#rnx

diff --git a/Engine/Source/Runtime/RenderCore/Private/RenderGraphUtils.cpp b/Engine/Source/Runtime/RenderCore/Private/RenderGraphUtils.cpp
--- a/Engine/Source/Runtime/RenderCore/Private/RenderGraphUtils.cpp
+++ b/Engine/Source/Runtime/RenderCore/Private/RenderGraphUtils.cpp
@@ -861,7 +861,7 @@ class FInitIndirectArgs1DCS : public FGlobalShader
 
 IMPLEMENT_GLOBAL_SHADER(FInitIndirectArgs1DCS, "/Engine/Private/Tools/SetupIndirectArgs.usf", "InitIndirectArgs1DCS", SF_Compute);
 
-FRDGBufferRef FComputeShaderUtils::AddIndirectArgsSetupCsPass1D(FRDGBuilder& GraphBuilder, ERHIFeatureLevel::Type FeatureLevel, FRDGBufferRef& InputCountBuffer, const TCHAR* OutputBufferName, uint32 Divisor, uint32 InputCountOffset, uint32 Multiplier)
+FRDGBufferRef FComputeShaderUtils::AddIndirectArgsSetupCsPass1D(FRDGBuilder& GraphBuilder, ERHIFeatureLevel::Type FeatureLevel, FRDGBuffer* InputCountBuffer, const TCHAR* OutputBufferName, uint32 Divisor, uint32 InputCountOffset, uint32 Multiplier)
 {
 	// 1. Add setup pass
 	FRDGBufferRef IndirectArgsBuffer = GraphBuilder.CreateBuffer(FRDGBufferDesc::CreateIndirectDesc<FRHIDispatchIndirectParameters>(), OutputBufferName);
diff --git a/Engine/Source/Runtime/RenderCore/Public/RenderGraphUtils.h b/Engine/Source/Runtime/RenderCore/Public/RenderGraphUtils.h
--- a/Engine/Source/Runtime/RenderCore/Public/RenderGraphUtils.h
+++ b/Engine/Source/Runtime/RenderCore/Public/RenderGraphUtils.h
@@ -700,7 +700,7 @@ namespace FComputeShaderUtils
 	 * 	Sets up a group count as (InputCountBuffer[InputCountOffset] * Multiplier + Divisor - 1U) / Divisor;
 	 *  Commonly use Divisor <=> number of threads per group.
 	 */
-	RENDERCORE_API FRDGBufferRef AddIndirectArgsSetupCsPass1D(FRDGBuilder& GraphBuilder, ERHIFeatureLevel::Type FeatureLevel, FRDGBufferRef& InputCountBuffer, const TCHAR* OutputBufferName, uint32 Divisor, uint32 InputCountOffset = 0U, uint32 Multiplier = 1U);
+	RENDERCORE_API FRDGBufferRef AddIndirectArgsSetupCsPass1D(FRDGBuilder& GraphBuilder, ERHIFeatureLevel::Type FeatureLevel, FRDGBuffer* InputCountBuffer, const TCHAR* OutputBufferName, uint32 Divisor, uint32 InputCountOffset = 0U, uint32 Multiplier = 1U);
 };
 
 /** Adds a render graph pass to copy a region from one texture to another. Uses RHICopyTexture under the hood.
