Remove r.RayTracing.RequireSM6.

[RN] Ray tracing is no longer supported on SM5

diff --git a/Engine/Config/BaseDeviceProfiles.ini b/Engine/Config/BaseDeviceProfiles.ini
--- a/Engine/Config/BaseDeviceProfiles.ini
+++ b/Engine/Config/BaseDeviceProfiles.ini
@@ -916,7 +916,6 @@ BaseProfileName=Android
 +CVars=r.DistanceFields=1
 +CVars=r.Vulkan.RayTracing.AllowCompaction=0
 +CVars=r.Vulkan.RayTracing.TLASPreferFastTraceTLAS=0
-+CVars=r.RayTracing.RequireSM6=0
 
 [Android_Adreno_Vulkan_SM5 DeviceProfile]
 DeviceType=Android
diff --git a/Engine/Config/Windows/DataDrivenPlatformInfo.ini b/Engine/Config/Windows/DataDrivenPlatformInfo.ini
--- a/Engine/Config/Windows/DataDrivenPlatformInfo.ini
+++ b/Engine/Config/Windows/DataDrivenPlatformInfo.ini
@@ -38,14 +38,7 @@ MaxFeatureLevel=SM5
 ShaderFormat=PCD3D_SM5
 bIsPC=true
 bSupportsDebugViewShaders=true
-bSupportsRayTracing=true
-bSupportsRayTracingShaders=true
-bSupportsInlineRayTracing=false
-bSupportsRayTracingCallableShaders=true
-bSupportsRayTracingProceduralPrimitive=true
-bSupportsRayTracingIndirectInstanceData=true
-bSupportsPathTracing=true
-bSupportsHighEndRayTracingEffects=true
+bSupportsRayTracing=false
 bSupportsGPUScene=true
 bSupportsNanite=false
 bSupportsSceneDataCompressedTransforms=true
@@ -68,7 +61,7 @@ bSupportsPercentageCloserShadows=true
 bSupportsIndexBufferUAVs=true
 bSupportsROV=true
 bSupportsOIT=true
-BindlessSupport=RayTracingOnly
+BindlessSupport=Unsupported
 StaticShaderBindingLayoutSupport=RayTracingOnly
 SupportsMultiViewport=RuntimeDependent
 FriendlyName=LOCTEXT("FriendlyShaderPlatformName_D3D_SM5", "PC D3D SM5")
diff --git a/Engine/Source/Developer/Windows/WindowsTargetPlatformSettings/Public/GenericWindowsTargetPlatformSettings.h b/Engine/Source/Developer/Windows/WindowsTargetPlatformSettings/Public/GenericWindowsTargetPlatformSettings.h
--- a/Engine/Source/Developer/Windows/WindowsTargetPlatformSettings/Public/GenericWindowsTargetPlatformSettings.h
+++ b/Engine/Source/Developer/Windows/WindowsTargetPlatformSettings/Public/GenericWindowsTargetPlatformSettings.h
@@ -190,13 +190,6 @@ public:
 
 			// We always support ray tracing shaders when cooking for D3D12 SM6, however we may skip them for SM5 based on project settings.
 			OutFormats.AddUnique(FName(TEXT("PCD3D_SM6")));
-
-			static IConsoleVariable* RequireSM6CVar = IConsoleManager::Get().FindConsoleVariable(TEXT("r.RayTracing.RequireSM6"));
-			const bool bRequireSM6 = RequireSM6CVar && RequireSM6CVar->GetBool();
-			if (!bRequireSM6)
-			{
-				OutFormats.AddUnique(FName(TEXT("PCD3D_SM5")));
-			}
 		}
 	}
 
diff --git a/Engine/Source/Runtime/D3D12RHI/Private/D3D12Adapter.cpp b/Engine/Source/Runtime/D3D12RHI/Private/D3D12Adapter.cpp
--- a/Engine/Source/Runtime/D3D12RHI/Private/D3D12Adapter.cpp
+++ b/Engine/Source/Runtime/D3D12RHI/Private/D3D12Adapter.cpp
@@ -997,10 +997,7 @@ void FD3D12Adapter::InitializeDevices()
 			D3D12_FEATURE_DATA_D3D12_OPTIONS5 D3D12Caps5 = {};
 			if (SUCCEEDED(RootDevice->CheckFeatureSupport(D3D12_FEATURE_D3D12_OPTIONS5, &D3D12Caps5, sizeof(D3D12Caps5))))
 			{
-				static IConsoleVariable* RequireSM6CVar = IConsoleManager::Get().FindConsoleVariable(TEXT("r.RayTracing.RequireSM6"));
-				const bool bRequireSM6 = RequireSM6CVar && RequireSM6CVar->GetBool();
-
-				const bool bRayTracingAllowedOnCurrentShaderPlatform = (bRequireSM6 == false) || (GMaxRHIShaderPlatform == SP_PCD3D_SM6);
+				const bool bRayTracingAllowedOnCurrentShaderPlatform = (GMaxRHIShaderPlatform == SP_PCD3D_SM6);
 
 				if (D3D12Caps5.RaytracingTier >= D3D12_RAYTRACING_TIER_1_0
 					&& GetResourceBindingTier() >= D3D12_RESOURCE_BINDING_TIER_2
@@ -1024,7 +1021,7 @@ void FD3D12Adapter::InitializeDevices()
 						}
 						else
 						{
- 							UE_LOG(LogD3D12RHI, Log, TEXT("Ray tracing is disabled because SM6 shader platform is required (r.RayTracing.RequireSM6=1)."));
+ 							UE_LOG(LogD3D12RHI, Log, TEXT("Ray tracing is disabled because SM6 shader platform is required."));
 						}
 					}
 					else if (GRHIBindlessSupport == ERHIBindlessSupport::Unsupported)
diff --git a/Engine/Source/Runtime/Engine/Private/RendererSettings.cpp b/Engine/Source/Runtime/Engine/Private/RendererSettings.cpp
--- a/Engine/Source/Runtime/Engine/Private/RendererSettings.cpp
+++ b/Engine/Source/Runtime/Engine/Private/RendererSettings.cpp
@@ -497,10 +497,7 @@ void URendererSettings::CheckForMissingShaderModels()
 {
 	// Don't show the SM6 toasts on non-Windows/Linux platforms to avoid confusion around platform requirements.
 #if PLATFORM_WINDOWS || PLATFORM_LINUX
-	static IConsoleVariable* RayTracingRequireSM6CVar = IConsoleManager::Get().FindConsoleVariable(TEXT("r.RayTracing.RequireSM6"));
-	const bool bRayTracingRequireSM6 = bEnableRayTracing && RayTracingRequireSM6CVar && RayTracingRequireSM6CVar->GetBool();
-
-	if (GIsEditor && (ShadowMapMethod == EShadowMapMethod::VirtualShadowMaps || bRayTracingRequireSM6))
+	if (GIsEditor && (ShadowMapMethod == EShadowMapMethod::VirtualShadowMaps || bEnableRayTracing))
 	{
 		auto CopySM6Format = [](const TCHAR* ShaderFormatName, const TArray<FString>& SrcArray, TArray<FString>& DstArray)
 		{
@@ -593,7 +590,7 @@ void URendererSettings::CheckForMissingShaderModels()
 				LOCTEXT("VirtualShadowMapsAndRayTracingNeedsSM6", "Shader Model 6 (SM6) is required to use Virtual Shadow Maps and Ray Tracing.")
 			};
 
-			const uint32 FeatureNeedsSM6Index = (ShadowMapMethod == EShadowMapMethod::VirtualShadowMaps && bRayTracingRequireSM6) ? 2 : (bRayTracingRequireSM6 ? 1 : 0);
+			const uint32 FeatureNeedsSM6Index = (ShadowMapMethod == EShadowMapMethod::VirtualShadowMaps && bEnableRayTracing) ? 2 : (bEnableRayTracing ? 1 : 0);
 
 			if (bProjectMissingD3DSM6)
 			{
diff --git a/Engine/Source/Runtime/RenderCore/Private/RenderUtils.cpp b/Engine/Source/Runtime/RenderCore/Private/RenderUtils.cpp
--- a/Engine/Source/Runtime/RenderCore/Private/RenderUtils.cpp
+++ b/Engine/Source/Runtime/RenderCore/Private/RenderUtils.cpp
@@ -52,14 +52,6 @@ static FAutoConsoleVariableRef CVarRayTracingEnableOnDemand(
 	TEXT(" 1: on"),
 	ECVF_RenderThreadSafe | ECVF_ReadOnly);
 
-static int32 GRayTracingRequireSM6 = 1;
-static FAutoConsoleVariableRef CVarRayTracingRequireSM6(
-	TEXT("r.RayTracing.RequireSM6"),
-	GRayTracingRequireSM6,
-	TEXT("Whether ray tracing shaders and features should only be available when targetting and running SM6. If disabled, ray tracing shaders will also be available when running in SM5 mode. (default = 1, allow only SM6)"),
-	ECVF_RenderThreadSafe | ECVF_ReadOnly
-);
-
 const uint16 GCubeIndices[12*3] =
 {
 	0, 2, 3,
@@ -1610,9 +1602,7 @@ bool ShouldCompileRayTracingShadersForProject(EShaderPlatform ShaderPlatform)
 {
 	if (RHISupportsRayTracingShaders(ShaderPlatform))
 	{
-		const bool bRayTracingRequireSM6 = (GRayTracingRequireSM6 != 0);
-
-		return IsRayTracingEnabledForProject(ShaderPlatform) && AreRayTracingShadersEnabledForProject(ShaderPlatform) && (IsFeatureLevelSupported(ShaderPlatform, ERHIFeatureLevel::SM6) || !bRayTracingRequireSM6);
+		return IsRayTracingEnabledForProject(ShaderPlatform) && AreRayTracingShadersEnabledForProject(ShaderPlatform) && IsFeatureLevelSupported(ShaderPlatform, ERHIFeatureLevel::SM6);
 	}
 	else
 	{
@@ -1656,9 +1646,7 @@ bool IsRayTracingAllowed()
 
 bool IsRayTracingEnabled(EShaderPlatform ShaderPlatform)
 {
-	const bool bRayTracingRequireSM6 = (GRayTracingRequireSM6 != 0);
-
-	return IsRayTracingEnabled() && RHISupportsRayTracing(ShaderPlatform) && (IsFeatureLevelSupported(ShaderPlatform, ERHIFeatureLevel::SM6) || !bRayTracingRequireSM6);
+	return IsRayTracingEnabled() && RHISupportsRayTracing(ShaderPlatform) && IsFeatureLevelSupported(ShaderPlatform, ERHIFeatureLevel::SM6);
 }
 
 ERayTracingMode GetRayTracingMode()
diff --git a/Engine/Source/Runtime/VulkanRHI/Private/VulkanRHI.cpp b/Engine/Source/Runtime/VulkanRHI/Private/VulkanRHI.cpp
--- a/Engine/Source/Runtime/VulkanRHI/Private/VulkanRHI.cpp
+++ b/Engine/Source/Runtime/VulkanRHI/Private/VulkanRHI.cpp
@@ -680,9 +680,7 @@ void FVulkanDynamicRHI::CreateInstance()
 	// Run a profile check to see if this device can support our raytacing requirements since it might change the required API version of the instance
 	if (FVulkanPlatform::SupportsProfileChecks() && GVulkanRayTracingCVar.GetValueOnAnyThread())
 	{
-		static IConsoleVariable* RequireSM6CVar = IConsoleManager::Get().FindConsoleVariable(TEXT("r.RayTracing.RequireSM6"));
-		const bool bRequireSM6 = RequireSM6CVar && RequireSM6CVar->GetBool();
-		const bool bRayTracingAllowedOnCurrentShaderPlatform = (bRequireSM6 == false) || (GMaxRHIShaderPlatform == SP_VULKAN_SM6 || IsVulkanMobileSM5Platform(GMaxRHIShaderPlatform));
+		const bool bRayTracingAllowedOnCurrentShaderPlatform = (GMaxRHIShaderPlatform == SP_VULKAN_SM6 || IsVulkanMobileSM5Platform(GMaxRHIShaderPlatform));
 
 		if (CheckVulkanProfile(GMaxRHIFeatureLevel, true) && bRayTracingAllowedOnCurrentShaderPlatform)
 		{
@@ -696,7 +694,7 @@ void FVulkanDynamicRHI::CreateInstance()
 
 			if (!bRayTracingAllowedOnCurrentShaderPlatform)
 			{
-				UE_LOG(LogVulkanRHI, Display, TEXT("Vulkan RayTracing disabled because SM6 shader platform is required (r.RayTracing.RequireSM6=1)."));
+				UE_LOG(LogVulkanRHI, Display, TEXT("Vulkan RayTracing disabled because SM6 shader platform is required."));
 			}
 			else
 			{
