[shaders] modify FShaderCode finalize to create a FSharedBuffer object, and modify all downstream uses of shader code to re-use this buffer (job cache, pushes to DDC, shader maps, and shader library). This reduces total amount of LLM tracked memory allocated at the end of a cold Lyra PS4 cook by about ~350MB; impact likely much larger for cooks of larger projects.

diff --git a/Engine/Plugins/FX/Niagara/Source/NiagaraEditor/Private/NiagaraCompiler.cpp b/Engine/Plugins/FX/Niagara/Source/NiagaraEditor/Private/NiagaraCompiler.cpp
--- a/Engine/Plugins/FX/Niagara/Source/NiagaraEditor/Private/NiagaraCompiler.cpp
+++ b/Engine/Plugins/FX/Niagara/Source/NiagaraEditor/Private/NiagaraCompiler.cpp
@@ -2225,9 +2225,9 @@ TOptional<FNiagaraCompileResults> FHlslNiagaraCompiler::GetCompileResult(int32 J
 	FVectorVMCompilationOutput CompilationOutput;
 	if (CompilationJob->ShaderCompileJob->bSucceeded)
 	{
-		const TArray<uint8>& Code = CompilationJob->ShaderCompileJob->Output.ShaderCode.GetReadAccess();
+		TConstArrayView<uint8> Code = CompilationJob->ShaderCompileJob->Output.ShaderCode.GetReadView();
 		FShaderCodeReader ShaderCode(Code);
-		FMemoryReader Ar(Code, true);
+		FMemoryReaderView Ar(Code, true);
 		Ar.SetLimitSize(ShaderCode.GetActualShaderCodeSize());
 		Ar << CompilationOutput;
 
diff --git a/Engine/Shaders/Public/ShaderVersion.ush b/Engine/Shaders/Public/ShaderVersion.ush
--- a/Engine/Shaders/Public/ShaderVersion.ush
+++ b/Engine/Shaders/Public/ShaderVersion.ush
@@ -3,4 +3,4 @@
 // Each time the console command is executed it generates a new GUID. As this file is included
 // in Platform.ush (which should be included in any shader) it allows to invalidate the shader DDC.
 // 
-#pragma message("UESHADERMETADATA_VERSION 0D63C246-532E-4D9D-83F3-906A518818CC")
\ No newline at end of file
+#pragma message("UESHADERMETADATA_VERSION 1115A78C-1E1B-4E96-A416-DA839794A1C5")
\ No newline at end of file
diff --git a/Engine/Source/Developer/Apple/MetalShaderFormat/Private/MetalShaderCompiler.cpp b/Engine/Source/Developer/Apple/MetalShaderFormat/Private/MetalShaderCompiler.cpp
--- a/Engine/Source/Developer/Apple/MetalShaderFormat/Private/MetalShaderCompiler.cpp
+++ b/Engine/Source/Developer/Apple/MetalShaderFormat/Private/MetalShaderCompiler.cpp
@@ -1184,7 +1184,7 @@ bool StripShader_Metal(TArray<uint8>& Code, class FString const& DebugPath, bool
 				// jam it into the output bytes
 				NewAr.Serialize(SourceCode.GetData(), SourceCode.Num());
 				
-				Code = NewCode.GetReadAccess();
+				Code = NewCode.GetReadView();
 			}
 		}
 		else
@@ -1278,7 +1278,7 @@ uint64 AppendShader_Metal(FString const& WorkingDir, const FSHAHash& Hash, TArra
 						NewAr << OfflineCompiledFlag;
                         Header.Serialize(NewAr, SRT);
 						
-						InShaderCode = NewCode.GetReadAccess();
+						InShaderCode = NewCode.GetReadView();
 						
 						UE_LOG(LogShaders, Verbose, TEXT("Archiving succeeded: shader %s (Len: %0.8x, CRC: %0.8x, SHA: %s)"), *ShaderName, Header.SourceLen, Header.SourceCRC, *Hash.ToString());
 					}
diff --git a/Engine/Source/Developer/ShaderCompilerCommon/Private/ShaderCompilerCommon.cpp b/Engine/Source/Developer/ShaderCompilerCommon/Private/ShaderCompilerCommon.cpp
--- a/Engine/Source/Developer/ShaderCompilerCommon/Private/ShaderCompilerCommon.cpp
+++ b/Engine/Source/Developer/ShaderCompilerCommon/Private/ShaderCompilerCommon.cpp
@@ -2223,12 +2223,12 @@ namespace UE::ShaderCompilerCommon
 				// always output decompressed code as it's slightly more useful for A/B comparisons
 				TArray<uint8> DecompressedCode;
 				DecompressedCode.SetNum(Output.ShaderCode.GetUncompressedSize());
-				bool bSucceed = FCompression::UncompressMemory(NAME_Oodle, DecompressedCode.GetData(), DecompressedCode.Num(), Output.ShaderCode.GetReadAccess().GetData(), Output.ShaderCode.GetShaderCodeSize());
+				bool bSucceed = FCompression::UncompressMemory(NAME_Oodle, DecompressedCode.GetData(), DecompressedCode.Num(), Output.ShaderCode.GetReadView().GetData(), Output.ShaderCode.GetShaderCodeSize());
 				FFileHelper::SaveArrayToFile(DecompressedCode, *ShaderCodeFileName);
 			}
 			else
 			{
-				FFileHelper::SaveArrayToFile(Output.ShaderCode.GetReadAccess(), *ShaderCodeFileName);
+				FFileHelper::SaveArrayToFile(Output.ShaderCode.GetReadView(), *ShaderCodeFileName);
 			}
 		}
 
diff --git a/Engine/Source/Developer/TargetPlatform/Public/Interfaces/IShaderFormat.h b/Engine/Source/Developer/TargetPlatform/Public/Interfaces/IShaderFormat.h
--- a/Engine/Source/Developer/TargetPlatform/Public/Interfaces/IShaderFormat.h
+++ b/Engine/Source/Developer/TargetPlatform/Public/Interfaces/IShaderFormat.h
@@ -9,6 +9,7 @@
 
 class FSerializedShaderArchive;
 class FShaderPreprocessOutput;
+class FSharedBuffer;
 struct FShaderCompilerEnvironment;
 struct FShaderCompilerInput;
 struct FShaderCompilerOutput;
@@ -60,6 +61,19 @@ public:
 	*/
 	virtual bool SupportsShaderArchives() const { return false; }
 	
+	UE_DEPRECATED(5.5, "Use version accepting a TArray of FSharedBuffer code buffers.")
+	virtual bool CreateShaderArchive(FString const& LibraryName,
+		FName ShaderFormatAndShaderPlatformName,
+		const FString& WorkingDirectory,
+		const FString& OutputDir,
+		const FString& DebugOutputDir,
+		const FSerializedShaderArchive& SerializedShaders,
+		const TArray<TArray<uint8>>& ShaderCode,
+		TArray<FString>* OutputFiles) const
+	{ 
+		return false; 
+	}
+
 	/**
      * Create a format specific archive for precompiled shader code.
      *
@@ -77,9 +91,11 @@ public:
 		const FString& OutputDir,
 		const FString& DebugOutputDir,
 		const FSerializedShaderArchive& SerializedShaders,
-		const TArray<TArray<uint8>>& ShaderCode,
+		const TArray<FSharedBuffer>& ShaderCode,
 		TArray<FString>* OutputFiles) const
-	{ return false; }
+	{ 
+		return false; 
+	}
 	
 	/**
 	 * Can the shader format compile shaders to the native binary format for the platform.
diff --git a/Engine/Source/Runtime/Core/Private/UObject/FortniteReleaseBranchCustomObjectVersion.cpp b/Engine/Source/Runtime/Core/Private/UObject/FortniteReleaseBranchCustomObjectVersion.cpp
--- a/Engine/Source/Runtime/Core/Private/UObject/FortniteReleaseBranchCustomObjectVersion.cpp
+++ b/Engine/Source/Runtime/Core/Private/UObject/FortniteReleaseBranchCustomObjectVersion.cpp
@@ -9,5 +9,7 @@ TMap<FGuid, FGuid> FFortniteReleaseBranchCustomObjectVersion::GetSystemGuids()
 	SystemGuids.Add(DevGuids.Niagara_LatestScriptCompileVersion, FGuid("D35A655CD3FDA4408724D77E1316E35B"));
 	SystemGuids.Add(DevGuids.SkeletalMeshDerivedDataVersion, FGuid("5705956EC7134274A5A1BD3FC74DB3A9"));
 	SystemGuids.Add(DevGuids.MaterialTranslationDDCVersion, FGuid("CC1B3A336E144A0B9E78014E07775988"));
+	SystemGuids.Add(DevGuids.GLOBALSHADERMAP_DERIVEDDATA_VER, FGuid("E60D877F790D43268E1596F6223E1677"));
+	SystemGuids.Add(DevGuids.MATERIALSHADERMAP_DERIVEDDATA_VER, FGuid("3D174FB4B42C433CB10198CE17CA920F"));
 	return SystemGuids;
 }
diff --git a/Engine/Source/Runtime/Engine/Private/ShaderCompiler/ShaderCompilerJobCache.cpp b/Engine/Source/Runtime/Engine/Private/ShaderCompiler/ShaderCompilerJobCache.cpp
--- a/Engine/Source/Runtime/Engine/Private/ShaderCompiler/ShaderCompilerJobCache.cpp
+++ b/Engine/Source/Runtime/Engine/Private/ShaderCompiler/ShaderCompilerJobCache.cpp
@@ -378,7 +378,7 @@ public:
 	FSharedBuffer JobOutput;
 
 	/** Separate blobs for shader code */
-	TArray<FSharedBuffer> JobCode;
+	TArray<FCompositeBuffer> JobCode;
 
 	/** Path to where the cached debug info is stored. */
 	FString CachedDebugInfoPath;
@@ -627,9 +627,12 @@ private:
 		FBlake3 Hasher;
 		check(SerializeContext.HasData());
 		Hasher.Update(SerializeContext.ShaderObjectData.GetData(), SerializeContext.ShaderObjectData.GetSize());
-		for (const FSharedBuffer& CodeBuf : SerializeContext.ShaderCode)
+		for (const FCompositeBuffer& CodeBuf : SerializeContext.ShaderCode)
 		{
-			Hasher.Update(CodeBuf.GetData(), CodeBuf.GetSize());
+			for (const FSharedBuffer& CodeBufSegment : CodeBuf.GetSegments())
+			{
+				Hasher.Update(CodeBufSegment.GetData(), CodeBuf.GetSize());
+			}
 		}
 		return Hasher.Finalize();
 	}
@@ -2232,8 +2235,10 @@ void FShaderJobCache::AddJobOutput(FShaderJobData& JobData, const FShaderCommonC
 				for (int32 CodeIndex = 0; CodeIndex < SaveContext.ShaderCode.Num(); ++CodeIndex)
 				{
 					const FString CacheCodeFilename = FString::Printf(TEXT("%s/%s_%d.bytecode"), *InputDebugInfoPath, *InputSourceFilename);
-					const FSharedBuffer& JobCode = SaveContext.ShaderCode[CodeIndex];
-					FFileHelper::SaveArrayToFile(TArrayView<const uint8>((const uint8*)JobCode.GetData(), JobCode.GetSize()), *CacheCodeFilename);
+					const FCompositeBuffer& JobCode = SaveContext.ShaderCode[CodeIndex];
+					check(JobCode.GetSegments().Num() == 2); // first segment is header, second is actual code
+					FSharedBuffer CodeBuffer = JobCode.GetSegments()[1];
+					FFileHelper::SaveArrayToFile(TArrayView<const uint8>((const uint8*)CodeBuffer.GetData(), CodeBuffer.GetSize()), *CacheCodeFilename);
 				}
 			}
 		}
diff --git a/Engine/Source/Runtime/RenderCore/Internal/ShaderSerialization.h b/Engine/Source/Runtime/RenderCore/Internal/ShaderSerialization.h
--- a/Engine/Source/Runtime/RenderCore/Internal/ShaderSerialization.h
+++ b/Engine/Source/Runtime/RenderCore/Internal/ShaderSerialization.h
@@ -2,6 +2,7 @@
 
 #pragma once
 
+#include "Memory/CompositeBuffer.h"
 #include "Memory/SharedBuffer.h"
 #include "Templates/Function.h"
 #include "Templates/UniquePtr.h"
@@ -9,7 +10,7 @@
 
 class FMemoryReaderView;
 class FMemoryWriter64;
-struct FShaderCodeResource;
+class FShaderCodeResource;
 
 #if WITH_EDITOR
 namespace UE::DerivedData
@@ -70,12 +71,12 @@ struct FShaderCacheSerializeContext : public FShaderSerializeContext
 	/* View on array of buffers which store the bytecode objects for the object, one per shader/stage.
 	 * Note that this may or may not point to the OwnedShaderCode array below, depending on usage.
 	 */
-	TArrayView<FSharedBuffer> ShaderCode;
+	TArrayView<FCompositeBuffer> ShaderCode;
 
 	/* Array of buffers actually owned by this context object; it is valid for this to be empty in cases
 	 * where the array of ShaderCode buffers is stored externally.
 	 */
-	TArray<FSharedBuffer> OwnedShaderCode;
+	TArray<FCompositeBuffer> OwnedShaderCode;
 
 	/* Get the total serialized size of data for this context; note that this will return 0 if called prior to the FSharedBuffers
 	 * being set (this is done in the derived implementations, see below).
@@ -87,7 +88,7 @@ struct FShaderCacheSerializeContext : public FShaderSerializeContext
 		{
 			Size += ShaderObjectData.GetSize();
 
-			for (FSharedBuffer& CodeBuf : ShaderCode)
+			for (FCompositeBuffer& CodeBuf : ShaderCode)
 			{
 				Size += CodeBuf.GetSize();
 			}
@@ -96,7 +97,7 @@ struct FShaderCacheSerializeContext : public FShaderSerializeContext
 	}
 
 	/* Populates the given code array (transfering ownership) and resets the internal view to point to the new owning array's data */
-	void MoveCode(TArray<FSharedBuffer>& TargetCode)
+	void MoveCode(TArray<FCompositeBuffer>& TargetCode)
 	{
 		TargetCode = MoveTemp(OwnedShaderCode);
 		ShaderCode = TargetCode;
@@ -142,7 +143,7 @@ struct FShaderCacheLoadContext : public FShaderCacheSerializeContext
 	/*	Default constructor, use when array of code buffers will be allocated via ReadFromRecord */
 	RENDERCORE_API FShaderCacheLoadContext();
 	/* Constructor which references buffers (and an array of code buffers) owned elsewhere, does not allocate the OwnedShaderCode array */
-	RENDERCORE_API FShaderCacheLoadContext(FSharedBuffer ShaderObjectData, TArrayView<FSharedBuffer> CodeBuffers);
+	RENDERCORE_API FShaderCacheLoadContext(FSharedBuffer ShaderObjectData, TArrayView<FCompositeBuffer> CodeBuffers);
 
 	/* Call to reset reader to start position so the same load context can be used to populate multiple objects. */
 	RENDERCORE_API void Reuse();
diff --git a/Engine/Source/Runtime/RenderCore/Private/Shader.cpp b/Engine/Source/Runtime/RenderCore/Private/Shader.cpp
--- a/Engine/Source/Runtime/RenderCore/Private/Shader.cpp
+++ b/Engine/Source/Runtime/RenderCore/Private/Shader.cpp
@@ -669,7 +669,7 @@ FShaderCompiledShaderInitializerType::FShaderCompiledShaderInitializerType(
 	: Type(InType)
 	, Parameters(InParameters)
 	, Target(CompilerOutput.Target)
-	, Code(CompilerOutput.ShaderCode.GetReadAccess())
+	, Code(CompilerOutput.ShaderCode.GetReadView())
 	, ParameterMap(CompilerOutput.ParameterMap)
 	, OutputHash(CompilerOutput.OutputHash)
 	, MaterialShaderMapHash(InMaterialShaderMapHash)
diff --git a/Engine/Source/Runtime/RenderCore/Private/ShaderCodeLibrary.cpp b/Engine/Source/Runtime/RenderCore/Private/ShaderCodeLibrary.cpp
--- a/Engine/Source/Runtime/RenderCore/Private/ShaderCodeLibrary.cpp
+++ b/Engine/Source/Runtime/RenderCore/Private/ShaderCodeLibrary.cpp
@@ -1427,17 +1427,18 @@ struct FEditorShaderCodeArchive
 				{
 					const FShaderCodeResource& SourceShaderResource = Code->ShaderCodeResources[i];
 					FShaderCodeEntry& SerializedShaderEntry = SerializedShaders.ShaderEntries[ShaderIndex];
-					SerializedShaderEntry.Frequency = SourceShaderResource.Frequency;
-					SerializedShaderEntry.Size = SourceShaderResource.Code.Num();
-					SerializedShaderEntry.UncompressedSize = SourceShaderResource.UncompressedSize;
-					check(!SourceShaderResource.Code.IsEmpty());
-					ShaderCode.Add(SourceShaderResource.Code);
+					SerializedShaderEntry.Frequency = SourceShaderResource.GetFrequency();
+					FSharedBuffer CodeBuffer = SourceShaderResource.GetCodeBuffer();
+					SerializedShaderEntry.Size = CodeBuffer.GetSize();
+					SerializedShaderEntry.UncompressedSize = SourceShaderResource.GetUncompressedSize();
+					check(SerializedShaderEntry.Size > 0);
+					ShaderCode.Add(CodeBuffer);
 					check(ShaderCode.Num() == SerializedShaders.ShaderEntries.Num());
 
 					CodeStats.NumUniqueShaders++;
-					CodeStats.ShadersUniqueSize += SourceShaderResource.Code.Num();
+					CodeStats.ShadersUniqueSize += CodeBuffer.GetSize();
 				}
-				CodeStats.ShadersSize += Code->ShaderCodeResources[i].Code.Num();
+				CodeStats.ShadersSize += Code->ShaderCodeResources[i].GetCodeView().NumBytes();
 				SerializedShaders.ShaderIndices[ShaderMapEntry.ShaderIndicesOffset + i] = ShaderIndex;
 			}
 
@@ -1517,7 +1518,7 @@ struct FEditorShaderCodeArchive
 		int32 NumShadersSentWithCode = 0;
 
 		const FSerializedShaderArchive& SourceArchive = this->SerializedShaders;
-		TArray<TArray<uint8>>& SourceShaderCodes = this->ShaderCode;
+		TArray<FSharedBuffer>& SourceShaderCodes = this->ShaderCode;
 		for (int32 SourceShaderMapIndex : LocalShaderMapsToCopy)
 		{
 			const FSHAHash& SourceShaderMapHash = SourceArchive.ShaderMapHashes[SourceShaderMapIndex];
@@ -1551,15 +1552,15 @@ struct FEditorShaderCodeArchive
 					check(TargetShaderIndex == TargetArchive.ShaderEntries.Num() - 1);
 
 					const FShaderCodeEntry& SourceShaderEntry = SourceArchive.ShaderEntries[SourceShaderIndex];
-					TArray<uint8>& SourceShaderCode = SourceShaderCodes[SourceShaderIndex];
+					FSharedBuffer& SourceShaderCode = SourceShaderCodes[SourceShaderIndex];
 					FShaderCodeEntry& TargetShaderEntry = TargetArchive.ShaderEntries[TargetShaderIndex];
 
 					TargetShaderEntry = SourceShaderEntry;
-					if (!SourceShaderCode.IsEmpty())
+					if (!SourceShaderCode.GetSize())
 					{
 						TargetShaderEntry.Offset = TargetFlatShaderCode.Num();
-						check(SourceShaderEntry.Size == SourceShaderCode.Num());
-						if ((MaxShaderSize > 0 && TargetFlatShaderCode.Num() + SourceShaderCode.Num() > MaxShaderSize) ||
+						check(SourceShaderEntry.Size == SourceShaderCode.GetSize());
+						if ((MaxShaderSize > 0 && TargetFlatShaderCode.Num() + (int64)SourceShaderCode.GetSize() > MaxShaderSize) ||
 							(MaxShaderCount > 0 && NumShadersSentWithCode > MaxShaderCount))
 						{
 							// We have to stop here to avoid overflowing the shader limit. Send the shaders we have accumulated
@@ -1576,11 +1577,11 @@ struct FEditorShaderCodeArchive
 						}
 						++NumShadersSentWithCode;
 
-						TargetFlatShaderCode.Append(SourceShaderCode);
+						TargetFlatShaderCode.Append(MakeArrayView(reinterpret_cast<const uint8*>(SourceShaderCode.GetData()), SourceShaderCode.GetSize()));
 
-						// Empty the ShaderCode to save memory in the local process. The consumer of the TargetArchive and
+						// Reset the ShaderCode reference to (potentially) save memory in the local process. The consumer of the TargetArchive and
 						// TargetFlatShaderCode will be the only one that needs to read it.
-						SourceShaderCode.Empty();
+						SourceShaderCode.Reset();
 					}
 					else
 					{
@@ -1601,7 +1602,7 @@ struct FEditorShaderCodeArchive
 	{
 		bool bOk = true;
 		FSerializedShaderArchive& TargetArchive = this->SerializedShaders;
-		TArray<TArray<uint8>>& TargetShaderCodes = this->ShaderCode;
+		TArray<FSharedBuffer>& TargetShaderCodes = this->ShaderCode;
 
 		// Add all the shaders; we can sometimes get messages that send the shaders in advance without sending the shadermaps that use them
 		for (int32 SourceShaderIndex = 0; SourceShaderIndex < SourceArchive.ShaderHashes.Num(); ++SourceShaderIndex)
@@ -1615,7 +1616,6 @@ struct FEditorShaderCodeArchive
 			}
 			check(TargetShaderIndex == TargetArchive.ShaderEntries.Num() - 1 &&
 				TargetShaderCodes.Num() == TargetArchive.ShaderEntries.Num() - 1);
-			TArray<uint8>& TargetShaderCode = TargetShaderCodes.Emplace_GetRef();
 
 			const FShaderCodeEntry& SourceShaderEntry = SourceArchive.ShaderEntries[SourceShaderIndex];
 			FShaderCodeEntry& TargetShaderEntry = TargetArchive.ShaderEntries[TargetShaderIndex];
@@ -1645,8 +1645,8 @@ struct FEditorShaderCodeArchive
 			}
 			else
 			{
-				const TConstArrayView<uint8> SourceShaderCode(SourceFlatShaderCode.GetData() + SourceShaderEntry.Offset, SourceShaderEntry.Size);
-				TargetShaderCode = SourceShaderCode; // Copy from source's flat list to the target's separate TArray<uint8> for each shader
+				// Copy from source's flat list to the target's separate FSharedBuffer for each shader
+				TargetShaderCodes.Add(FSharedBuffer::Clone(SourceFlatShaderCode.GetData() + SourceShaderEntry.Offset, SourceShaderEntry.Size));
 			}
 		}
 
@@ -1708,8 +1708,8 @@ struct FEditorShaderCodeArchive
 		NewChunk->ShaderCode.Empty();
 		for (int32 NewArchiveIdx = 0, NumIndices = ShaderCodeEntriesNeeded.Num(); NewArchiveIdx < NumIndices; ++NewArchiveIdx)
 		{
-			TArray<uint8>& SourceShaderCodeEntry = ShaderCode[ShaderCodeEntriesNeeded[NewArchiveIdx]];
-			check(!SourceShaderCodeEntry.IsEmpty());
+			FSharedBuffer& SourceShaderCodeEntry = ShaderCode[ShaderCodeEntriesNeeded[NewArchiveIdx]];
+			check(SourceShaderCodeEntry.GetSize() > 0);
 			NewChunk->ShaderCode.Add(SourceShaderCodeEntry);
 		}
 
@@ -1739,8 +1739,8 @@ struct FEditorShaderCodeArchive
 					const FShaderCodeEntry& OtherShaderEntry = OtherArchive.SerializedShaders.ShaderEntries[OtherShaderIndex];
 					SerializedShaders.ShaderEntries[ShaderIndex] = OtherShaderEntry;
 
-					const TArray<uint8>& OtherShaderCodeEntry = OtherArchive.ShaderCode[OtherShaderIndex];
-					check(!OtherShaderCodeEntry.IsEmpty());
+					const FSharedBuffer& OtherShaderCodeEntry = OtherArchive.ShaderCode[OtherShaderIndex];
+					check(OtherShaderCodeEntry.GetSize() > 0);
 					ShaderCode.Add(OtherShaderCodeEntry);
 					check(ShaderCode.Num() == SerializedShaders.ShaderEntries.Num());
 				}
@@ -1777,16 +1777,16 @@ struct FEditorShaderCodeArchive
 					const FShaderCodeEntry& OtherShaderEntry = OtherShaders.ShaderEntries[OtherShaderIndex];
 					SerializedShaders.ShaderEntries[ShaderIndex] = OtherShaderEntry;
 
-					TArray<uint8>& Code = ShaderCode.AddDefaulted_GetRef();
-					check(ShaderCode.Num() == SerializedShaders.GetNumShaders());
+					check(ShaderCode.Num() == SerializedShaders.GetNumShaders() - 1);
 
 					// Read shader code from archive and add shader to set
 					const int64 ReadSize = OtherShaderEntry.Size;
 					check(ReadSize > 0);
 					const int64 ReadOffset = OtherShaderCodeOffset + OtherShaderEntry.Offset;
-					Code.SetNumUninitialized(ReadSize);
+					FUniqueBuffer Code = FUniqueBuffer::Alloc(ReadSize);
 					Ar.Seek(ReadOffset);
 					Ar.Serialize(Code.GetData(), ReadSize);
+					ShaderCode.Add(Code.MoveToShared());
 				}
 				SerializedShaders.ShaderIndices[ShaderMapEntry.ShaderIndicesOffset + i] = ShaderIndex;
 			}
@@ -1815,16 +1815,18 @@ struct FEditorShaderCodeArchive
 				for (int32 Index = 0; Index < ShaderCode.Num(); ++Index)
 				{
 					const FShaderCodeEntry& Entry = SerializedShaders.ShaderEntries[Index];
-					TArray<uint8>& Code = ShaderCode[Index];
 					check(Entry.Size > 0);
-					Code.SetNumUninitialized(Entry.Size);
+					FUniqueBuffer Code = FUniqueBuffer::Alloc(Entry.Size);
 					PrevCookedAr->Serialize(Code.GetData(), Entry.Size);
+
 					bOK = !PrevCookedAr->GetError();
 					if (!bOK)
 					{
 						UE_LOG(LogShaderLibrary, Error, TEXT("Failed to deserialize shader code for %s from %s"), *SerializedShaders.ShaderHashes[Index].ToString(), *IntermediateFormatPath);
 						break;
 					}
+
+					ShaderCode[Index] = Code.MoveToShared();
 				}
 			}
 			else
@@ -1972,8 +1974,8 @@ struct FEditorShaderCodeArchive
 					*FileWriter << SerializedShaders;
 					for (auto& Code : ShaderCode)
 					{
-						check(!Code.IsEmpty());
-						FileWriter->Serialize(Code.GetData(), Code.Num());
+						check(Code.GetSize() > 0);
+						FileWriter->Serialize(const_cast<void*>(Code.GetData()), Code.GetSize());
 					}
 
 					FileWriter->Close();
@@ -2250,7 +2252,7 @@ private:
 	 * The element at index N holds the ShaderCode for the element of SerializedShaders.ShaderEntries at index N.
 	 * In MultiprocessCooking elements can be empty if they have been transferred to the Director (bHasCopiedAndCleared will be true in this case).
 	 */
-	TArray<TArray<uint8>> ShaderCode;
+	TArray<FSharedBuffer> ShaderCode;
 	/** A list of ShaderMaps that have not yet been copied to the CookDirector. Used only by MultiprocessCookWorkers. */
 	TSet<int32> ShaderMapsToCopy;
 	/** True if CopyToArchiveAndClear has been called, otherwise false. If false we avoid doing some tracking since it might never be used. */
diff --git a/Engine/Source/Runtime/RenderCore/Private/ShaderCore.cpp b/Engine/Source/Runtime/RenderCore/Private/ShaderCore.cpp
--- a/Engine/Source/Runtime/RenderCore/Private/ShaderCore.cpp
+++ b/Engine/Source/Runtime/RenderCore/Private/ShaderCore.cpp
@@ -1292,7 +1292,7 @@ void FShaderCompilerOutput::GenerateOutputHash()
 {
 	FSHA1 HashState;
 	
-	const TArray<uint8>& Code = ShaderCode.GetReadAccess();
+	TArrayView<const uint8> Code = ShaderCode.GetReadView();
 
 	// we don't hash the optional attachments as they would prevent sharing (e.g. many materials share the same VS)
 	uint32 ShaderCodeSize = ShaderCode.GetShaderCodeSize();
@@ -1609,19 +1609,21 @@ public:
 		int32 PackedShaderKey = Compiler->GetPackedShaderKey();
 		CombinedSource.Append(reinterpret_cast<const uint8*>(&PackedShaderKey), sizeof(PackedShaderKey));
 
-		const uint32 PrimaryLength = Job.Output.ShaderCode.GetReadAccess().Num();
-		CombinedSource.Append(reinterpret_cast<const uint8*>(&PrimaryLength), sizeof(PrimaryLength));
+		TArrayView<const uint8> PrimaryCode = Job.Output.ShaderCode.GetReadView();
+		const uint32 PrimaryLength = PrimaryCode.NumBytes();
+		TArrayView<const uint8> SecondaryCode = Job.SecondaryOutput->ShaderCode.GetReadView();
+		const uint32 SecondaryLength = SecondaryCode.NumBytes();
 
-		const uint32 SecondaryLength = Job.SecondaryOutput->ShaderCode.GetReadAccess().Num();
+		CombinedSource.Reserve(PrimaryLength + sizeof(PrimaryLength) + SecondaryLength + sizeof(SecondaryLength));
+		CombinedSource.Append(reinterpret_cast<const uint8*>(&PrimaryLength), sizeof(PrimaryLength));
 		CombinedSource.Append(reinterpret_cast<const uint8*>(&SecondaryLength), sizeof(SecondaryLength));
-
-		CombinedSource.Append(Job.Output.ShaderCode.GetReadAccess());
-		CombinedSource.Append(Job.SecondaryOutput->ShaderCode.GetReadAccess());
+		CombinedSource.Append(PrimaryCode);
+		CombinedSource.Append(SecondaryCode);
 
 		// Replace Output shader code with the combined result
 		Job.Output.ShaderCode = {};
 		TArray<uint8>& FinalShaderCode = Job.Output.ShaderCode.GetWriteAccess();
-		FinalShaderCode.Append(CombinedSource);
+		FinalShaderCode = MoveTemp(CombinedSource);
 		Job.Output.ShaderCode.FinalizeShaderCode();
 	}
 
@@ -3778,27 +3780,28 @@ void FShaderCode::Compress(FName ShaderCompressionFormat, FOodleDataCompression:
 
 	TArray<uint8> Compressed;
 	// conventional formats will fail if the compressed size isn't enough, Oodle needs a more precise estimate
-	int32 CompressedSize = (ShaderCompressionFormat != NAME_Oodle) ? ShaderCodeWithOptionalData.Num() : FOodleDataCompression::CompressedBufferSizeNeeded(ShaderCodeWithOptionalData.Num());
+	TConstArrayView<uint8> Code = ShaderCodeResource.GetCodeView();
+	int32 CompressedSize = (ShaderCompressionFormat != NAME_Oodle) ? Code.NumBytes() : FOodleDataCompression::CompressedBufferSizeNeeded(Code.NumBytes());
 	Compressed.AddUninitialized(CompressedSize);
 
 	// non-Oodle format names use the old API, for NAME_Oodle we replace the call with the custom invocation
 	bool bCompressed = false;
 	if (ShaderCompressionFormat != NAME_Oodle)
 	{
-		bCompressed = FCompression::CompressMemory(ShaderCompressionFormat, Compressed.GetData(), CompressedSize, ShaderCodeWithOptionalData.GetData(), ShaderCodeWithOptionalData.Num(), COMPRESS_BiasSize);
+		bCompressed = FCompression::CompressMemory(ShaderCompressionFormat, Compressed.GetData(), CompressedSize, Code.GetData(), Code.Num(), COMPRESS_BiasSize);
 	}
 	else
 	{
-		CompressedSize = FOodleDataCompression::Compress(Compressed.GetData(), CompressedSize, ShaderCodeWithOptionalData.GetData(), ShaderCodeWithOptionalData.Num(),
+		CompressedSize = FOodleDataCompression::Compress(Compressed.GetData(), CompressedSize, Code.GetData(), Code.Num(),
 			InOodleCompressor, InOodleLevel);
 		bCompressed = CompressedSize != 0;
 	}
 
 	// there is code that assumes that if CompressedSize == CodeSize, the shader isn't compressed. Because of that, do not accept equal compressed size (very unlikely anyway)
-	if (bCompressed && CompressedSize < ShaderCodeWithOptionalData.Num())
+	if (bCompressed && CompressedSize < Code.Num())
 	{
 		// cache the ShaderCodeSize since it will no longer possible to get it as the reader will fail to parse the compressed data
-		FShaderCodeReader Wrapper(ShaderCodeWithOptionalData);
+		FShaderCodeReader Wrapper(Code);
 		ShaderCodeSize = Wrapper.GetShaderCodeSize();
 		checkf(ShaderCodeSize >= 0, TEXT("Unable to determine ShaderCodeSize from uncompressed code"), ShaderCodeSize);
 
@@ -3806,13 +3809,44 @@ void FShaderCode::Compress(FName ShaderCompressionFormat, FOodleDataCompression:
 		CompressionFormat = ShaderCompressionFormat;
 		OodleCompressor = InOodleCompressor;
 		OodleLevel = InOodleLevel;
-		UncompressedSize = ShaderCodeWithOptionalData.Num();
+		UncompressedSize = Code.Num();
 
 		Compressed.SetNum(CompressedSize);
-		ShaderCodeWithOptionalData = Compressed;
+		ShaderCodeResource.Code = MakeSharedBufferFromArray(MoveTemp(Compressed));
 	}
 }
 
+FArchive& operator<<(FArchive& Ar, FSharedBuffer& Buffer)
+{
+	uint64 Len = Buffer.GetSize();
+	Ar << Len;
+
+	if (Ar.IsLoading())
+	{
+		Buffer.Reset();
+
+		if (Len >= 0)
+		{
+			FUniqueBuffer BufTmp = FUniqueBuffer::Alloc(Len);
+			Ar.Serialize(BufTmp.GetData(), Len);
+			Buffer = BufTmp.MoveToShared();
+		}
+	}
+	else if (Ar.IsSaving())
+	{
+		Ar.Serialize(const_cast<void*>(Buffer.GetData()), Len);
+	}
+
+	return Ar;
+}
+
+FArchive& operator<<(FArchive& Ar, FShaderCodeResource& Resource)
+{
+	Ar << Resource.Header;
+	Ar << Resource.Code;
+	return Ar;
+}
+
 FArchive& operator<<(FArchive& Ar, FShaderCode& Output)
 {
 	if (Ar.IsLoading())
@@ -3825,7 +3859,8 @@ FArchive& operator<<(FArchive& Ar, FShaderCode& Output)
 	}
 
 	// Note: this serialize is used to pass between UE and the shader compile worker, recompile both when modifying
-	Ar << Output.ShaderCodeWithOptionalData;
+	FSharedBuffer& CodeBuffer = Output.ShaderCodeResource.Code;
+	Ar << CodeBuffer;
 	Ar << Output.UncompressedSize;
 	{
 		FString CompressionFormatString(Output.CompressionFormat.ToString());
@@ -4011,7 +4046,7 @@ void FShaderCompileJob::SerializeOutput(FShaderCacheSerializeContext& Ctx, int32
 	FShaderCodeResource CodeResource;
 	if (bIsSaving)
 	{
-		CodeResource = Output.ConvertCodeToResource();
+		CodeResource = Output.GetFinalizedCodeResource();
 	}
 
 	Ctx.SerializeCodeFunc(CodeResource, CodeIndex);
diff --git a/Engine/Source/Runtime/RenderCore/Private/ShaderResource.cpp b/Engine/Source/Runtime/RenderCore/Private/ShaderResource.cpp
--- a/Engine/Source/Runtime/RenderCore/Private/ShaderResource.cpp
+++ b/Engine/Source/Runtime/RenderCore/Private/ShaderResource.cpp
@@ -179,7 +179,7 @@ static void ApplyResourceStats(FShaderMapResourceCode& Resource)
 	INC_DWORD_STAT_BY(STAT_Shaders_ShaderResourceMemory, Resource.GetSizeBytes());
 	for (const FShaderCodeResource& Shader : Resource.ShaderCodeResources)
 	{
-		INC_DWORD_STAT_BY_FName(GetMemoryStatType(Shader.Frequency).GetName(), Shader.Code.Num());
+		INC_DWORD_STAT_BY_FName(GetMemoryStatType(Shader.GetFrequency()).GetName(), Shader.GetCodeBuffer().GetSize());
 	}
 #endif // STATS
 }
@@ -190,7 +190,7 @@ static void RemoveResourceStats(FShaderMapResourceCode& Resource)
 	DEC_DWORD_STAT_BY(STAT_Shaders_ShaderResourceMemory, Resource.GetSizeBytes());
 	for (const FShaderCodeResource& Shader : Resource.ShaderCodeResources)
 	{
-		DEC_DWORD_STAT_BY_FName(GetMemoryStatType(Shader.Frequency).GetName(), Shader.Code.Num());
+		DEC_DWORD_STAT_BY_FName(GetMemoryStatType(Shader.GetFrequency()).GetName(), Shader.GetCodeBuffer().GetSize());
 	}
 #endif // STATS
 }
@@ -231,7 +231,7 @@ uint32 FShaderMapResourceCode::GetSizeBytes() const
 	uint64 Size = sizeof(*this) + ShaderHashes.GetAllocatedSize() + ShaderCodeResources.GetAllocatedSize();
 	for (const FShaderCodeResource& Entry : ShaderCodeResources)
 	{
-		Size += Entry.Code.GetAllocatedSize();
+		Size += Entry.GetCacheBuffer().GetSize();
 	}
 	check(Size <= TNumericLimits<uint32>::Max());
 	return static_cast<uint32>(Size);
@@ -257,9 +257,7 @@ void FShaderMapResourceCode::AddShaderCompilerOutput(const FShaderCompilerOutput
 		AddEditorOnlyData(Index, DebugName, Output.PlatformDebugData, Output.Errors, Output.ShaderStatistics, DebugInfo);
 #endif
 
-		// This moves the internal array in the Output's ShaderCode object into the resource format; assuming that after
-		// the code is moved to the shadermap that it's no longer required in the job output struct.
-		ShaderCodeResources.Insert(Output.ConvertCodeToResource(), Index);
+		ShaderCodeResources.Insert(Output.GetFinalizedCodeResource(), Index);
 	}
 #if WITH_EDITORONLY_DATA
 	else
@@ -355,8 +353,8 @@ void FShaderMapResourceCode::ToString(FStringBuilderBase& OutString) const
 	for (int32 i = 0; i < ShaderHashes.Num(); ++i)
 	{
 		const FShaderCodeResource& Res = ShaderCodeResources[i];
-		OutString.Appendf(TEXT("    [%d]: { Hash: %s, Freq: %s, Size: %d, UncompressedSize: %d }\n"),
-			i, *ShaderHashes[i].ToString(), GetShaderFrequencyString(Res.Frequency), Res.Code.Num(), Res.UncompressedSize);
+		OutString.Appendf(TEXT("    [%d]: { Hash: %s, Freq: %s, Size: %llu, UncompressedSize: %d }\n"),
+			i, *ShaderHashes[i].ToString(), GetShaderFrequencyString(Res.GetFrequency()), Res.GetCodeBuffer().GetSize(), Res.GetUncompressedSize());
 	}
 }
 
@@ -655,20 +653,21 @@ FRHIShader* FShaderMapResource_InlineCode::CreateRHIShaderOrCrash(int32 ShaderIn
 
 	FMemStackBase& MemStack = FMemStack::Get();
 	const FShaderCodeResource& ShaderCodeResource = Code->ShaderCodeResources[ShaderIndex];
-	const uint8* ShaderCode = ShaderCodeResource.Code.GetData();
+	FSharedBuffer ShaderCode = ShaderCodeResource.GetCodeBuffer();
+	TConstArrayView<uint8> ShaderCodeView = ShaderCodeResource.GetCodeView();
 
 	FMemMark Mark(MemStack);
-	if (ShaderCodeResource.Code.Num() != ShaderCodeResource.UncompressedSize)
+	int32 UncompressedSize = ShaderCodeResource.GetUncompressedSize();
+	if (ShaderCode.GetSize() != UncompressedSize)
 	{
-		void* UncompressedCode = MemStack.Alloc(ShaderCodeResource.UncompressedSize, 16);
-		bool bSucceed = FCompression::UncompressMemory(GetShaderCompressionFormat(), UncompressedCode, ShaderCodeResource.UncompressedSize, ShaderCode, ShaderCodeResource.Code.Num());
+		void* UncompressedCode = MemStack.Alloc(UncompressedSize, 16);
+		bool bSucceed = FCompression::UncompressMemory(GetShaderCompressionFormat(), UncompressedCode, UncompressedSize, ShaderCode.GetData(), ShaderCode.GetSize());
 		check(bSucceed);
-		ShaderCode = (uint8*)UncompressedCode;
+		ShaderCodeView = MakeArrayView(reinterpret_cast<const uint8*>(UncompressedCode), UncompressedSize);
 	}
 
-	const auto ShaderCodeView = MakeArrayView(ShaderCode, ShaderCodeResource.UncompressedSize);
 	const FSHAHash& ShaderHash = Code->ShaderHashes[ShaderIndex];
-	const EShaderFrequency Frequency = ShaderCodeResource.Frequency;
+	const EShaderFrequency Frequency = ShaderCodeResource.GetFrequency();
 
 	TRefCountPtr<FRHIShader> RHIShader;
 	switch (Frequency)
diff --git a/Engine/Source/Runtime/RenderCore/Private/ShaderSerialization.cpp b/Engine/Source/Runtime/RenderCore/Private/ShaderSerialization.cpp
--- a/Engine/Source/Runtime/RenderCore/Private/ShaderSerialization.cpp
+++ b/Engine/Source/Runtime/RenderCore/Private/ShaderSerialization.cpp
@@ -14,10 +14,7 @@ FShaderCacheSaveContext::FShaderCacheSaveContext()
 {
 	SerializeCodeFunc = [this](FShaderCodeResource& Res, int32 Index)
 		{
-			TArray64<uint8> CodeData;
-			FMemoryWriter64 Ar(CodeData, true);
-			Ar << Res;
-			OwnedShaderCode.Add(MakeSharedBufferFromArray(MoveTemp(CodeData)));
+			OwnedShaderCode.Add(Res.GetCacheBuffer());
 			// reset the array view any time an entry is added; we do this instead of calling Resize in the reserve delegate
 			// and setting it there since not all code paths (i.e. single job cache records) call reserve
 			ShaderCode = OwnedShaderCode;
@@ -58,11 +55,14 @@ UE::DerivedData::FCacheRecord FShaderCacheSaveContext::BuildCacheRecord(const UE
 	Finalize();
 
 	UE::DerivedData::FCacheRecordBuilder RecordBuilder(Key);
-	RecordBuilder.AddValue(ShaderObjectDataValue, UE::DerivedData::FValue::Compress(ShaderObjectData));
+	RecordBuilder.AddValue(ShaderObjectDataValue, ShaderObjectData);
 	int32 CodeIndex = 0;
-	for (FSharedBuffer& CodeBuf : ShaderCode)
+	// Code buffers are already compressed, don't waste cycles attempting (and failing) to recompress them
+	const ECompressedBufferCompressor CodeComp = ECompressedBufferCompressor::NotSet;
+	const ECompressedBufferCompressionLevel CodeCompLevel = ECompressedBufferCompressionLevel::None;
+	for (FCompositeBuffer& CodeBuf : ShaderCode)
 	{
-		RecordBuilder.AddValue(ShaderCodeDataValue.MakeIndexed(CodeIndex++), UE::DerivedData::FValue::Compress(CodeBuf));
+		RecordBuilder.AddValue(ShaderCodeDataValue.MakeIndexed(CodeIndex++), UE::DerivedData::FValue(FCompressedBuffer::Compress(CodeBuf, CodeComp, CodeCompLevel)));
 	}
 
 	TCbWriter<16> MetaWriter;
@@ -80,12 +80,11 @@ FShaderCacheLoadContext::FShaderCacheLoadContext()
 {
 	SerializeCodeFunc = [this](FShaderCodeResource& Res, int32 Index)
 		{
-			FMemoryReaderView Ar(ShaderCode[Index]);
-			Ar << Res;
+			Res.PopulateFromComposite(ShaderCode[Index]);
 		};
 }
 
-FShaderCacheLoadContext::FShaderCacheLoadContext(FSharedBuffer InShaderObjectData, TArrayView<FSharedBuffer> InCodeBuffers) : FShaderCacheLoadContext()
+FShaderCacheLoadContext::FShaderCacheLoadContext(FSharedBuffer InShaderObjectData, TArrayView<FCompositeBuffer> InCodeBuffers) : FShaderCacheLoadContext()
 {
 	ShaderObjectData = InShaderObjectData;
 	ShaderCode = InCodeBuffers;
@@ -113,7 +112,8 @@ void FShaderCacheLoadContext::ReadFromRecord(const UE::DerivedData::FCacheRecord
 	OwnedShaderCode.Reserve(CodeCount);
 	for (int32 CodeIndex = 0; CodeIndex < CodeCount; ++CodeIndex)
 	{
-		OwnedShaderCode.Add(Record.GetValue(ShaderCodeDataValue.MakeIndexed(CodeIndex)).GetData().Decompress());
+		FSharedBuffer CombinedBuffer = Record.GetValue(ShaderCodeDataValue.MakeIndexed(CodeIndex)).GetData().Decompress();
+		OwnedShaderCode.Add(FShaderCodeResource::Unpack(CombinedBuffer));
 	}
 	ShaderCode = OwnedShaderCode;
 }
diff --git a/Engine/Source/Runtime/RenderCore/Public/Shader.h b/Engine/Source/Runtime/RenderCore/Public/Shader.h
--- a/Engine/Source/Runtime/RenderCore/Public/Shader.h
+++ b/Engine/Source/Runtime/RenderCore/Public/Shader.h
@@ -1556,7 +1556,7 @@ struct FShaderCompiledShaderInitializerType
 	const FShaderType* Type;
 	const FShaderType::FParameters* Parameters;
 	FShaderTarget Target;
-	TArrayView<const uint8> Code;
+	TConstArrayView<uint8> Code;
 	const FShaderParameterMap& ParameterMap;
 	const FSHAHash& OutputHash;
 	FSHAHash MaterialShaderMapHash;
diff --git a/Engine/Source/Runtime/RenderCore/Public/ShaderCompilerCore.h b/Engine/Source/Runtime/RenderCore/Public/ShaderCompilerCore.h
--- a/Engine/Source/Runtime/RenderCore/Public/ShaderCompilerCore.h
+++ b/Engine/Source/Runtime/RenderCore/Public/ShaderCompilerCore.h
@@ -28,9 +28,9 @@ class FShaderPipelineCompileJob;
 typedef TSharedPtr<TArray<ANSICHAR>, ESPMode::ThreadSafe> FShaderSharedAnsiStringPtr;
 
 // this is for the protocol, not the data, bump if FShaderCompilerInput/FShaderPreprocessOutput serialization, SerializeWorkerInput or ProcessInputFromArchive changes.
-inline const int32 ShaderCompileWorkerInputVersion = 28;
+inline const int32 ShaderCompileWorkerInputVersion = 29;
 // this is for the protocol, not the data, bump if FShaderCompilerOutput or WriteToOutputArchive changes.
-inline const int32 ShaderCompileWorkerOutputVersion = 24;
+inline const int32 ShaderCompileWorkerOutputVersion = 25;
 // this is for the protocol, not the data.
 inline const int32 ShaderCompileWorkerSingleJobHeader = 'S';
 // this is for the protocol, not the data.
@@ -38,7 +38,7 @@ inline const int32 ShaderCompileWorkerPipelineJobHeader = 'P';
 
 // Modify this to invalidate _just_ the cache/DDC entries for individual shaders (will not cause shadermaps to rebuild if they are not otherwise out-of-date).
 // This should be bumped for changes to the FShaderCompilerOutput data structure (in addition to ShaderCompileWorkerOutputVersion)
-inline static const FGuid UE_SHADER_CACHE_VERSION = FGuid("55EB8969-BA75-4CBD-B7A1-84194D140CBD");
+inline static const FGuid UE_SHADER_CACHE_VERSION = FGuid("31D16D5D-E4DC-43FF-BB50-6648399DC4B0");
 
 namespace UE::ShaderCompiler
 {
@@ -629,9 +629,9 @@ struct FShaderCompilerOutput
 		Stat.Flags = Flags;
 	}
 
-	FShaderCodeResource ConvertCodeToResource() const
+	const FShaderCodeResource& GetFinalizedCodeResource() const
 	{
-		return ShaderCode.ConvertToResource(Target.GetFrequency(), OutputHash);
+		return ShaderCode.GetFinalizedResource(Target.GetFrequency(), OutputHash);
 	}
 
 	void SetCodeFromResource(FShaderCodeResource&& Resource)
diff --git a/Engine/Source/Runtime/RenderCore/Public/ShaderCore.h b/Engine/Source/Runtime/RenderCore/Public/ShaderCore.h
--- a/Engine/Source/Runtime/RenderCore/Public/ShaderCore.h
+++ b/Engine/Source/Runtime/RenderCore/Public/ShaderCore.h
@@ -17,6 +17,8 @@
 #include "HAL/Platform.h"
 #include "HAL/UnrealMemory.h"
 #include "Logging/LogMacros.h"
+#include "Memory/CompositeBuffer.h"
+#include "Memory/SharedBuffer.h"
 #include "Misc/AssertionMacros.h"
 #include "Misc/CString.h"
 #include "Misc/CoreStats.h"
@@ -937,10 +939,10 @@ typedef uint32 RENDERCORE_ATTRIBUTE_UNALIGNED unaligned_uint32;
 // later we can transform that to the actual class passed around at the RHI level
 class FShaderCodeReader
 {
-	TArrayView<const uint8> ShaderCode;
+	TConstArrayView<uint8> ShaderCode;
 
 public:
-	FShaderCodeReader(TArrayView<const uint8> InShaderCode)
+	FShaderCodeReader(TConstArrayView<uint8> InShaderCode)
 		: ShaderCode(InShaderCode)
 	{
 		check(ShaderCode.Num());
@@ -951,9 +953,9 @@ public:
 		return ShaderCode.Num() - GetOptionalDataSize();
 	}
 
-	TArrayView<const uint8> GetOffsetShaderCode(int32 Offset)
+	TConstArrayView<uint8> GetOffsetShaderCode(int32 Offset)
 	{
-		return MakeArrayView(ShaderCode.GetData() + Offset, GetActualShaderCodeSize() - Offset);
+		return MakeConstArrayView(ShaderCode.GetData() + Offset, GetActualShaderCodeSize() - Offset);
 	}
 
 	// for convenience
@@ -1082,20 +1084,101 @@ public:
 	}
 };
 
-struct FShaderCodeResource
+class FShaderCode;
+
+class FShaderCodeResource
 {
-	TArray<uint8> Code;
-	int32 UncompressedSize; // full size of code array before compression
-	int32 ShaderCodeSize;	// uncompressed size excluding optional data
-	EShaderFrequency Frequency;
+	struct FHeader
+	{
+		int32 UncompressedSize;		// full size of code array before compression
+		int32 ShaderCodeSize;		// uncompressed size excluding optional data
+		EShaderFrequency Frequency;
+	};
+
+	FSharedBuffer Header;		// The above FHeader struct persisted in a shared buffer
+	FSharedBuffer Code;			// The bytecode buffer as constructed by FShaderCode::FinalizeShaderCode
+
+	friend class FShaderCode;
+	friend RENDERCORE_API FArchive& operator<<(FArchive& Ar, FShaderCode& Code);
+	friend FArchive& operator<<(FArchive& Ar, FShaderCodeResource& Resource);
+
+public:
+
+	/* Returns a uint8 array view representation of the Code FSharedBuffer, for compatibility's sake (much downstream
+	 * usage of shader code expects an array of uint8)
+	 */
+	TConstArrayView<uint8> GetCodeView() const
+	{
+		return MakeConstArrayView(reinterpret_cast<const uint8*>(Code.GetData()), static_cast<int32>(Code.GetSize()));
+	}
 
-	friend FArchive& operator<<(FArchive& Ar, FShaderCodeResource& Resource)
+	/* Return the buffer storing just the shader code for this resource */
+	FSharedBuffer GetCodeBuffer() const
 	{
-		uint8 Freq = Resource.Frequency;
-		Ar << Resource.Code << Resource.UncompressedSize << Resource.ShaderCodeSize << Freq;
-		Resource.Frequency = (EShaderFrequency)Freq;
-		return Ar;
+		return Code;
+	}
+
+	/* Returns a single composite buffer referencing both the header and code data to be cached. */
+	FCompositeBuffer GetCacheBuffer() const
+	{
+		return FCompositeBuffer(Header, Code);
+	}
+
+	/* Unpacks the given FSharedBuffer into separate header/code buffer views and returns them as a 2-segment composite buffer. 
+	 * Note that this is required since when pushing a composite buffer to DDC it does not maintain the segment structure.
+	 */
+	static FCompositeBuffer Unpack(FSharedBuffer MonolithicBuffer)
+	{
+		FMemoryView FullBufferView = MonolithicBuffer.GetView();
+
+		return FCompositeBuffer(
+			MonolithicBuffer.MakeView(FullBufferView.Left(sizeof(FHeader)), MonolithicBuffer),
+			MonolithicBuffer.MakeView(FullBufferView.RightChop(sizeof(FHeader)), MonolithicBuffer));
+	}
+	
+	/* Sets the Header and Code shared buffer references in this resource to the segments referenced
+	 * by the given composite buffer.
+	 */
+	void PopulateFromComposite(FCompositeBuffer CacheBuffer)
+	{
+		check(CacheBuffer.GetSegments().Num() == 2);
+		Header = CacheBuffer.GetSegments()[0];
+		check(Header.GetSize() == sizeof(FHeader));
+		Code = CacheBuffer.GetSegments()[1];
+	}
+
+	/* Populates the header for this code resource with the given sizes and frequency. 
+	 * Note that this is done as a separate process from the construction of the Code buffer
+	 * as the shader frequency is only known by the owning job, and not stored in the FShaderCode.
+	 */
+	void PopulateHeader(int32 UncompressedSize, int32 ShaderCodeSize, EShaderFrequency Frequency)
+	{
+		check(Code);
+		FHeader HeaderData{ UncompressedSize, ShaderCodeSize, Frequency };
+		Header = FSharedBuffer::Clone(&HeaderData, sizeof(HeaderData));
+	}
+
+	/* Retrieves the uncompressed size of the shader code as stored in the FHeader buffer. */
+	int32 GetUncompressedSize() const
+	{
+		check(Header);
+		return reinterpret_cast<const FHeader*>(Header.GetData())->UncompressedSize;
+	}
+
+	/* Retrieves the actual shader code size (excluding optional data) as stored in the FHeader buffer. */
+	int32 GetShaderCodeSize() const
+	{
+		check(Header);
+		return reinterpret_cast<const FHeader*>(Header.GetData())->ShaderCodeSize;
+	}
+
+	/* Retrieves the shader frequency as stored in the FHeader buffer. */
+	EShaderFrequency GetFrequency() const
+	{
+		check(Header);
+		return reinterpret_cast<const FHeader*>(Header.GetData())->Frequency;
 	}
+
 };
 
 class FShaderCode
@@ -1105,6 +1188,8 @@ class FShaderCode
 	// access through class methods
 	mutable TArray<uint8> ShaderCodeWithOptionalData;
 
+	mutable FShaderCodeResource ShaderCodeResource;
+
 	/** ShaderCode may be compressed in SCWs on demand. If this value isn't null, the shader code is compressed. */
 	mutable int32 UncompressedSize;
 
@@ -1135,12 +1220,14 @@ public:
 	// adds CustomData or does nothing if that was already done before
 	void FinalizeShaderCode() const
 	{
-		if(OptionalDataSize != -1)
+		if (OptionalDataSize != -1)
 		{
 			checkf(UncompressedSize == 0, TEXT("FShaderCode::FinalizeShaderCode() was called after compressing the code"));
 			OptionalDataSize += sizeof(OptionalDataSize);
 			ShaderCodeWithOptionalData.Append((const uint8*)&OptionalDataSize, sizeof(OptionalDataSize));
 			OptionalDataSize = -1;
+
+			ShaderCodeResource.Code = MakeSharedBufferFromArray(MoveTemp(ShaderCodeWithOptionalData));
 		}
 	}
 
@@ -1165,17 +1252,32 @@ public:
 		{
 			FinalizeShaderCode();
 
-			FShaderCodeReader Wrapper(ShaderCodeWithOptionalData);
-			return Wrapper.GetShaderCodeSize();
+			if (UncompressedSize != 0) // already compressed, get code size from resource
+			{
+				return ShaderCodeResource.GetShaderCodeSize();
+			}
+			else
+			{
+				// code buffer has been populated but not compressed, can still read additional fields from code buffer
+				FShaderCodeReader Wrapper(ShaderCodeResource.GetCodeView());
+				return Wrapper.GetShaderCodeSize();
+			}
 		}
 	}
 
-	// for read access, can have additional data attached to the end. Can also be compressed
+	UE_DEPRECATED(5.5, "Use GetReadView")
 	const TArray<uint8>& GetReadAccess() const
+	{
+		static TArray<uint8> Dummy;
+		return Dummy;
+	}
+
+	// for read access, can have additional data attached to the end. Can also be compressed
+	TConstArrayView<uint8> GetReadView() const
 	{
 		FinalizeShaderCode();
 
-		return ShaderCodeWithOptionalData;
+		return ShaderCodeResource.GetCodeView();
 	}
 
 	bool IsCompressed() const
@@ -1236,21 +1338,25 @@ public:
 		AddOptionalData(Key, (uint8*)InString, Size);
 	}
 
-	// Converts code to FShaderCodeResource format. Note that this is destructive (internal code array will be moved into the resource).
-	FShaderCodeResource ConvertToResource(EShaderFrequency Frequency, FSHAHash OutputHash) const
+	// Populates FShaderCodeResource's header buffer and returns the fully populated resource struct
+	const FShaderCodeResource& GetFinalizedResource(EShaderFrequency Frequency, FSHAHash OutputHash) const
 	{
-		check(OptionalDataSize == -1); // shader code must be finalized before converting to FShaderCodeResource
-		FShaderCodeResource Resource;
-		// need to persist this in the resource since it will not be readable from the code array after compression
-		Resource.ShaderCodeSize = GetShaderCodeSize();
-		Resource.Frequency = Frequency;
+		// shader code must be finalized prior to calling this function
+		// the finalize process will have created the code FSharedBuffer on the resource already
+		check(OptionalDataSize == -1);
+
+		// If the header is already populated, resource has already been finalized, early out
+		if (ShaderCodeResource.Header)
+		{
+			// sanity check
+			check(ShaderCodeResource.GetFrequency() == Frequency);
+			return ShaderCodeResource;
+		}
 
 		// Validate that compression settings used for this ShaderCode by the compilation process match what is expected
 		FName ShaderCompressionFormat = GetShaderCompressionFormat();
 		if (ShaderCompressionFormat != NAME_None)
 		{
-			Resource.UncompressedSize = GetUncompressedSize();
-
 			// we trust that SCWs also obeyed by the same CVar, so we expect a compressed shader code at this point
 			// However, if we see an uncompressed shader, it perhaps means that SCW tried to compress it, but the result was worse than uncompressed. 
 			// Because of that we special-case NAME_None here
@@ -1264,11 +1370,10 @@ public:
 						*GetCompressionFormat().ToString()
 					);
 					// unreachable
-					return Resource;
+					return ShaderCodeResource;
 				}
 
 				// assume uncompressed due to worse ratio than the compression
-				Resource.UncompressedSize = ShaderCodeWithOptionalData.Num();
 				UE_LOG(LogShaders, Verbose, TEXT("Shader %s is expected to be compressed with %s, but it arrived uncompressed (size=%d). Assuming compressing made it longer and storing uncompressed."),
 					*OutputHash.ToString(),
 					*ShaderCompressionFormat.ToString(),
@@ -1292,31 +1397,31 @@ public:
 						static_cast<int32>(GetOodleLevel())
 					);
 					// unreachable
-					return Resource;
+					return ShaderCodeResource;
 				}
 			}
 		}
-		else
-		{
-			Resource.UncompressedSize = ShaderCodeWithOptionalData.Num();
-		}
 
-		Resource.Code = MoveTemp(ShaderCodeWithOptionalData);
-		return Resource;
+		// Shader library/shader map usage expects uncompressed size to be set to the full code buffer size if uncompressed; so we need to apply that
+		// transformation here (and reverse it when populating from a FShaderCodeResource, see mirroring code in SetFromResource below)
+		ShaderCodeResource.PopulateHeader(UncompressedSize == 0 ? static_cast<int32>(ShaderCodeResource.Code.GetSize()) : UncompressedSize, GetShaderCodeSize(), Frequency);
+		return ShaderCodeResource;
 	}
 
 	void SetFromResource(FShaderCodeResource&& Resource)
 	{
-		// Set to the state of a finalized ShaderCode object
-		UncompressedSize = Resource.UncompressedSize;
-		ShaderCodeSize = Resource.ShaderCodeSize;
+		ShaderCodeResource = MoveTemp(Resource);
+		// Set the internal state of this FShaderCode to that of a finalized (and possibly compressed) ShaderCode object
 		OptionalDataSize = -1;
+		ShaderCodeSize = ShaderCodeResource.GetShaderCodeSize();
+
+		// as above, set UncompressedSize to 0 if not compressed, indicated by the resource uncompressed size matching the code buffer size.
+		int32 ResourceUncompressedSize = ShaderCodeResource.GetUncompressedSize();
+		UncompressedSize = ResourceUncompressedSize == ShaderCodeResource.Code.GetSize() ? 0 : ResourceUncompressedSize;
 
 		// already validated that compression settings matched when serializing the resource, so we can just initialize them to the known-correct values
 		CompressionFormat = GetShaderCompressionFormat();
 		GetShaderCompressionOodleSettings(OodleCompressor, OodleLevel);
-
-		ShaderCodeWithOptionalData = MoveTemp(Resource.Code);
 	}
 
 	friend RENDERCORE_API FArchive& operator<<(FArchive& Ar, FShaderCode& Output);
