[Backout] - CL36437712
Original CL Desc
-----------------------------------------------------------------
[shaders] modify FShaderCode finalize to create a FSharedBuffer object, and modify all downstream uses of shader code to re-use this buffer (job cache, pushes to DDC, shader maps, and shader library). This reduces total amount of LLM tracked memory allocated at the end of a cold Lyra PS4 cook by about ~350MB; impact likely much larger for cooks of larger projects.

diff --git a/Engine/Plugins/FX/Niagara/Source/NiagaraEditor/Private/NiagaraCompiler.cpp b/Engine/Plugins/FX/Niagara/Source/NiagaraEditor/Private/NiagaraCompiler.cpp
--- a/Engine/Plugins/FX/Niagara/Source/NiagaraEditor/Private/NiagaraCompiler.cpp
+++ b/Engine/Plugins/FX/Niagara/Source/NiagaraEditor/Private/NiagaraCompiler.cpp
@@ -2225,9 +2225,9 @@ TOptional<FNiagaraCompileResults> FHlslNiagaraCompiler::GetCompileResult(int32 J
 	FVectorVMCompilationOutput CompilationOutput;
 	if (CompilationJob->ShaderCompileJob->bSucceeded)
 	{
-		TConstArrayView<uint8> Code = CompilationJob->ShaderCompileJob->Output.ShaderCode.GetReadView();
+		const TArray<uint8>& Code = CompilationJob->ShaderCompileJob->Output.ShaderCode.GetReadAccess();
 		FShaderCodeReader ShaderCode(Code);
-		FMemoryReaderView Ar(Code, true);
+		FMemoryReader Ar(Code, true);
 		Ar.SetLimitSize(ShaderCode.GetActualShaderCodeSize());
 		Ar << CompilationOutput;
 
diff --git a/Engine/Shaders/Public/ShaderVersion.ush b/Engine/Shaders/Public/ShaderVersion.ush
--- a/Engine/Shaders/Public/ShaderVersion.ush
+++ b/Engine/Shaders/Public/ShaderVersion.ush
@@ -3,4 +3,4 @@
 // Each time the console command is executed it generates a new GUID. As this file is included
 // in Platform.ush (which should be included in any shader) it allows to invalidate the shader DDC.
 // 
-#pragma message("UESHADERMETADATA_VERSION 1115A78C-1E1B-4E96-A416-DA839794A1C5")
\ No newline at end of file
+#pragma message("UESHADERMETADATA_VERSION 0D63C246-532E-4D9D-83F3-906A518818CC")
\ No newline at end of file
diff --git a/Engine/Source/Developer/Apple/MetalShaderFormat/Private/MetalShaderCompiler.cpp b/Engine/Source/Developer/Apple/MetalShaderFormat/Private/MetalShaderCompiler.cpp
--- a/Engine/Source/Developer/Apple/MetalShaderFormat/Private/MetalShaderCompiler.cpp
+++ b/Engine/Source/Developer/Apple/MetalShaderFormat/Private/MetalShaderCompiler.cpp
@@ -1184,7 +1184,7 @@ bool StripShader_Metal(TArray<uint8>& Code, class FString const& DebugPath, bool
 				// jam it into the output bytes
 				NewAr.Serialize(SourceCode.GetData(), SourceCode.Num());
 				
-				Code = NewCode.GetReadView();
+				Code = NewCode.GetReadAccess();
 			}
 		}
 		else
@@ -1278,7 +1278,7 @@ uint64 AppendShader_Metal(FString const& WorkingDir, const FSHAHash& Hash, TArra
 						NewAr << OfflineCompiledFlag;
                         Header.Serialize(NewAr, SRT);
 						
-						InShaderCode = NewCode.GetReadView();
+						InShaderCode = NewCode.GetReadAccess();
 						
 						UE_LOG(LogShaders, Verbose, TEXT("Archiving succeeded: shader %s (Len: %0.8x, CRC: %0.8x, SHA: %s)"), *ShaderName, Header.SourceLen, Header.SourceCRC, *Hash.ToString());
 					}
diff --git a/Engine/Source/Developer/ShaderCompilerCommon/Private/ShaderCompilerCommon.cpp b/Engine/Source/Developer/ShaderCompilerCommon/Private/ShaderCompilerCommon.cpp
--- a/Engine/Source/Developer/ShaderCompilerCommon/Private/ShaderCompilerCommon.cpp
+++ b/Engine/Source/Developer/ShaderCompilerCommon/Private/ShaderCompilerCommon.cpp
@@ -2223,12 +2223,12 @@ namespace UE::ShaderCompilerCommon
 				// always output decompressed code as it's slightly more useful for A/B comparisons
 				TArray<uint8> DecompressedCode;
 				DecompressedCode.SetNum(Output.ShaderCode.GetUncompressedSize());
-				bool bSucceed = FCompression::UncompressMemory(NAME_Oodle, DecompressedCode.GetData(), DecompressedCode.Num(), Output.ShaderCode.GetReadView().GetData(), Output.ShaderCode.GetShaderCodeSize());
+				bool bSucceed = FCompression::UncompressMemory(NAME_Oodle, DecompressedCode.GetData(), DecompressedCode.Num(), Output.ShaderCode.GetReadAccess().GetData(), Output.ShaderCode.GetShaderCodeSize());
 				FFileHelper::SaveArrayToFile(DecompressedCode, *ShaderCodeFileName);
 			}
 			else
 			{
-				FFileHelper::SaveArrayToFile(Output.ShaderCode.GetReadView(), *ShaderCodeFileName);
+				FFileHelper::SaveArrayToFile(Output.ShaderCode.GetReadAccess(), *ShaderCodeFileName);
 			}
 		}
 
diff --git a/Engine/Source/Developer/TargetPlatform/Public/Interfaces/IShaderFormat.h b/Engine/Source/Developer/TargetPlatform/Public/Interfaces/IShaderFormat.h
--- a/Engine/Source/Developer/TargetPlatform/Public/Interfaces/IShaderFormat.h
+++ b/Engine/Source/Developer/TargetPlatform/Public/Interfaces/IShaderFormat.h
@@ -9,7 +9,6 @@
 
 class FSerializedShaderArchive;
 class FShaderPreprocessOutput;
-class FSharedBuffer;
 struct FShaderCompilerEnvironment;
 struct FShaderCompilerInput;
 struct FShaderCompilerOutput;
@@ -61,19 +60,6 @@ public:
 	*/
 	virtual bool SupportsShaderArchives() const { return false; }
 	
-	UE_DEPRECATED(5.5, "Use version accepting a TArray of FSharedBuffer code buffers.")
-	virtual bool CreateShaderArchive(FString const& LibraryName,
-		FName ShaderFormatAndShaderPlatformName,
-		const FString& WorkingDirectory,
-		const FString& OutputDir,
-		const FString& DebugOutputDir,
-		const FSerializedShaderArchive& SerializedShaders,
-		const TArray<TArray<uint8>>& ShaderCode,
-		TArray<FString>* OutputFiles) const
-	{ 
-		return false; 
-	}
-
 	/**
      * Create a format specific archive for precompiled shader code.
      *
@@ -91,11 +77,9 @@ public:
 		const FString& OutputDir,
 		const FString& DebugOutputDir,
 		const FSerializedShaderArchive& SerializedShaders,
-		const TArray<FSharedBuffer>& ShaderCode,
+		const TArray<TArray<uint8>>& ShaderCode,
 		TArray<FString>* OutputFiles) const
-	{ 
-		return false; 
-	}
+	{ return false; }
 	
 	/**
 	 * Can the shader format compile shaders to the native binary format for the platform.
diff --git a/Engine/Source/Runtime/Core/Private/UObject/FortniteReleaseBranchCustomObjectVersion.cpp b/Engine/Source/Runtime/Core/Private/UObject/FortniteReleaseBranchCustomObjectVersion.cpp
--- a/Engine/Source/Runtime/Core/Private/UObject/FortniteReleaseBranchCustomObjectVersion.cpp
+++ b/Engine/Source/Runtime/Core/Private/UObject/FortniteReleaseBranchCustomObjectVersion.cpp
@@ -9,7 +9,5 @@ TMap<FGuid, FGuid> FFortniteReleaseBranchCustomObjectVersion::GetSystemGuids()
 	SystemGuids.Add(DevGuids.Niagara_LatestScriptCompileVersion, FGuid("D35A655CD3FDA4408724D77E1316E35B"));
 	SystemGuids.Add(DevGuids.SkeletalMeshDerivedDataVersion, FGuid("5705956EC7134274A5A1BD3FC74DB3A9"));
 	SystemGuids.Add(DevGuids.MaterialTranslationDDCVersion, FGuid("CC1B3A336E144A0B9E78014E07775988"));
-	SystemGuids.Add(DevGuids.GLOBALSHADERMAP_DERIVEDDATA_VER, FGuid("E60D877F790D43268E1596F6223E1677"));
-	SystemGuids.Add(DevGuids.MATERIALSHADERMAP_DERIVEDDATA_VER, FGuid("3D174FB4B42C433CB10198CE17CA920F"));
 	return SystemGuids;
 }
diff --git a/Engine/Source/Runtime/Engine/Private/ShaderCompiler/ShaderCompilerJobCache.cpp b/Engine/Source/Runtime/Engine/Private/ShaderCompiler/ShaderCompilerJobCache.cpp
--- a/Engine/Source/Runtime/Engine/Private/ShaderCompiler/ShaderCompilerJobCache.cpp
+++ b/Engine/Source/Runtime/Engine/Private/ShaderCompiler/ShaderCompilerJobCache.cpp
@@ -378,7 +378,7 @@ public:
 	FSharedBuffer JobOutput;
 
 	/** Separate blobs for shader code */
-	TArray<FCompositeBuffer> JobCode;
+	TArray<FSharedBuffer> JobCode;
 
 	/** Path to where the cached debug info is stored. */
 	FString CachedDebugInfoPath;
@@ -627,12 +627,9 @@ private:
 		FBlake3 Hasher;
 		check(SerializeContext.HasData());
 		Hasher.Update(SerializeContext.ShaderObjectData.GetData(), SerializeContext.ShaderObjectData.GetSize());
-		for (const FCompositeBuffer& CodeBuf : SerializeContext.ShaderCode)
+		for (const FSharedBuffer& CodeBuf : SerializeContext.ShaderCode)
 		{
-			for (const FSharedBuffer& CodeBufSegment : CodeBuf.GetSegments())
-			{
-				Hasher.Update(CodeBufSegment.GetData(), CodeBuf.GetSize());
-			}
+			Hasher.Update(CodeBuf.GetData(), CodeBuf.GetSize());
 		}
 		return Hasher.Finalize();
 	}
@@ -2235,10 +2232,8 @@ void FShaderJobCache::AddJobOutput(FShaderJobData& JobData, const FShaderCommonC
 				for (int32 CodeIndex = 0; CodeIndex < SaveContext.ShaderCode.Num(); ++CodeIndex)
 				{
 					const FString CacheCodeFilename = FString::Printf(TEXT("%s/%s_%d.bytecode"), *InputDebugInfoPath, *InputSourceFilename);
-					const FCompositeBuffer& JobCode = SaveContext.ShaderCode[CodeIndex];
-					check(JobCode.GetSegments().Num() == 2); // first segment is header, second is actual code
-					FSharedBuffer CodeBuffer = JobCode.GetSegments()[1];
-					FFileHelper::SaveArrayToFile(TArrayView<const uint8>((const uint8*)CodeBuffer.GetData(), CodeBuffer.GetSize()), *CacheCodeFilename);
+					const FSharedBuffer& JobCode = SaveContext.ShaderCode[CodeIndex];
+					FFileHelper::SaveArrayToFile(TArrayView<const uint8>((const uint8*)JobCode.GetData(), JobCode.GetSize()), *CacheCodeFilename);
 				}
 			}
 		}
diff --git a/Engine/Source/Runtime/RenderCore/Internal/ShaderSerialization.h b/Engine/Source/Runtime/RenderCore/Internal/ShaderSerialization.h
--- a/Engine/Source/Runtime/RenderCore/Internal/ShaderSerialization.h
+++ b/Engine/Source/Runtime/RenderCore/Internal/ShaderSerialization.h
@@ -2,7 +2,6 @@
 
 #pragma once
 
-#include "Memory/CompositeBuffer.h"
 #include "Memory/SharedBuffer.h"
 #include "Templates/Function.h"
 #include "Templates/UniquePtr.h"
@@ -10,7 +9,7 @@
 
 class FMemoryReaderView;
 class FMemoryWriter64;
-class FShaderCodeResource;
+struct FShaderCodeResource;
 
 #if WITH_EDITOR
 namespace UE::DerivedData
@@ -71,12 +70,12 @@ struct FShaderCacheSerializeContext : public FShaderSerializeContext
 	/* View on array of buffers which store the bytecode objects for the object, one per shader/stage.
 	 * Note that this may or may not point to the OwnedShaderCode array below, depending on usage.
 	 */
-	TArrayView<FCompositeBuffer> ShaderCode;
+	TArrayView<FSharedBuffer> ShaderCode;
 
 	/* Array of buffers actually owned by this context object; it is valid for this to be empty in cases
 	 * where the array of ShaderCode buffers is stored externally.
 	 */
-	TArray<FCompositeBuffer> OwnedShaderCode;
+	TArray<FSharedBuffer> OwnedShaderCode;
 
 	/* Get the total serialized size of data for this context; note that this will return 0 if called prior to the FSharedBuffers
 	 * being set (this is done in the derived implementations, see below).
@@ -88,7 +87,7 @@ struct FShaderCacheSerializeContext : public FShaderSerializeContext
 		{
 			Size += ShaderObjectData.GetSize();
 
-			for (FCompositeBuffer& CodeBuf : ShaderCode)
+			for (FSharedBuffer& CodeBuf : ShaderCode)
 			{
 				Size += CodeBuf.GetSize();
 			}
@@ -97,7 +96,7 @@ struct FShaderCacheSerializeContext : public FShaderSerializeContext
 	}
 
 	/* Populates the given code array (transfering ownership) and resets the internal view to point to the new owning array's data */
-	void MoveCode(TArray<FCompositeBuffer>& TargetCode)
+	void MoveCode(TArray<FSharedBuffer>& TargetCode)
 	{
 		TargetCode = MoveTemp(OwnedShaderCode);
 		ShaderCode = TargetCode;
@@ -143,7 +142,7 @@ struct FShaderCacheLoadContext : public FShaderCacheSerializeContext
 	/*	Default constructor, use when array of code buffers will be allocated via ReadFromRecord */
 	RENDERCORE_API FShaderCacheLoadContext();
 	/* Constructor which references buffers (and an array of code buffers) owned elsewhere, does not allocate the OwnedShaderCode array */
-	RENDERCORE_API FShaderCacheLoadContext(FSharedBuffer ShaderObjectData, TArrayView<FCompositeBuffer> CodeBuffers);
+	RENDERCORE_API FShaderCacheLoadContext(FSharedBuffer ShaderObjectData, TArrayView<FSharedBuffer> CodeBuffers);
 
 	/* Call to reset reader to start position so the same load context can be used to populate multiple objects. */
 	RENDERCORE_API void Reuse();
diff --git a/Engine/Source/Runtime/RenderCore/Private/Shader.cpp b/Engine/Source/Runtime/RenderCore/Private/Shader.cpp
--- a/Engine/Source/Runtime/RenderCore/Private/Shader.cpp
+++ b/Engine/Source/Runtime/RenderCore/Private/Shader.cpp
@@ -669,7 +669,7 @@ FShaderCompiledShaderInitializerType::FShaderCompiledShaderInitializerType(
 	: Type(InType)
 	, Parameters(InParameters)
 	, Target(CompilerOutput.Target)
-	, Code(CompilerOutput.ShaderCode.GetReadView())
+	, Code(CompilerOutput.ShaderCode.GetReadAccess())
 	, ParameterMap(CompilerOutput.ParameterMap)
 	, OutputHash(CompilerOutput.OutputHash)
 	, MaterialShaderMapHash(InMaterialShaderMapHash)
diff --git a/Engine/Source/Runtime/RenderCore/Private/ShaderCodeLibrary.cpp b/Engine/Source/Runtime/RenderCore/Private/ShaderCodeLibrary.cpp
--- a/Engine/Source/Runtime/RenderCore/Private/ShaderCodeLibrary.cpp
+++ b/Engine/Source/Runtime/RenderCore/Private/ShaderCodeLibrary.cpp
@@ -1427,18 +1427,17 @@ struct FEditorShaderCodeArchive
 				{
 					const FShaderCodeResource& SourceShaderResource = Code->ShaderCodeResources[i];
 					FShaderCodeEntry& SerializedShaderEntry = SerializedShaders.ShaderEntries[ShaderIndex];
-					SerializedShaderEntry.Frequency = SourceShaderResource.GetFrequency();
-					FSharedBuffer CodeBuffer = SourceShaderResource.GetCodeBuffer();
-					SerializedShaderEntry.Size = CodeBuffer.GetSize();
-					SerializedShaderEntry.UncompressedSize = SourceShaderResource.GetUncompressedSize();
-					check(SerializedShaderEntry.Size > 0);
-					ShaderCode.Add(CodeBuffer);
+					SerializedShaderEntry.Frequency = SourceShaderResource.Frequency;
+					SerializedShaderEntry.Size = SourceShaderResource.Code.Num();
+					SerializedShaderEntry.UncompressedSize = SourceShaderResource.UncompressedSize;
+					check(!SourceShaderResource.Code.IsEmpty());
+					ShaderCode.Add(SourceShaderResource.Code);
 					check(ShaderCode.Num() == SerializedShaders.ShaderEntries.Num());
 
 					CodeStats.NumUniqueShaders++;
-					CodeStats.ShadersUniqueSize += CodeBuffer.GetSize();
+					CodeStats.ShadersUniqueSize += SourceShaderResource.Code.Num();
 				}
-				CodeStats.ShadersSize += Code->ShaderCodeResources[i].GetCodeView().NumBytes();
+				CodeStats.ShadersSize += Code->ShaderCodeResources[i].Code.Num();
 				SerializedShaders.ShaderIndices[ShaderMapEntry.ShaderIndicesOffset + i] = ShaderIndex;
 			}
 
@@ -1518,7 +1517,7 @@ struct FEditorShaderCodeArchive
 		int32 NumShadersSentWithCode = 0;
 
 		const FSerializedShaderArchive& SourceArchive = this->SerializedShaders;
-		TArray<FSharedBuffer>& SourceShaderCodes = this->ShaderCode;
+		TArray<TArray<uint8>>& SourceShaderCodes = this->ShaderCode;
 		for (int32 SourceShaderMapIndex : LocalShaderMapsToCopy)
 		{
 			const FSHAHash& SourceShaderMapHash = SourceArchive.ShaderMapHashes[SourceShaderMapIndex];
@@ -1552,15 +1551,15 @@ struct FEditorShaderCodeArchive
 					check(TargetShaderIndex == TargetArchive.ShaderEntries.Num() - 1);
 
 					const FShaderCodeEntry& SourceShaderEntry = SourceArchive.ShaderEntries[SourceShaderIndex];
-					FSharedBuffer& SourceShaderCode = SourceShaderCodes[SourceShaderIndex];
+					TArray<uint8>& SourceShaderCode = SourceShaderCodes[SourceShaderIndex];
 					FShaderCodeEntry& TargetShaderEntry = TargetArchive.ShaderEntries[TargetShaderIndex];
 
 					TargetShaderEntry = SourceShaderEntry;
-					if (!SourceShaderCode.GetSize())
+					if (!SourceShaderCode.IsEmpty())
 					{
 						TargetShaderEntry.Offset = TargetFlatShaderCode.Num();
-						check(SourceShaderEntry.Size == SourceShaderCode.GetSize());
-						if ((MaxShaderSize > 0 && TargetFlatShaderCode.Num() + (int64)SourceShaderCode.GetSize() > MaxShaderSize) ||
+						check(SourceShaderEntry.Size == SourceShaderCode.Num());
+						if ((MaxShaderSize > 0 && TargetFlatShaderCode.Num() + SourceShaderCode.Num() > MaxShaderSize) ||
 							(MaxShaderCount > 0 && NumShadersSentWithCode > MaxShaderCount))
 						{
 							// We have to stop here to avoid overflowing the shader limit. Send the shaders we have accumulated
@@ -1577,11 +1576,11 @@ struct FEditorShaderCodeArchive
 						}
 						++NumShadersSentWithCode;
 
-						TargetFlatShaderCode.Append(MakeArrayView(reinterpret_cast<const uint8*>(SourceShaderCode.GetData()), SourceShaderCode.GetSize()));
+						TargetFlatShaderCode.Append(SourceShaderCode);
 
-						// Reset the ShaderCode reference to (potentially) save memory in the local process. The consumer of the TargetArchive and
+						// Empty the ShaderCode to save memory in the local process. The consumer of the TargetArchive and
 						// TargetFlatShaderCode will be the only one that needs to read it.
-						SourceShaderCode.Reset();
+						SourceShaderCode.Empty();
 					}
 					else
 					{
@@ -1602,7 +1601,7 @@ struct FEditorShaderCodeArchive
 	{
 		bool bOk = true;
 		FSerializedShaderArchive& TargetArchive = this->SerializedShaders;
-		TArray<FSharedBuffer>& TargetShaderCodes = this->ShaderCode;
+		TArray<TArray<uint8>>& TargetShaderCodes = this->ShaderCode;
 
 		// Add all the shaders; we can sometimes get messages that send the shaders in advance without sending the shadermaps that use them
 		for (int32 SourceShaderIndex = 0; SourceShaderIndex < SourceArchive.ShaderHashes.Num(); ++SourceShaderIndex)
@@ -1616,6 +1615,7 @@ struct FEditorShaderCodeArchive
 			}
 			check(TargetShaderIndex == TargetArchive.ShaderEntries.Num() - 1 &&
 				TargetShaderCodes.Num() == TargetArchive.ShaderEntries.Num() - 1);
+			TArray<uint8>& TargetShaderCode = TargetShaderCodes.Emplace_GetRef();
 
 			const FShaderCodeEntry& SourceShaderEntry = SourceArchive.ShaderEntries[SourceShaderIndex];
 			FShaderCodeEntry& TargetShaderEntry = TargetArchive.ShaderEntries[TargetShaderIndex];
@@ -1645,8 +1645,8 @@ struct FEditorShaderCodeArchive
 			}
 			else
 			{
-				// Copy from source's flat list to the target's separate FSharedBuffer for each shader
-				TargetShaderCodes.Add(FSharedBuffer::Clone(SourceFlatShaderCode.GetData() + SourceShaderEntry.Offset, SourceShaderEntry.Size));
+				const TConstArrayView<uint8> SourceShaderCode(SourceFlatShaderCode.GetData() + SourceShaderEntry.Offset, SourceShaderEntry.Size);
+				TargetShaderCode = SourceShaderCode; // Copy from source's flat list to the target's separate TArray<uint8> for each shader
 			}
 		}
 
@@ -1708,8 +1708,8 @@ struct FEditorShaderCodeArchive
 		NewChunk->ShaderCode.Empty();
 		for (int32 NewArchiveIdx = 0, NumIndices = ShaderCodeEntriesNeeded.Num(); NewArchiveIdx < NumIndices; ++NewArchiveIdx)
 		{
-			FSharedBuffer& SourceShaderCodeEntry = ShaderCode[ShaderCodeEntriesNeeded[NewArchiveIdx]];
-			check(SourceShaderCodeEntry.GetSize() > 0);
+			TArray<uint8>& SourceShaderCodeEntry = ShaderCode[ShaderCodeEntriesNeeded[NewArchiveIdx]];
+			check(!SourceShaderCodeEntry.IsEmpty());
 			NewChunk->ShaderCode.Add(SourceShaderCodeEntry);
 		}
 
@@ -1739,8 +1739,8 @@ struct FEditorShaderCodeArchive
 					const FShaderCodeEntry& OtherShaderEntry = OtherArchive.SerializedShaders.ShaderEntries[OtherShaderIndex];
 					SerializedShaders.ShaderEntries[ShaderIndex] = OtherShaderEntry;
 
-					const FSharedBuffer& OtherShaderCodeEntry = OtherArchive.ShaderCode[OtherShaderIndex];
-					check(OtherShaderCodeEntry.GetSize() > 0);
+					const TArray<uint8>& OtherShaderCodeEntry = OtherArchive.ShaderCode[OtherShaderIndex];
+					check(!OtherShaderCodeEntry.IsEmpty());
 					ShaderCode.Add(OtherShaderCodeEntry);
 					check(ShaderCode.Num() == SerializedShaders.ShaderEntries.Num());
 				}
@@ -1777,16 +1777,16 @@ struct FEditorShaderCodeArchive
 					const FShaderCodeEntry& OtherShaderEntry = OtherShaders.ShaderEntries[OtherShaderIndex];
 					SerializedShaders.ShaderEntries[ShaderIndex] = OtherShaderEntry;
 
-					check(ShaderCode.Num() == SerializedShaders.GetNumShaders() - 1);
+					TArray<uint8>& Code = ShaderCode.AddDefaulted_GetRef();
+					check(ShaderCode.Num() == SerializedShaders.GetNumShaders());
 
 					// Read shader code from archive and add shader to set
 					const int64 ReadSize = OtherShaderEntry.Size;
 					check(ReadSize > 0);
 					const int64 ReadOffset = OtherShaderCodeOffset + OtherShaderEntry.Offset;
-					FUniqueBuffer Code = FUniqueBuffer::Alloc(ReadSize);
+					Code.SetNumUninitialized(ReadSize);
 					Ar.Seek(ReadOffset);
 					Ar.Serialize(Code.GetData(), ReadSize);
-					ShaderCode.Add(Code.MoveToShared());
 				}
 				SerializedShaders.ShaderIndices[ShaderMapEntry.ShaderIndicesOffset + i] = ShaderIndex;
 			}
@@ -1815,18 +1815,16 @@ struct FEditorShaderCodeArchive
 				for (int32 Index = 0; Index < ShaderCode.Num(); ++Index)
 				{
 					const FShaderCodeEntry& Entry = SerializedShaders.ShaderEntries[Index];
+					TArray<uint8>& Code = ShaderCode[Index];
 					check(Entry.Size > 0);
-					FUniqueBuffer Code = FUniqueBuffer::Alloc(Entry.Size);
+					Code.SetNumUninitialized(Entry.Size);
 					PrevCookedAr->Serialize(Code.GetData(), Entry.Size);
-
 					bOK = !PrevCookedAr->GetError();
 					if (!bOK)
 					{
 						UE_LOG(LogShaderLibrary, Error, TEXT("Failed to deserialize shader code for %s from %s"), *SerializedShaders.ShaderHashes[Index].ToString(), *IntermediateFormatPath);
 						break;
 					}
-
-					ShaderCode[Index] = Code.MoveToShared();
 				}
 			}
 			else
@@ -1974,8 +1972,8 @@ struct FEditorShaderCodeArchive
 					*FileWriter << SerializedShaders;
 					for (auto& Code : ShaderCode)
 					{
-						check(Code.GetSize() > 0);
-						FileWriter->Serialize(const_cast<void*>(Code.GetData()), Code.GetSize());
+						check(!Code.IsEmpty());
+						FileWriter->Serialize(Code.GetData(), Code.Num());
 					}
 
 					FileWriter->Close();
@@ -2252,7 +2250,7 @@ private:
 	 * The element at index N holds the ShaderCode for the element of SerializedShaders.ShaderEntries at index N.
 	 * In MultiprocessCooking elements can be empty if they have been transferred to the Director (bHasCopiedAndCleared will be true in this case).
 	 */
-	TArray<FSharedBuffer> ShaderCode;
+	TArray<TArray<uint8>> ShaderCode;
 	/** A list of ShaderMaps that have not yet been copied to the CookDirector. Used only by MultiprocessCookWorkers. */
 	TSet<int32> ShaderMapsToCopy;
 	/** True if CopyToArchiveAndClear has been called, otherwise false. If false we avoid doing some tracking since it might never be used. */
diff --git a/Engine/Source/Runtime/RenderCore/Private/ShaderCore.cpp b/Engine/Source/Runtime/RenderCore/Private/ShaderCore.cpp
--- a/Engine/Source/Runtime/RenderCore/Private/ShaderCore.cpp
+++ b/Engine/Source/Runtime/RenderCore/Private/ShaderCore.cpp
@@ -1292,7 +1292,7 @@ void FShaderCompilerOutput::GenerateOutputHash()
 {
 	FSHA1 HashState;
 	
-	TArrayView<const uint8> Code = ShaderCode.GetReadView();
+	const TArray<uint8>& Code = ShaderCode.GetReadAccess();
 
 	// we don't hash the optional attachments as they would prevent sharing (e.g. many materials share the same VS)
 	uint32 ShaderCodeSize = ShaderCode.GetShaderCodeSize();
@@ -1609,21 +1609,19 @@ public:
 		int32 PackedShaderKey = Compiler->GetPackedShaderKey();
 		CombinedSource.Append(reinterpret_cast<const uint8*>(&PackedShaderKey), sizeof(PackedShaderKey));
 
-		TArrayView<const uint8> PrimaryCode = Job.Output.ShaderCode.GetReadView();
-		const uint32 PrimaryLength = PrimaryCode.NumBytes();
-		TArrayView<const uint8> SecondaryCode = Job.SecondaryOutput->ShaderCode.GetReadView();
-		const uint32 SecondaryLength = SecondaryCode.NumBytes();
-
-		CombinedSource.Reserve(PrimaryLength + sizeof(PrimaryLength) + SecondaryLength + sizeof(SecondaryLength));
+		const uint32 PrimaryLength = Job.Output.ShaderCode.GetReadAccess().Num();
 		CombinedSource.Append(reinterpret_cast<const uint8*>(&PrimaryLength), sizeof(PrimaryLength));
+
+		const uint32 SecondaryLength = Job.SecondaryOutput->ShaderCode.GetReadAccess().Num();
 		CombinedSource.Append(reinterpret_cast<const uint8*>(&SecondaryLength), sizeof(SecondaryLength));
-		CombinedSource.Append(PrimaryCode);
-		CombinedSource.Append(SecondaryCode);
+
+		CombinedSource.Append(Job.Output.ShaderCode.GetReadAccess());
+		CombinedSource.Append(Job.SecondaryOutput->ShaderCode.GetReadAccess());
 
 		// Replace Output shader code with the combined result
 		Job.Output.ShaderCode = {};
 		TArray<uint8>& FinalShaderCode = Job.Output.ShaderCode.GetWriteAccess();
-		FinalShaderCode = MoveTemp(CombinedSource);
+		FinalShaderCode.Append(CombinedSource);
 		Job.Output.ShaderCode.FinalizeShaderCode();
 	}
 
@@ -3780,28 +3778,27 @@ void FShaderCode::Compress(FName ShaderCompressionFormat, FOodleDataCompression:
 
 	TArray<uint8> Compressed;
 	// conventional formats will fail if the compressed size isn't enough, Oodle needs a more precise estimate
-	TConstArrayView<uint8> Code = ShaderCodeResource.GetCodeView();
-	int32 CompressedSize = (ShaderCompressionFormat != NAME_Oodle) ? Code.NumBytes() : FOodleDataCompression::CompressedBufferSizeNeeded(Code.NumBytes());
+	int32 CompressedSize = (ShaderCompressionFormat != NAME_Oodle) ? ShaderCodeWithOptionalData.Num() : FOodleDataCompression::CompressedBufferSizeNeeded(ShaderCodeWithOptionalData.Num());
 	Compressed.AddUninitialized(CompressedSize);
 
 	// non-Oodle format names use the old API, for NAME_Oodle we replace the call with the custom invocation
 	bool bCompressed = false;
 	if (ShaderCompressionFormat != NAME_Oodle)
 	{
-		bCompressed = FCompression::CompressMemory(ShaderCompressionFormat, Compressed.GetData(), CompressedSize, Code.GetData(), Code.Num(), COMPRESS_BiasSize);
+		bCompressed = FCompression::CompressMemory(ShaderCompressionFormat, Compressed.GetData(), CompressedSize, ShaderCodeWithOptionalData.GetData(), ShaderCodeWithOptionalData.Num(), COMPRESS_BiasSize);
 	}
 	else
 	{
-		CompressedSize = FOodleDataCompression::Compress(Compressed.GetData(), CompressedSize, Code.GetData(), Code.Num(),
+		CompressedSize = FOodleDataCompression::Compress(Compressed.GetData(), CompressedSize, ShaderCodeWithOptionalData.GetData(), ShaderCodeWithOptionalData.Num(),
 			InOodleCompressor, InOodleLevel);
 		bCompressed = CompressedSize != 0;
 	}
 
 	// there is code that assumes that if CompressedSize == CodeSize, the shader isn't compressed. Because of that, do not accept equal compressed size (very unlikely anyway)
-	if (bCompressed && CompressedSize < Code.Num())
+	if (bCompressed && CompressedSize < ShaderCodeWithOptionalData.Num())
 	{
 		// cache the ShaderCodeSize since it will no longer possible to get it as the reader will fail to parse the compressed data
-		FShaderCodeReader Wrapper(Code);
+		FShaderCodeReader Wrapper(ShaderCodeWithOptionalData);
 		ShaderCodeSize = Wrapper.GetShaderCodeSize();
 		checkf(ShaderCodeSize >= 0, TEXT("Unable to determine ShaderCodeSize from uncompressed code"), ShaderCodeSize);
 
@@ -3809,44 +3806,13 @@ void FShaderCode::Compress(FName ShaderCompressionFormat, FOodleDataCompression:
 		CompressionFormat = ShaderCompressionFormat;
 		OodleCompressor = InOodleCompressor;
 		OodleLevel = InOodleLevel;
-		UncompressedSize = Code.Num();
+		UncompressedSize = ShaderCodeWithOptionalData.Num();
 
 		Compressed.SetNum(CompressedSize);
-		ShaderCodeResource.Code = MakeSharedBufferFromArray(MoveTemp(Compressed));
+		ShaderCodeWithOptionalData = Compressed;
 	}
 }
 
-FArchive& operator<<(FArchive& Ar, FSharedBuffer& Buffer)
-{
-	uint64 Len = Buffer.GetSize();
-	Ar << Len;
-
-	if (Ar.IsLoading())
-	{
-		Buffer.Reset();
-
-		if (Len >= 0)
-		{
-			FUniqueBuffer BufTmp = FUniqueBuffer::Alloc(Len);
-			Ar.Serialize(BufTmp.GetData(), Len);
-			Buffer = BufTmp.MoveToShared();
-		}
-	}
-	else if (Ar.IsSaving())
-	{
-		Ar.Serialize(const_cast<void*>(Buffer.GetData()), Len);
-	}
-
-	return Ar;
-}
-
-FArchive& operator<<(FArchive& Ar, FShaderCodeResource& Resource)
-{
-	Ar << Resource.Header;
-	Ar << Resource.Code;
-	return Ar;
-}
-
 FArchive& operator<<(FArchive& Ar, FShaderCode& Output)
 {
 	if (Ar.IsLoading())
@@ -3859,8 +3825,7 @@ FArchive& operator<<(FArchive& Ar, FShaderCode& Output)
 	}
 
 	// Note: this serialize is used to pass between UE and the shader compile worker, recompile both when modifying
-	FSharedBuffer& CodeBuffer = Output.ShaderCodeResource.Code;
-	Ar << CodeBuffer;
+	Ar << Output.ShaderCodeWithOptionalData;
 	Ar << Output.UncompressedSize;
 	{
 		FString CompressionFormatString(Output.CompressionFormat.ToString());
@@ -4046,7 +4011,7 @@ void FShaderCompileJob::SerializeOutput(FShaderCacheSerializeContext& Ctx, int32
 	FShaderCodeResource CodeResource;
 	if (bIsSaving)
 	{
-		CodeResource = Output.GetFinalizedCodeResource();
+		CodeResource = Output.ConvertCodeToResource();
 	}
 
 	Ctx.SerializeCodeFunc(CodeResource, CodeIndex);
diff --git a/Engine/Source/Runtime/RenderCore/Private/ShaderResource.cpp b/Engine/Source/Runtime/RenderCore/Private/ShaderResource.cpp
--- a/Engine/Source/Runtime/RenderCore/Private/ShaderResource.cpp
+++ b/Engine/Source/Runtime/RenderCore/Private/ShaderResource.cpp
@@ -179,7 +179,7 @@ static void ApplyResourceStats(FShaderMapResourceCode& Resource)
 	INC_DWORD_STAT_BY(STAT_Shaders_ShaderResourceMemory, Resource.GetSizeBytes());
 	for (const FShaderCodeResource& Shader : Resource.ShaderCodeResources)
 	{
-		INC_DWORD_STAT_BY_FName(GetMemoryStatType(Shader.GetFrequency()).GetName(), Shader.GetCodeBuffer().GetSize());
+		INC_DWORD_STAT_BY_FName(GetMemoryStatType(Shader.Frequency).GetName(), Shader.Code.Num());
 	}
 #endif // STATS
 }
@@ -190,7 +190,7 @@ static void RemoveResourceStats(FShaderMapResourceCode& Resource)
 	DEC_DWORD_STAT_BY(STAT_Shaders_ShaderResourceMemory, Resource.GetSizeBytes());
 	for (const FShaderCodeResource& Shader : Resource.ShaderCodeResources)
 	{
-		DEC_DWORD_STAT_BY_FName(GetMemoryStatType(Shader.GetFrequency()).GetName(), Shader.GetCodeBuffer().GetSize());
+		DEC_DWORD_STAT_BY_FName(GetMemoryStatType(Shader.Frequency).GetName(), Shader.Code.Num());
 	}
 #endif // STATS
 }
@@ -231,7 +231,7 @@ uint32 FShaderMapResourceCode::GetSizeBytes() const
 	uint64 Size = sizeof(*this) + ShaderHashes.GetAllocatedSize() + ShaderCodeResources.GetAllocatedSize();
 	for (const FShaderCodeResource& Entry : ShaderCodeResources)
 	{
-		Size += Entry.GetCacheBuffer().GetSize();
+		Size += Entry.Code.GetAllocatedSize();
 	}
 	check(Size <= TNumericLimits<uint32>::Max());
 	return static_cast<uint32>(Size);
@@ -257,7 +257,9 @@ void FShaderMapResourceCode::AddShaderCompilerOutput(const FShaderCompilerOutput
 		AddEditorOnlyData(Index, DebugName, Output.PlatformDebugData, Output.Errors, Output.ShaderStatistics, DebugInfo);
 #endif
 
-		ShaderCodeResources.Insert(Output.GetFinalizedCodeResource(), Index);
+		// This moves the internal array in the Output's ShaderCode object into the resource format; assuming that after
+		// the code is moved to the shadermap that it's no longer required in the job output struct.
+		ShaderCodeResources.Insert(Output.ConvertCodeToResource(), Index);
 	}
 #if WITH_EDITORONLY_DATA
 	else
@@ -353,8 +355,8 @@ void FShaderMapResourceCode::ToString(FStringBuilderBase& OutString) const
 	for (int32 i = 0; i < ShaderHashes.Num(); ++i)
 	{
 		const FShaderCodeResource& Res = ShaderCodeResources[i];
-		OutString.Appendf(TEXT("    [%d]: { Hash: %s, Freq: %s, Size: %llu, UncompressedSize: %d }\n"),
-			i, *ShaderHashes[i].ToString(), GetShaderFrequencyString(Res.GetFrequency()), Res.GetCodeBuffer().GetSize(), Res.GetUncompressedSize());
+		OutString.Appendf(TEXT("    [%d]: { Hash: %s, Freq: %s, Size: %d, UncompressedSize: %d }\n"),
+			i, *ShaderHashes[i].ToString(), GetShaderFrequencyString(Res.Frequency), Res.Code.Num(), Res.UncompressedSize);
 	}
 }
 
@@ -653,21 +655,20 @@ FRHIShader* FShaderMapResource_InlineCode::CreateRHIShaderOrCrash(int32 ShaderIn
 
 	FMemStackBase& MemStack = FMemStack::Get();
 	const FShaderCodeResource& ShaderCodeResource = Code->ShaderCodeResources[ShaderIndex];
-	FSharedBuffer ShaderCode = ShaderCodeResource.GetCodeBuffer();
-	TConstArrayView<uint8> ShaderCodeView = ShaderCodeResource.GetCodeView();
+	const uint8* ShaderCode = ShaderCodeResource.Code.GetData();
 
 	FMemMark Mark(MemStack);
-	int32 UncompressedSize = ShaderCodeResource.GetUncompressedSize();
-	if (ShaderCode.GetSize() != UncompressedSize)
+	if (ShaderCodeResource.Code.Num() != ShaderCodeResource.UncompressedSize)
 	{
-		void* UncompressedCode = MemStack.Alloc(UncompressedSize, 16);
-		bool bSucceed = FCompression::UncompressMemory(GetShaderCompressionFormat(), UncompressedCode, UncompressedSize, ShaderCode.GetData(), ShaderCode.GetSize());
+		void* UncompressedCode = MemStack.Alloc(ShaderCodeResource.UncompressedSize, 16);
+		bool bSucceed = FCompression::UncompressMemory(GetShaderCompressionFormat(), UncompressedCode, ShaderCodeResource.UncompressedSize, ShaderCode, ShaderCodeResource.Code.Num());
 		check(bSucceed);
-		ShaderCodeView = MakeArrayView(reinterpret_cast<const uint8*>(UncompressedCode), UncompressedSize);
+		ShaderCode = (uint8*)UncompressedCode;
 	}
 
+	const auto ShaderCodeView = MakeArrayView(ShaderCode, ShaderCodeResource.UncompressedSize);
 	const FSHAHash& ShaderHash = Code->ShaderHashes[ShaderIndex];
-	const EShaderFrequency Frequency = ShaderCodeResource.GetFrequency();
+	const EShaderFrequency Frequency = ShaderCodeResource.Frequency;
 
 	TRefCountPtr<FRHIShader> RHIShader;
 	switch (Frequency)
diff --git a/Engine/Source/Runtime/RenderCore/Private/ShaderSerialization.cpp b/Engine/Source/Runtime/RenderCore/Private/ShaderSerialization.cpp
--- a/Engine/Source/Runtime/RenderCore/Private/ShaderSerialization.cpp
+++ b/Engine/Source/Runtime/RenderCore/Private/ShaderSerialization.cpp
@@ -14,7 +14,10 @@ FShaderCacheSaveContext::FShaderCacheSaveContext()
 {
 	SerializeCodeFunc = [this](FShaderCodeResource& Res, int32 Index)
 		{
-			OwnedShaderCode.Add(Res.GetCacheBuffer());
+			TArray64<uint8> CodeData;
+			FMemoryWriter64 Ar(CodeData, true);
+			Ar << Res;
+			OwnedShaderCode.Add(MakeSharedBufferFromArray(MoveTemp(CodeData)));
 			// reset the array view any time an entry is added; we do this instead of calling Resize in the reserve delegate
 			// and setting it there since not all code paths (i.e. single job cache records) call reserve
 			ShaderCode = OwnedShaderCode;
@@ -55,14 +58,11 @@ UE::DerivedData::FCacheRecord FShaderCacheSaveContext::BuildCacheRecord(const UE
 	Finalize();
 
 	UE::DerivedData::FCacheRecordBuilder RecordBuilder(Key);
-	RecordBuilder.AddValue(ShaderObjectDataValue, ShaderObjectData);
+	RecordBuilder.AddValue(ShaderObjectDataValue, UE::DerivedData::FValue::Compress(ShaderObjectData));
 	int32 CodeIndex = 0;
-	// Code buffers are already compressed, don't waste cycles attempting (and failing) to recompress them
-	const ECompressedBufferCompressor CodeComp = ECompressedBufferCompressor::NotSet;
-	const ECompressedBufferCompressionLevel CodeCompLevel = ECompressedBufferCompressionLevel::None;
-	for (FCompositeBuffer& CodeBuf : ShaderCode)
+	for (FSharedBuffer& CodeBuf : ShaderCode)
 	{
-		RecordBuilder.AddValue(ShaderCodeDataValue.MakeIndexed(CodeIndex++), UE::DerivedData::FValue(FCompressedBuffer::Compress(CodeBuf, CodeComp, CodeCompLevel)));
+		RecordBuilder.AddValue(ShaderCodeDataValue.MakeIndexed(CodeIndex++), UE::DerivedData::FValue::Compress(CodeBuf));
 	}
 
 	TCbWriter<16> MetaWriter;
@@ -80,11 +80,12 @@ FShaderCacheLoadContext::FShaderCacheLoadContext()
 {
 	SerializeCodeFunc = [this](FShaderCodeResource& Res, int32 Index)
 		{
-			Res.PopulateFromComposite(ShaderCode[Index]);
+			FMemoryReaderView Ar(ShaderCode[Index]);
+			Ar << Res;
 		};
 }
 
-FShaderCacheLoadContext::FShaderCacheLoadContext(FSharedBuffer InShaderObjectData, TArrayView<FCompositeBuffer> InCodeBuffers) : FShaderCacheLoadContext()
+FShaderCacheLoadContext::FShaderCacheLoadContext(FSharedBuffer InShaderObjectData, TArrayView<FSharedBuffer> InCodeBuffers) : FShaderCacheLoadContext()
 {
 	ShaderObjectData = InShaderObjectData;
 	ShaderCode = InCodeBuffers;
@@ -112,8 +113,7 @@ void FShaderCacheLoadContext::ReadFromRecord(const UE::DerivedData::FCacheRecord
 	OwnedShaderCode.Reserve(CodeCount);
 	for (int32 CodeIndex = 0; CodeIndex < CodeCount; ++CodeIndex)
 	{
-		FSharedBuffer CombinedBuffer = Record.GetValue(ShaderCodeDataValue.MakeIndexed(CodeIndex)).GetData().Decompress();
-		OwnedShaderCode.Add(FShaderCodeResource::Unpack(CombinedBuffer));
+		OwnedShaderCode.Add(Record.GetValue(ShaderCodeDataValue.MakeIndexed(CodeIndex)).GetData().Decompress());
 	}
 	ShaderCode = OwnedShaderCode;
 }
diff --git a/Engine/Source/Runtime/RenderCore/Public/Shader.h b/Engine/Source/Runtime/RenderCore/Public/Shader.h
--- a/Engine/Source/Runtime/RenderCore/Public/Shader.h
+++ b/Engine/Source/Runtime/RenderCore/Public/Shader.h
@@ -1556,7 +1556,7 @@ struct FShaderCompiledShaderInitializerType
 	const FShaderType* Type;
 	const FShaderType::FParameters* Parameters;
 	FShaderTarget Target;
-	TConstArrayView<uint8> Code;
+	TArrayView<const uint8> Code;
 	const FShaderParameterMap& ParameterMap;
 	const FSHAHash& OutputHash;
 	FSHAHash MaterialShaderMapHash;
diff --git a/Engine/Source/Runtime/RenderCore/Public/ShaderCompilerCore.h b/Engine/Source/Runtime/RenderCore/Public/ShaderCompilerCore.h
--- a/Engine/Source/Runtime/RenderCore/Public/ShaderCompilerCore.h
+++ b/Engine/Source/Runtime/RenderCore/Public/ShaderCompilerCore.h
@@ -28,9 +28,9 @@ class FShaderPipelineCompileJob;
 typedef TSharedPtr<TArray<ANSICHAR>, ESPMode::ThreadSafe> FShaderSharedAnsiStringPtr;
 
 // this is for the protocol, not the data, bump if FShaderCompilerInput/FShaderPreprocessOutput serialization, SerializeWorkerInput or ProcessInputFromArchive changes.
-inline const int32 ShaderCompileWorkerInputVersion = 29;
+inline const int32 ShaderCompileWorkerInputVersion = 28;
 // this is for the protocol, not the data, bump if FShaderCompilerOutput or WriteToOutputArchive changes.
-inline const int32 ShaderCompileWorkerOutputVersion = 25;
+inline const int32 ShaderCompileWorkerOutputVersion = 24;
 // this is for the protocol, not the data.
 inline const int32 ShaderCompileWorkerSingleJobHeader = 'S';
 // this is for the protocol, not the data.
@@ -38,7 +38,7 @@ inline const int32 ShaderCompileWorkerPipelineJobHeader = 'P';
 
 // Modify this to invalidate _just_ the cache/DDC entries for individual shaders (will not cause shadermaps to rebuild if they are not otherwise out-of-date).
 // This should be bumped for changes to the FShaderCompilerOutput data structure (in addition to ShaderCompileWorkerOutputVersion)
-inline static const FGuid UE_SHADER_CACHE_VERSION = FGuid("31D16D5D-E4DC-43FF-BB50-6648399DC4B0");
+inline static const FGuid UE_SHADER_CACHE_VERSION = FGuid("55EB8969-BA75-4CBD-B7A1-84194D140CBD");
 
 namespace UE::ShaderCompiler
 {
@@ -629,9 +629,9 @@ struct FShaderCompilerOutput
 		Stat.Flags = Flags;
 	}
 
-	const FShaderCodeResource& GetFinalizedCodeResource() const
+	FShaderCodeResource ConvertCodeToResource() const
 	{
-		return ShaderCode.GetFinalizedResource(Target.GetFrequency(), OutputHash);
+		return ShaderCode.ConvertToResource(Target.GetFrequency(), OutputHash);
 	}
 
 	void SetCodeFromResource(FShaderCodeResource&& Resource)
diff --git a/Engine/Source/Runtime/RenderCore/Public/ShaderCore.h b/Engine/Source/Runtime/RenderCore/Public/ShaderCore.h
--- a/Engine/Source/Runtime/RenderCore/Public/ShaderCore.h
+++ b/Engine/Source/Runtime/RenderCore/Public/ShaderCore.h
@@ -17,8 +17,6 @@
 #include "HAL/Platform.h"
 #include "HAL/UnrealMemory.h"
 #include "Logging/LogMacros.h"
-#include "Memory/CompositeBuffer.h"
-#include "Memory/SharedBuffer.h"
 #include "Misc/AssertionMacros.h"
 #include "Misc/CString.h"
 #include "Misc/CoreStats.h"
@@ -939,10 +937,10 @@ typedef uint32 RENDERCORE_ATTRIBUTE_UNALIGNED unaligned_uint32;
 // later we can transform that to the actual class passed around at the RHI level
 class FShaderCodeReader
 {
-	TConstArrayView<uint8> ShaderCode;
+	TArrayView<const uint8> ShaderCode;
 
 public:
-	FShaderCodeReader(TConstArrayView<uint8> InShaderCode)
+	FShaderCodeReader(TArrayView<const uint8> InShaderCode)
 		: ShaderCode(InShaderCode)
 	{
 		check(ShaderCode.Num());
@@ -953,9 +951,9 @@ public:
 		return ShaderCode.Num() - GetOptionalDataSize();
 	}
 
-	TConstArrayView<uint8> GetOffsetShaderCode(int32 Offset)
+	TArrayView<const uint8> GetOffsetShaderCode(int32 Offset)
 	{
-		return MakeConstArrayView(ShaderCode.GetData() + Offset, GetActualShaderCodeSize() - Offset);
+		return MakeArrayView(ShaderCode.GetData() + Offset, GetActualShaderCodeSize() - Offset);
 	}
 
 	// for convenience
@@ -1084,101 +1082,20 @@ public:
 	}
 };
 
-class FShaderCode;
-
-class FShaderCodeResource
+struct FShaderCodeResource
 {
-	struct FHeader
-	{
-		int32 UncompressedSize;		// full size of code array before compression
-		int32 ShaderCodeSize;		// uncompressed size excluding optional data
-		EShaderFrequency Frequency;
-	};
-
-	FSharedBuffer Header;		// The above FHeader struct persisted in a shared buffer
-	FSharedBuffer Code;			// The bytecode buffer as constructed by FShaderCode::FinalizeShaderCode
-
-	friend class FShaderCode;
-	friend RENDERCORE_API FArchive& operator<<(FArchive& Ar, FShaderCode& Code);
-	friend FArchive& operator<<(FArchive& Ar, FShaderCodeResource& Resource);
-
-public:
-
-	/* Returns a uint8 array view representation of the Code FSharedBuffer, for compatibility's sake (much downstream
-	 * usage of shader code expects an array of uint8)
-	 */
-	TConstArrayView<uint8> GetCodeView() const
-	{
-		return MakeConstArrayView(reinterpret_cast<const uint8*>(Code.GetData()), static_cast<int32>(Code.GetSize()));
-	}
-
-	/* Return the buffer storing just the shader code for this resource */
-	FSharedBuffer GetCodeBuffer() const
-	{
-		return Code;
-	}
-
-	/* Returns a single composite buffer referencing both the header and code data to be cached. */
-	FCompositeBuffer GetCacheBuffer() const
-	{
-		return FCompositeBuffer(Header, Code);
-	}
-
-	/* Unpacks the given FSharedBuffer into separate header/code buffer views and returns them as a 2-segment composite buffer. 
-	 * Note that this is required since when pushing a composite buffer to DDC it does not maintain the segment structure.
-	 */
-	static FCompositeBuffer Unpack(FSharedBuffer MonolithicBuffer)
-	{
-		FMemoryView FullBufferView = MonolithicBuffer.GetView();
-
-		return FCompositeBuffer(
-			MonolithicBuffer.MakeView(FullBufferView.Left(sizeof(FHeader)), MonolithicBuffer),
-			MonolithicBuffer.MakeView(FullBufferView.RightChop(sizeof(FHeader)), MonolithicBuffer));
-	}
-	
-	/* Sets the Header and Code shared buffer references in this resource to the segments referenced
-	 * by the given composite buffer.
-	 */
-	void PopulateFromComposite(FCompositeBuffer CacheBuffer)
-	{
-		check(CacheBuffer.GetSegments().Num() == 2);
-		Header = CacheBuffer.GetSegments()[0];
-		check(Header.GetSize() == sizeof(FHeader));
-		Code = CacheBuffer.GetSegments()[1];
-	}
-
-	/* Populates the header for this code resource with the given sizes and frequency. 
-	 * Note that this is done as a separate process from the construction of the Code buffer
-	 * as the shader frequency is only known by the owning job, and not stored in the FShaderCode.
-	 */
-	void PopulateHeader(int32 UncompressedSize, int32 ShaderCodeSize, EShaderFrequency Frequency)
-	{
-		check(Code);
-		FHeader HeaderData{ UncompressedSize, ShaderCodeSize, Frequency };
-		Header = FSharedBuffer::Clone(&HeaderData, sizeof(HeaderData));
-	}
-
-	/* Retrieves the uncompressed size of the shader code as stored in the FHeader buffer. */
-	int32 GetUncompressedSize() const
-	{
-		check(Header);
-		return reinterpret_cast<const FHeader*>(Header.GetData())->UncompressedSize;
-	}
-
-	/* Retrieves the actual shader code size (excluding optional data) as stored in the FHeader buffer. */
-	int32 GetShaderCodeSize() const
-	{
-		check(Header);
-		return reinterpret_cast<const FHeader*>(Header.GetData())->ShaderCodeSize;
-	}
+	TArray<uint8> Code;
+	int32 UncompressedSize; // full size of code array before compression
+	int32 ShaderCodeSize;	// uncompressed size excluding optional data
+	EShaderFrequency Frequency;
 
-	/* Retrieves the shader frequency as stored in the FHeader buffer. */
-	EShaderFrequency GetFrequency() const
+	friend FArchive& operator<<(FArchive& Ar, FShaderCodeResource& Resource)
 	{
-		check(Header);
-		return reinterpret_cast<const FHeader*>(Header.GetData())->Frequency;
+		uint8 Freq = Resource.Frequency;
+		Ar << Resource.Code << Resource.UncompressedSize << Resource.ShaderCodeSize << Freq;
+		Resource.Frequency = (EShaderFrequency)Freq;
+		return Ar;
 	}
-
 };
 
 class FShaderCode
@@ -1188,8 +1105,6 @@ class FShaderCode
 	// access through class methods
 	mutable TArray<uint8> ShaderCodeWithOptionalData;
 
-	mutable FShaderCodeResource ShaderCodeResource;
-
 	/** ShaderCode may be compressed in SCWs on demand. If this value isn't null, the shader code is compressed. */
 	mutable int32 UncompressedSize;
 
@@ -1220,14 +1135,12 @@ public:
 	// adds CustomData or does nothing if that was already done before
 	void FinalizeShaderCode() const
 	{
-		if (OptionalDataSize != -1)
+		if(OptionalDataSize != -1)
 		{
 			checkf(UncompressedSize == 0, TEXT("FShaderCode::FinalizeShaderCode() was called after compressing the code"));
 			OptionalDataSize += sizeof(OptionalDataSize);
 			ShaderCodeWithOptionalData.Append((const uint8*)&OptionalDataSize, sizeof(OptionalDataSize));
 			OptionalDataSize = -1;
-
-			ShaderCodeResource.Code = MakeSharedBufferFromArray(MoveTemp(ShaderCodeWithOptionalData));
 		}
 	}
 
@@ -1252,32 +1165,17 @@ public:
 		{
 			FinalizeShaderCode();
 
-			if (UncompressedSize != 0) // already compressed, get code size from resource
-			{
-				return ShaderCodeResource.GetShaderCodeSize();
-			}
-			else
-			{
-				// code buffer has been populated but not compressed, can still read additional fields from code buffer
-				FShaderCodeReader Wrapper(ShaderCodeResource.GetCodeView());
-				return Wrapper.GetShaderCodeSize();
-			}
+			FShaderCodeReader Wrapper(ShaderCodeWithOptionalData);
+			return Wrapper.GetShaderCodeSize();
 		}
 	}
 
-	UE_DEPRECATED(5.5, "Use GetReadView")
-	const TArray<uint8>& GetReadAccess() const
-	{
-		static TArray<uint8> Dummy;
-		return Dummy;
-	}
-
 	// for read access, can have additional data attached to the end. Can also be compressed
-	TConstArrayView<uint8> GetReadView() const
+	const TArray<uint8>& GetReadAccess() const
 	{
 		FinalizeShaderCode();
 
-		return ShaderCodeResource.GetCodeView();
+		return ShaderCodeWithOptionalData;
 	}
 
 	bool IsCompressed() const
@@ -1338,25 +1236,21 @@ public:
 		AddOptionalData(Key, (uint8*)InString, Size);
 	}
 
-	// Populates FShaderCodeResource's header buffer and returns the fully populated resource struct
-	const FShaderCodeResource& GetFinalizedResource(EShaderFrequency Frequency, FSHAHash OutputHash) const
+	// Converts code to FShaderCodeResource format. Note that this is destructive (internal code array will be moved into the resource).
+	FShaderCodeResource ConvertToResource(EShaderFrequency Frequency, FSHAHash OutputHash) const
 	{
-		// shader code must be finalized prior to calling this function
-		// the finalize process will have created the code FSharedBuffer on the resource already
-		check(OptionalDataSize == -1);
-
-		// If the header is already populated, resource has already been finalized, early out
-		if (ShaderCodeResource.Header)
-		{
-			// sanity check
-			check(ShaderCodeResource.GetFrequency() == Frequency);
-			return ShaderCodeResource;
-		}
+		check(OptionalDataSize == -1); // shader code must be finalized before converting to FShaderCodeResource
+		FShaderCodeResource Resource;
+		// need to persist this in the resource since it will not be readable from the code array after compression
+		Resource.ShaderCodeSize = GetShaderCodeSize();
+		Resource.Frequency = Frequency;
 
 		// Validate that compression settings used for this ShaderCode by the compilation process match what is expected
 		FName ShaderCompressionFormat = GetShaderCompressionFormat();
 		if (ShaderCompressionFormat != NAME_None)
 		{
+			Resource.UncompressedSize = GetUncompressedSize();
+
 			// we trust that SCWs also obeyed by the same CVar, so we expect a compressed shader code at this point
 			// However, if we see an uncompressed shader, it perhaps means that SCW tried to compress it, but the result was worse than uncompressed. 
 			// Because of that we special-case NAME_None here
@@ -1370,10 +1264,11 @@ public:
 						*GetCompressionFormat().ToString()
 					);
 					// unreachable
-					return ShaderCodeResource;
+					return Resource;
 				}
 
 				// assume uncompressed due to worse ratio than the compression
+				Resource.UncompressedSize = ShaderCodeWithOptionalData.Num();
 				UE_LOG(LogShaders, Verbose, TEXT("Shader %s is expected to be compressed with %s, but it arrived uncompressed (size=%d). Assuming compressing made it longer and storing uncompressed."),
 					*OutputHash.ToString(),
 					*ShaderCompressionFormat.ToString(),
@@ -1397,31 +1292,31 @@ public:
 						static_cast<int32>(GetOodleLevel())
 					);
 					// unreachable
-					return ShaderCodeResource;
+					return Resource;
 				}
 			}
 		}
+		else
+		{
+			Resource.UncompressedSize = ShaderCodeWithOptionalData.Num();
+		}
 
-		// Shader library/shader map usage expects uncompressed size to be set to the full code buffer size if uncompressed; so we need to apply that
-		// transformation here (and reverse it when populating from a FShaderCodeResource, see mirroring code in SetFromResource below)
-		ShaderCodeResource.PopulateHeader(UncompressedSize == 0 ? static_cast<int32>(ShaderCodeResource.Code.GetSize()) : UncompressedSize, GetShaderCodeSize(), Frequency);
-		return ShaderCodeResource;
+		Resource.Code = MoveTemp(ShaderCodeWithOptionalData);
+		return Resource;
 	}
 
 	void SetFromResource(FShaderCodeResource&& Resource)
 	{
-		ShaderCodeResource = MoveTemp(Resource);
-		// Set the internal state of this FShaderCode to that of a finalized (and possibly compressed) ShaderCode object
+		// Set to the state of a finalized ShaderCode object
+		UncompressedSize = Resource.UncompressedSize;
+		ShaderCodeSize = Resource.ShaderCodeSize;
 		OptionalDataSize = -1;
-		ShaderCodeSize = ShaderCodeResource.GetShaderCodeSize();
-
-		// as above, set UncompressedSize to 0 if not compressed, indicated by the resource uncompressed size matching the code buffer size.
-		int32 ResourceUncompressedSize = ShaderCodeResource.GetUncompressedSize();
-		UncompressedSize = ResourceUncompressedSize == ShaderCodeResource.Code.GetSize() ? 0 : ResourceUncompressedSize;
 
 		// already validated that compression settings matched when serializing the resource, so we can just initialize them to the known-correct values
 		CompressionFormat = GetShaderCompressionFormat();
 		GetShaderCompressionOodleSettings(OodleCompressor, OodleLevel);
+
+		ShaderCodeWithOptionalData = MoveTemp(Resource.Code);
 	}
 
 	friend RENDERCORE_API FArchive& operator<<(FArchive& Ar, FShaderCode& Output);
