Fix frame sync issues
 - The engine still uses legacy syncing behaviour where the game thread syncs with the render thread. Renderer refactors have removed most of the RHI thread flushes that happened per frame, which were the only thing synchronizing the game/render threads with the RHI thread. Without these flushes, and when occlusion queries are disabled, the game thread can run ahead of the RHI thread by several hundred frames, since it is now entirely unsynchronized.
 - This fix changes mode 0 of "r.GTSyncType" to sync with the N-2 RHI thread frame when the RHI thread is active, rather than the N-1 render thread frame. The game thread is now always synchronized with the RHI thread to prevent it running ahead.
 - Mode 1 of "r.GTSyncType" now works even when vsync is disabled, and syncs the game thread with the N-1 RHI thread frame (same behaviour as before).

diff --git a/Engine/Plugins/Experimental/TextureGraph/Source/TextureGraphEngine/Model/Mix/MixManager.cpp b/Engine/Plugins/Experimental/TextureGraph/Source/TextureGraphEngine/Model/Mix/MixManager.cpp
--- a/Engine/Plugins/Experimental/TextureGraph/Source/TextureGraphEngine/Model/Mix/MixManager.cpp
+++ b/Engine/Plugins/Experimental/TextureGraph/Source/TextureGraphEngine/Model/Mix/MixManager.cpp
@@ -141,7 +141,7 @@ void MixManager::Exit()
 	
 	FRenderCommandFence SuspendFence;
 	
-	SuspendFence.BeginFence(true);
+	SuspendFence.BeginFence();
 	SuspendFence.Wait();
 	
 	Suspend();
diff --git a/Engine/Source/Runtime/Engine/Private/UnrealEngine.cpp b/Engine/Source/Runtime/Engine/Private/UnrealEngine.cpp
--- a/Engine/Source/Runtime/Engine/Private/UnrealEngine.cpp
+++ b/Engine/Source/Runtime/Engine/Private/UnrealEngine.cpp
@@ -455,6 +455,17 @@ TAutoConsoleVariable<int32> CVarAllowOneFrameThreadLag(
 	TEXT("Whether to allow the rendering thread to lag one frame behind the game thread (0: disabled, otherwise enabled)")
 );
 
+TAutoConsoleVariable<int32> CVarGTSyncType(
+	TEXT("r.GTSyncType"),
+	0,
+	TEXT("Determines how the game thread syncs with the render thread, RHI thread and GPU.\n")
+	TEXT("Syncing to the GPU swap chain flip allows for lower frame latency.\n")
+	TEXT(" 0 - Sync the game thread with the N-2 RHI thread frame when the RHI thread is enabled, or the N-1 render thread frame when it is disabled (default).\n")
+	TEXT(" 1 - Sync the game thread with the N-1 RHI thread frame.\n")
+	TEXT(" 2 - Sync the game thread with the GPU swap chain flip (only on supported platforms).\n"),
+	ECVF_Default
+);
+
 static FAutoConsoleVariable CVarSystemResolution(
 	TEXT("r.SetRes"),
 	TEXT("1280x720w"),
@@ -736,9 +747,8 @@ void CalculateFPSTimings()
 	double CurrentTime = FPlatformTime::Seconds();
 	float FrameTimeMS = (float)((CurrentTime - LastTime) * 1000.0);
 
-	static auto CVarGTSyncType = IConsoleManager::Get().FindConsoleVariable(TEXT("r.GTSyncType"));
 	static auto CVarVsync = IConsoleManager::Get().FindConsoleVariable(TEXT("r.VSync"));
-	if (CVarGTSyncType->GetInt() == 2 && CVarVsync->GetInt() != 0)
+	if (CVarGTSyncType.GetValueOnGameThread() == 2 && CVarVsync->GetInt() != 0)
 	{
 		float RHIFrameTime = RHIGetFrameTime();
 		if (RHIFrameTime != 0)
@@ -13309,15 +13319,17 @@ void FFrameEndSync::Cleanup()
 {
 	// forcing static deallocation order:
 	// fences can hold task completion handles that need to be freed before their allocator is destroyed. wating for fences does the job
-	Fence[0].Wait();
-	Fence[1].Wait();
+	for (FRenderCommandFence& Fence : Fences)
+	{
+		Fence.Wait();
+	}
 
 	FCoreDelegates::OnEnginePreExit.Remove(CleanupDelegate);
 	// notify the destructor
 	CleanupDelegate = {};
 }
 
-void FFrameEndSync::Sync( bool bAllowOneFrameThreadLag )
+void FFrameEndSync::Sync()
 {
 	check(IsInGameThread());			
 
@@ -13335,8 +13347,50 @@ void FFrameEndSync::Sync( bool bAllowOneFrameThreadLag )
 	});
 #endif
 
-	// Since this is the frame end sync, allow sync with the RHI and GPU (true).
-	Fence[EventIndex].BeginFence(true);
+	using ESyncDepth = FRenderCommandFence::ESyncDepth;
+
+	ESyncDepth SyncDepth;
+	int32 NumFramesOverlap;
+
+	uint32 const GTSyncType = (uint32)FMath::Max(CVarGTSyncType.GetValueOnAnyThread(), 0);
+
+	if (CVarAllowOneFrameThreadLag.GetValueOnAnyThread() <= 0)
+	{
+		// Legacy r.OneFrameThreadLag cvar is set to zero... the gamethread should not start frame N+1 until all work for frame N has completed.
+		// Depending on r.GTSyncType, sync with either the RHI thread or the swapchain. It is no longer correct to only sync with the render thread (pre UE 5.5 behaviour).
+		
+		SyncDepth = GTSyncType == 2
+			? ESyncDepth::Swapchain
+			: ESyncDepth::RHIThread;
+
+		NumFramesOverlap = 0;
+	}
+	else
+	{
+		bool bHasRHIThread = IsRunningRHIInSeparateThread();
+
+		switch (GTSyncType)
+		{
+		case 0:  SyncDepth = ESyncDepth::RHIThread; NumFramesOverlap = bHasRHIThread ? 2 : 1; break;
+		case 1:  SyncDepth = ESyncDepth::RHIThread; NumFramesOverlap = 1; break;
+		case 2:  SyncDepth = ESyncDepth::Swapchain; NumFramesOverlap = 1; break;
+		default: SyncDepth = ESyncDepth::Swapchain; NumFramesOverlap = 1; break;
+		}
+	}
+
+	if (SyncDepth == ESyncDepth::Swapchain)
+	{
+		// Swapchain sync mode does not work when vsync is disabled. Fallback to RHI thread sync in that case.
+		static auto CVarVsync = IConsoleManager::Get().FindConsoleVariable(TEXT("r.VSync"));
+		check(CVarVsync != nullptr);
+
+		if (CVarVsync->GetInt() == 0)
+		{
+			SyncDepth = ESyncDepth::RHIThread;
+		}
+	}
+
+	Fences.Emplace_GetRef().BeginFence(SyncDepth);
 
 	bool bEmptyGameThreadTasks = !FTaskGraphInterface::Get().IsThreadProcessingTasks(ENamedThreads::GameThread);
 
@@ -13346,13 +13400,11 @@ void FFrameEndSync::Sync( bool bAllowOneFrameThreadLag )
 		FTaskGraphInterface::Get().ProcessThreadUntilIdle(ENamedThreads::GameThread);
 	}
 
-	// Use two events if we allow a one frame lag.
-	if( bAllowOneFrameThreadLag )
+	while (Fences.Num() > NumFramesOverlap)
 	{
-		EventIndex = (EventIndex + 1) % 2;
+		Fences[0].Wait(bEmptyGameThreadTasks); // here we also opportunistically execute game thread tasks while we wait
+		Fences.RemoveAt(0);
 	}
-
-	Fence[EventIndex].Wait(bEmptyGameThreadTasks);  // here we also opportunistically execute game thread tasks while we wait
 }
 
 FString appGetStartupMap(const TCHAR* CommandLine)
diff --git a/Engine/Source/Runtime/Engine/Public/UnrealEngine.h b/Engine/Source/Runtime/Engine/Public/UnrealEngine.h
--- a/Engine/Source/Runtime/Engine/Public/UnrealEngine.h
+++ b/Engine/Source/Runtime/Engine/Public/UnrealEngine.h
@@ -375,9 +375,8 @@ extern ENGINE_API UFont* GetStatsFont();
 class FFrameEndSync
 {
 	/** Pair of fences. */
-	FRenderCommandFence Fence[2];
-	/** Current index into events array. */
-	int32 EventIndex;
+	TArray<FRenderCommandFence, TInlineAllocator<3>> Fences;
+
 	/** cleanup delegate for engine pre-exit */
 	FDelegateHandle CleanupDelegate;
 
@@ -389,7 +388,7 @@ public:
 	 * Syncs the game thread with the render thread. Depending on passed in bool this will be a total
 	 * sync or a one frame lag.
 	 */
-	ENGINE_API void Sync( bool bAllowOneFrameThreadLag );
+	ENGINE_API void Sync();
 
 private:
 	void Cleanup();
diff --git a/Engine/Source/Runtime/Launch/Private/LaunchEngineLoop.cpp b/Engine/Source/Runtime/Launch/Private/LaunchEngineLoop.cpp
--- a/Engine/Source/Runtime/Launch/Private/LaunchEngineLoop.cpp
+++ b/Engine/Source/Runtime/Launch/Private/LaunchEngineLoop.cpp
@@ -6071,10 +6071,9 @@ void FEngineLoop::Tick()
 		{
 			SCOPE_CYCLE_COUNTER(STAT_FrameSyncTime);
 			// this could be perhaps moved down to get greater parallelism
-			// Sync game and render thread. Either total sync or allowing one frame lag.
+			// Sync game and render/RHI threads.
 			static FFrameEndSync FrameEndSync;
-			static auto CVarAllowOneFrameThreadLag = IConsoleManager::Get().FindTConsoleVariableDataInt(TEXT("r.OneFrameThreadLag"));
-			FrameEndSync.Sync( CVarAllowOneFrameThreadLag->GetValueOnGameThread() != 0 );
+			FrameEndSync.Sync();
 		}
 
 		// tick core ticker, threads & deferred commands
diff --git a/Engine/Source/Runtime/RHI/Private/RHIUtilities.cpp b/Engine/Source/Runtime/RHI/Private/RHIUtilities.cpp
--- a/Engine/Source/Runtime/RHI/Private/RHIUtilities.cpp
+++ b/Engine/Source/Runtime/RHI/Private/RHIUtilities.cpp
@@ -178,6 +178,7 @@ class FRHIFrameFlipTrackingRunnable : public FRunnable
 
 		~FFramePair()
 		{
+			TRACE_CPUPROFILER_EVENT_SCOPE(SyncTrigger_Swapchain);
 			Event.Trigger();
 		}
 
@@ -524,6 +525,7 @@ void FRHIFrameFlipTrackingRunnable::TriggerTaskEventOnFlip(uint64 PresentIndex,
 	{
 		// Platform does not support flip tracking.
 		// Signal the event now...
+		TRACE_CPUPROFILER_EVENT_SCOPE(SyncTrigger_Swapchain);
 		Event.Trigger();
 	}
 }
diff --git a/Engine/Source/Runtime/RenderCore/Private/RenderingThread.cpp b/Engine/Source/Runtime/RenderCore/Private/RenderingThread.cpp
--- a/Engine/Source/Runtime/RenderCore/Private/RenderingThread.cpp
+++ b/Engine/Source/Runtime/RenderCore/Private/RenderingThread.cpp
@@ -1008,48 +1008,10 @@ bool IsRenderThreadTimeoutSuspended()
 	return GTimeoutSuspendCount > 0;
 }
 
-TAutoConsoleVariable<int32> CVarGTSyncType(
-	TEXT("r.GTSyncType"),
-	0,
-	TEXT("Determines how the game thread syncs with the render thread, RHI thread and GPU.\n")
-	TEXT("Syncing to the GPU swap chain flip allows for lower frame latency.\n")
-	TEXT(" 0 - Sync the game thread with the render thread (default).\n")
-	TEXT(" 1 - Sync the game thread with the RHI thread.\n")
-	TEXT(" 2 - Sync the game thread with the GPU swap chain flip (only on supported platforms).\n"),
-	ECVF_Default);
-
-FRHICOMMAND_MACRO(FRHISyncFrameCommand)
-{
-	UE::Tasks::FTaskEvent TaskEvent;
-	int32 GTSyncType;
-
-	FORCEINLINE_DEBUGGABLE FRHISyncFrameCommand(UE::Tasks::FTaskEvent InTaskEvent, int32 InGTSyncType)
-		: TaskEvent(MoveTemp(InTaskEvent))
-		, GTSyncType(InGTSyncType)
-	{}
-
-	void Execute(FRHICommandListBase& CmdList)
-	{
-		if (GTSyncType == 1)
-		{
-			// Sync the Game Thread with the RHI Thread
-
-			// "Complete" the graph event
-			TaskEvent.Trigger();
-		}
-		else
-		{
-			// This command runs *after* a present has happened, so the counter has already been incremented.
-			// Subtracting 1 gives us the index of the frame that has *just* been presented.
-			RHITriggerTaskEventOnFlip(GRHIPresentCounter - 1, TaskEvent);
-		}
-	}
-};
-
 FRenderCommandFence::FRenderCommandFence() = default;
 FRenderCommandFence::~FRenderCommandFence() = default;
 
-void FRenderCommandFence::BeginFence(bool bSyncToRHIAndGPU)
+void FRenderCommandFence::BeginFence(ESyncDepth SyncDepth)
 {
 	if (!GIsThreadedRendering)
 	{
@@ -1058,128 +1020,62 @@ void FRenderCommandFence::BeginFence(bool bSyncToRHIAndGPU)
 	
 	if (GRenderCommandFenceBundlerState.Event && IsInGameThread())
 	{
+		check(SyncDepth == ESyncDepth::RenderThread);
 		CompletionTask = *GRenderCommandFenceBundlerState.Event;
 		return;
 	}
 
 	TRACE_CPUPROFILER_EVENT_SCOPE(FRenderCommandFence::BeginFence);
+	UE::Tasks::FTaskEvent Event { UE_SOURCE_LOCATION };
 
-	struct FRenderCommandPipeFence : public TConcurrentLinearObject<FRenderCommandPipeFence>
+	if (GRenderCommandPipeMode == ERenderCommandPipeMode::All)
 	{
-		FRenderCommandPipeFence(int32 InNumRefs)
-			: NumRefs(InNumRefs)
-		{}
-
-		void Trigger(int32 NumTriggerRefs = 1)
+		for (FRenderCommandPipe* Pipe : UE::RenderCommandPipe::GetPipes())
 		{
-			if (NumRefs.fetch_sub(NumTriggerRefs, std::memory_order_release) == 1)
+			// Skip pipes that aren't recording or replaying any work.
+			if (Pipe->IsRecording() && !Pipe->IsEmpty())
 			{
-				std::atomic_thread_fence(std::memory_order_acquire);
-				CompletionTaskEvent.Trigger();
-				delete this;
-			}
-		}
-
-		UE::Tasks::FTaskEvent CompletionTaskEvent{ UE_SOURCE_LOCATION };
-		std::atomic_int32_t NumRefs;
-	};
-
-	TConstArrayView<FRenderCommandPipe*> Pipes = GRenderCommandPipeMode == ERenderCommandPipeMode::All
-		? UE::RenderCommandPipe::GetPipes()
-		: TConstArrayView<FRenderCommandPipe*>{};
-
-	FRenderCommandPipeBitArray ActivePipeBits;
-	int32 NumActivePipes = 0;
-
-	for (FRenderCommandPipe* Pipe : Pipes)
-	{
-		// Skip pipes that aren't recording or replaying any work.
-		const bool bIsActive = Pipe->IsRecording() && !Pipe->IsEmpty();
-		ActivePipeBits.Add(bIsActive);
-		NumActivePipes += bIsActive ? 1 : 0;
-	}
-
-	FRenderCommandPipeFence* Fence = nullptr;
-
-	if (NumActivePipes > 0)
-	{
-		Fence = new FRenderCommandPipeFence(NumActivePipes + 1);
-
-		for (FRenderCommandPipeSetBitIterator BitIt(ActivePipeBits); BitIt; ++BitIt)
-		{
-			FRenderCommandPipe* Pipe = Pipes[BitIt.GetIndex()];
-
-			ENQUEUE_RENDER_COMMAND(BeginFence)(Pipe, [Fence]
-			{
-				Fence->Trigger();
-			});
-		}
-	}
+				UE::Tasks::FTaskEvent PipeEvent { UE_SOURCE_LOCATION };
+				Event.AddPrerequisites(PipeEvent);
 
-	const int32 GTSyncType = CVarGTSyncType.GetValueOnAnyThread();
-
-	if (bSyncToRHIAndGPU)
-	{
-		// Don't sync to the RHI and GPU if GtSyncType is disabled, or we're not vsyncing
-		//@TODO: do this logic in the caller?
-		static auto CVarVsync = IConsoleManager::Get().FindConsoleVariable(TEXT("r.VSync"));
-		check(CVarVsync != nullptr);
-
-		if (GTSyncType == 0 || CVarVsync->GetInt() == 0)
-		{
-			bSyncToRHIAndGPU = false;
+				ENQUEUE_RENDER_COMMAND(BeginFence)([PipeEvent = MoveTemp(PipeEvent)](FRHICommandList&) mutable
+				{
+					PipeEvent.Trigger();
+				});
+			}
 		}
 	}
 
-	if (bSyncToRHIAndGPU)
+	ENQUEUE_RENDER_COMMAND(BeginFence)([Event, SyncDepth](FRHICommandListImmediate& RHICmdList) mutable
 	{
-		UE::Tasks::FTaskEvent CompletionTaskEvent{ UE_SOURCE_LOCATION };
-
-		if (Fence)
+		if (SyncDepth == ESyncDepth::RenderThread)
 		{
-			// RHI frame sync command requires a task event, so connect it to the ref-counted fence event.
-			CompletionTaskEvent.AddPrerequisites(Fence->CompletionTaskEvent);
-			Fence->Trigger();
+			TRACE_CPUPROFILER_EVENT_SCOPE(SyncTrigger_RenderThread);
+			Event.Trigger();
 		}
-
-		ENQUEUE_RENDER_COMMAND(FSyncFrameCommand)(
-			[CompletionTaskEvent, GTSyncType, bSyncToRHIAndGPU](FRHICommandListImmediate& RHICmdList) mutable
+		else
 		{
-			if (IsRHIThreadRunning())
+			RHICmdList.EnqueueLambda([SyncDepth, Event](FRHICommandListImmediate&) mutable
 			{
-				ALLOC_COMMAND_CL(RHICmdList, FRHISyncFrameCommand)(MoveTemp(CompletionTaskEvent), GTSyncType);
-				RHICmdList.ImmediateFlush(EImmediateFlushType::DispatchToRHIThread);
-			}
-			else
-			{
-				FRHISyncFrameCommand Command(MoveTemp(CompletionTaskEvent), GTSyncType);
-				Command.Execute(RHICmdList);
-			}
-		});
-
-		CompletionTask = MoveTemp(CompletionTaskEvent);
-	}
-	else if (Fence)
-	{
-		CompletionTask = Fence->CompletionTaskEvent;
-
-		ENQUEUE_RENDER_COMMAND(BeginFence)([Fence](FRHICommandListBase& RHICmdList)
-		{
-			Fence->Trigger();
-		});
-	}
-	else
-	{
-		UE::Tasks::FTaskEvent CompletionTaskEvent{ UE_SOURCE_LOCATION };
-
-		ENQUEUE_RENDER_COMMAND(BeginFence)([CompletionTaskEvent](FRHICommandListBase& RHICmdList) mutable
-		{
-			CompletionTaskEvent.Trigger();
-		});
+				if (SyncDepth == ESyncDepth::RHIThread)
+				{
+					// Sync the Game Thread with the RHI Thread
+					TRACE_CPUPROFILER_EVENT_SCOPE(SyncTrigger_RHIThread);
+					Event.Trigger();
+				}
+				else
+				{
+					// This command runs *after* a present has happened, so the counter has already been incremented.
+					// Subtracting 1 gives us the index of the frame that has *just* been presented.
+					RHITriggerTaskEventOnFlip(GRHIPresentCounter - 1, Event);
+				}
+			});
+			RHICmdList.ImmediateFlush(EImmediateFlushType::DispatchToRHIThread);
+		}
+	});
 
-		CompletionTask = MoveTemp(CompletionTaskEvent);
-	}
-} //-V773
+	CompletionTask = MoveTemp(Event);
+}
 
 bool FRenderCommandFence::IsFenceComplete() const
 {
diff --git a/Engine/Source/Runtime/RenderCore/Public/RenderCommandFence.h b/Engine/Source/Runtime/RenderCore/Public/RenderCommandFence.h
--- a/Engine/Source/Runtime/RenderCore/Public/RenderCommandFence.h
+++ b/Engine/Source/Runtime/RenderCore/Public/RenderCommandFence.h
@@ -14,14 +14,25 @@
 class FRenderCommandFence
 {
 public:
+	enum class ESyncDepth
+	{
+		// The fence will be signalled by the render thread.
+		RenderThread,
+
+		// The fence will be enqueued to the RHI thread via a command on the immediate command list
+		// and signalled once all prior parallel translation and submission is complete.
+		RHIThread,
+
+		// The fence will be signalled according to the rate of flips in the swapchain.
+		// This is only supported on some platforms. On unsupported platforms, this behaves like RHIThread mode.
+		Swapchain
+	};
 
 	/**
-	 * Adds a fence command to the rendering command queue.
-	 * Conceptually, the pending fence count is incremented to reflect the pending fence command.
-	 * Once the rendering thread has executed the fence command, it decrements the pending fence count.
-	 * @param bSyncToRHIAndGPU, true if we should wait for the RHI thread or GPU, otherwise we only wait for the render thread.
+	 * Inserts this fence in the rendering pipeline.
+	 * @param SyncDepth, determines which stage of the pipeline will signal the fence.
 	 */
-	RENDERCORE_API void BeginFence(bool bSyncToRHIAndGPU = false);
+	RENDERCORE_API void BeginFence(ESyncDepth SyncDepth = ESyncDepth::RenderThread);
 
 	/**
 	 * Waits for pending fence commands to retire.
diff --git a/Engine/Source/Runtime/SlateNullRenderer/Private/SlateNullRenderer.cpp b/Engine/Source/Runtime/SlateNullRenderer/Private/SlateNullRenderer.cpp
--- a/Engine/Source/Runtime/SlateNullRenderer/Private/SlateNullRenderer.cpp
+++ b/Engine/Source/Runtime/SlateNullRenderer/Private/SlateNullRenderer.cpp
@@ -144,9 +144,7 @@ void FSlateNullRenderer::ClearScenes()
 void FSlateNullRenderer::Sync() const
 {
 #if UE_SLATE_NULL_RENDERER_WITH_ENGINE
-	// Sync game and render thread. Either total sync or allowing one frame lag.
 	static FFrameEndSync FrameEndSync;
-	static auto CVarAllowOneFrameThreadLag = IConsoleManager::Get().FindTConsoleVariableDataInt(TEXT("r.OneFrameThreadLag"));
-	FrameEndSync.Sync(CVarAllowOneFrameThreadLag->GetValueOnAnyThread() != 0);
+	FrameEndSync.Sync();
 #endif
 }
diff --git a/Engine/Source/Runtime/SlateRHIRenderer/Private/SlateRHIRenderer.cpp b/Engine/Source/Runtime/SlateRHIRenderer/Private/SlateRHIRenderer.cpp
--- a/Engine/Source/Runtime/SlateRHIRenderer/Private/SlateRHIRenderer.cpp
+++ b/Engine/Source/Runtime/SlateRHIRenderer/Private/SlateRHIRenderer.cpp
@@ -1445,10 +1445,8 @@ void FSlateRHIRenderer::FlushCommands() const
 
 void FSlateRHIRenderer::Sync() const
 {
-	// Sync game and render thread. Either total sync or allowing one frame lag.
 	static FFrameEndSync FrameEndSync;
-	static auto CVarAllowOneFrameThreadLag = IConsoleManager::Get().FindTConsoleVariableDataInt(TEXT("r.OneFrameThreadLag"));
-	FrameEndSync.Sync(CVarAllowOneFrameThreadLag->GetValueOnAnyThread() != 0);
+	FrameEndSync.Sync();
 }
 
 void FSlateRHIRenderer::BeginFrame() const
