Cleaning up RDG dllexports.

diff --git a/Engine/Source/Runtime/RenderCore/Private/RenderGraphResources.cpp b/Engine/Source/Runtime/RenderCore/Private/RenderGraphResources.cpp
--- a/Engine/Source/Runtime/RenderCore/Private/RenderGraphResources.cpp
+++ b/Engine/Source/Runtime/RenderCore/Private/RenderGraphResources.cpp
@@ -31,7 +31,14 @@ FRDGBuffer::FRDGBuffer(const TCHAR* InName, const FRDGBufferDesc& InDesc, ERDGBu
 	: FRDGViewableResource(InName, ERDGViewableResourceType::Buffer, EnumHasAnyFlags(InFlags, ERDGBufferFlags::SkipTracking), !EnumHasAnyFlags(InFlags, ERDGBufferFlags::ForceImmediateFirstBarrier))
 	, Desc(InDesc)
 	, Flags(InFlags)
-{}
+{
+}
+
+FRDGBuffer::FRDGBuffer(const TCHAR* InName, const FRDGBufferDesc& InDesc, ERDGBufferFlags InFlags, FRDGBufferNumElementsCallback&& InNumElementsCallback)
+	: FRDGBuffer(InName, InDesc, InFlags)
+{
+	NumElementsCallback = MoveTemp(InNumElementsCallback);
+}
 
 FRDGViewableResource::FRDGViewableResource(const TCHAR* InName, const ERDGViewableResourceType InType, bool bSkipTracking, bool bInSplitFirstTransition)
 	: FRDGResource(InName)
diff --git a/Engine/Source/Runtime/RenderCore/Public/RenderGraphBuilder.h b/Engine/Source/Runtime/RenderCore/Public/RenderGraphBuilder.h
--- a/Engine/Source/Runtime/RenderCore/Public/RenderGraphBuilder.h
+++ b/Engine/Source/Runtime/RenderCore/Public/RenderGraphBuilder.h
@@ -83,12 +83,12 @@ public:
 	RENDERCORE_API ~FRDGBuilder();
 
 	/** Finds an RDG texture associated with the external texture, or returns null if none is found. */
-	RENDERCORE_API FRDGTexture* FindExternalTexture(FRHITexture* Texture) const;
-	RENDERCORE_API FRDGTexture* FindExternalTexture(IPooledRenderTarget* ExternalPooledTexture) const;
+	FRDGTexture* FindExternalTexture(FRHITexture* Texture) const;
+	FRDGTexture* FindExternalTexture(IPooledRenderTarget* ExternalPooledTexture) const;
 
 	/** Finds an RDG buffer associated with the external buffer, or returns null if none is found. */
-	RENDERCORE_API FRDGBuffer* FindExternalBuffer(FRHIBuffer* Buffer) const;
-	RENDERCORE_API FRDGBuffer* FindExternalBuffer(FRDGPooledBuffer* ExternalPooledBuffer) const;
+	FRDGBuffer* FindExternalBuffer(FRHIBuffer* Buffer) const;
+	FRDGBuffer* FindExternalBuffer(FRDGPooledBuffer* ExternalPooledBuffer) const;
 
 	/** Registers a external pooled render target texture to be tracked by the render graph. The name of the registered RDG texture is pulled from the pooled render target. */
 	RENDERCORE_API FRDGTextureRef RegisterExternalTexture(
@@ -116,25 +116,25 @@ public:
 	 *  declare the texture in the pass parameter struct. The name is the name used for GPU debugging tools and the the
 	 *  VisualizeTexture/Vis command.
 	 */
-	RENDERCORE_API FRDGTextureRef CreateTexture(const FRDGTextureDesc& Desc, const TCHAR* Name, ERDGTextureFlags Flags = ERDGTextureFlags::None);
+	FRDGTextureRef CreateTexture(const FRDGTextureDesc& Desc, const TCHAR* Name, ERDGTextureFlags Flags = ERDGTextureFlags::None);
 
 	/** Create graph tracked buffer from a descriptor. The CPU memory is guaranteed to be valid through execution of
 	 *  the graph, at which point it is released. The underlying RHI buffer lifetime is only guaranteed for passes which
 	 *  declare the buffer in the pass parameter struct. The name is the name used for GPU debugging tools.
 	 */
-	RENDERCORE_API FRDGBufferRef CreateBuffer(const FRDGBufferDesc& Desc, const TCHAR* Name, ERDGBufferFlags Flags = ERDGBufferFlags::None);
+	FRDGBufferRef CreateBuffer(const FRDGBufferDesc& Desc, const TCHAR* Name, ERDGBufferFlags Flags = ERDGBufferFlags::None);
 
 	/** A variant of CreateBuffer where users supply NumElements through a callback. This allows creating buffers with
 	 *  sizes unknown at creation time. The callback is called before executing the most recent RDG pass that references
 	 *  the buffer so data must be ready before that.
 	 */
-	RENDERCORE_API FRDGBufferRef CreateBuffer(const FRDGBufferDesc& Desc, const TCHAR* Name, FRDGBufferNumElementsCallback&& NumElementsCallback, ERDGBufferFlags Flags = ERDGBufferFlags::None);
+	FRDGBufferRef CreateBuffer(const FRDGBufferDesc& Desc, const TCHAR* Name, FRDGBufferNumElementsCallback&& NumElementsCallback, ERDGBufferFlags Flags = ERDGBufferFlags::None);
 
 	/** Create graph tracked SRV for a texture from a descriptor. */
-	RENDERCORE_API FRDGTextureSRVRef CreateSRV(const FRDGTextureSRVDesc& Desc);
+	FRDGTextureSRVRef CreateSRV(const FRDGTextureSRVDesc& Desc);
 
 	/** Create graph tracked SRV for a buffer from a descriptor. */
-	RENDERCORE_API FRDGBufferSRVRef CreateSRV(const FRDGBufferSRVDesc& Desc);
+	FRDGBufferSRVRef CreateSRV(const FRDGBufferSRVDesc& Desc);
 
 	FORCEINLINE FRDGBufferSRVRef CreateSRV(FRDGBufferRef Buffer, EPixelFormat Format)
 	{
@@ -142,7 +142,7 @@ public:
 	}
 
 	/** Create graph tracked UAV for a texture from a descriptor. */
-	RENDERCORE_API FRDGTextureUAVRef CreateUAV(const FRDGTextureUAVDesc& Desc, ERDGUnorderedAccessViewFlags Flags = ERDGUnorderedAccessViewFlags::None);
+	FRDGTextureUAVRef CreateUAV(const FRDGTextureUAVDesc& Desc, ERDGUnorderedAccessViewFlags Flags = ERDGUnorderedAccessViewFlags::None);
 
 	FORCEINLINE FRDGTextureUAVRef CreateUAV(FRDGTextureRef Texture, ERDGUnorderedAccessViewFlags Flags = ERDGUnorderedAccessViewFlags::None, EPixelFormat Format = PF_Unknown)
 	{
@@ -150,7 +150,7 @@ public:
 	}
 
 	/** Create graph tracked UAV for a buffer from a descriptor. */
-	RENDERCORE_API FRDGBufferUAVRef CreateUAV(const FRDGBufferUAVDesc& Desc, ERDGUnorderedAccessViewFlags Flags = ERDGUnorderedAccessViewFlags::None);
+	FRDGBufferUAVRef CreateUAV(const FRDGBufferUAVDesc& Desc, ERDGUnorderedAccessViewFlags Flags = ERDGUnorderedAccessViewFlags::None);
 
 	FORCEINLINE FRDGBufferUAVRef CreateUAV(FRDGBufferRef Buffer, EPixelFormat Format, ERDGUnorderedAccessViewFlags Flags = ERDGUnorderedAccessViewFlags::None)
 	{
@@ -168,7 +168,7 @@ public:
 	// Allocation Methods
 
 	/** Allocates raw memory using an allocator tied to the lifetime of the graph. */
-	RENDERCORE_API void* Alloc(uint64 SizeInBytes, uint32 AlignInBytes = 16);
+	void* Alloc(uint64 SizeInBytes, uint32 AlignInBytes = 16);
 
 	/** Allocates POD memory using an allocator tied to the lifetime of the graph. Does not construct / destruct. */
 	template <typename PODType>
@@ -258,7 +258,7 @@ public:
 	 *  Recommended usage is to set a workload equal to the number of complex draw / dispatch calls (each with its own parameters, etc), and
 	 *  only as a performance tweak if a particular pass is very expensive relative to other passes.
 	 */
-	RENDERCORE_API void SetPassWorkload(FRDGPass* Pass, uint32 Workload);
+	void SetPassWorkload(FRDGPass* Pass, uint32 Workload);
 
 	/** Adds a user-defined dependency between two passes. This can be used to fine-tune async compute overlap by forcing a sync point. */
 	RENDERCORE_API void AddPassDependency(FRDGPass* Producer, FRDGPass* Consumer);
@@ -351,15 +351,15 @@ public:
 	 *  the lifetime of the GPU resource until execution, at which point the pointer is filled. If specified, the texture is transitioned
 	 *  to the AccessFinal state, or kDefaultAccessFinal otherwise.
 	 */
-	RENDERCORE_API void QueueTextureExtraction(FRDGTextureRef Texture, TRefCountPtr<IPooledRenderTarget>* OutPooledTexturePtr, ERDGResourceExtractionFlags Flags = ERDGResourceExtractionFlags::None);
-	RENDERCORE_API void QueueTextureExtraction(FRDGTextureRef Texture, TRefCountPtr<IPooledRenderTarget>* OutPooledTexturePtr, ERHIAccess AccessFinal, ERDGResourceExtractionFlags Flags = ERDGResourceExtractionFlags::None);
+	void QueueTextureExtraction(FRDGTextureRef Texture, TRefCountPtr<IPooledRenderTarget>* OutPooledTexturePtr, ERDGResourceExtractionFlags Flags = ERDGResourceExtractionFlags::None);
+	void QueueTextureExtraction(FRDGTextureRef Texture, TRefCountPtr<IPooledRenderTarget>* OutPooledTexturePtr, ERHIAccess AccessFinal, ERDGResourceExtractionFlags Flags = ERDGResourceExtractionFlags::None);
 
 	/** Queues a pooled buffer extraction to happen at the end of graph execution. For graph-created buffers, this extends the lifetime
 	 *  of the GPU resource until execution, at which point the pointer is filled. If specified, the buffer is transitioned to the
 	 *  AccessFinal state, or kDefaultAccessFinal otherwise.
 	 */
-	RENDERCORE_API void QueueBufferExtraction(FRDGBufferRef Buffer, TRefCountPtr<FRDGPooledBuffer>* OutPooledBufferPtr);
-	RENDERCORE_API void QueueBufferExtraction(FRDGBufferRef Buffer, TRefCountPtr<FRDGPooledBuffer>* OutPooledBufferPtr, ERHIAccess AccessFinal);
+	void QueueBufferExtraction(FRDGBufferRef Buffer, TRefCountPtr<FRDGPooledBuffer>* OutPooledBufferPtr);
+	void QueueBufferExtraction(FRDGBufferRef Buffer, TRefCountPtr<FRDGPooledBuffer>* OutPooledBufferPtr, ERHIAccess AccessFinal);
 
 	/** For graph-created resources, this forces immediate allocation of the underlying pooled resource, effectively promoting it
 	 *  to an external resource. This will increase memory pressure, but allows for querying the pooled resource with GetPooled{Texture, Buffer}.
@@ -375,14 +375,14 @@ public:
 	RENDERCORE_API FRHIUniformBuffer* ConvertToExternalUniformBuffer(FRDGUniformBufferRef UniformBuffer);
 
 	/** Performs an immediate query for the underlying pooled resource. This is only allowed for external or extracted resources. */
-	RENDERCORE_API const TRefCountPtr<IPooledRenderTarget>& GetPooledTexture(FRDGTextureRef Texture) const;
-	RENDERCORE_API const TRefCountPtr<FRDGPooledBuffer>& GetPooledBuffer(FRDGBufferRef Buffer) const;
+	const TRefCountPtr<IPooledRenderTarget>& GetPooledTexture(FRDGTextureRef Texture) const;
+	const TRefCountPtr<FRDGPooledBuffer>& GetPooledBuffer(FRDGBufferRef Buffer) const;
 
 	/** (External | Extracted only) Sets the access to transition to after execution at the end of the graph. Overwrites any previously set final access. */
-	RENDERCORE_API void SetTextureAccessFinal(FRDGTextureRef Texture, ERHIAccess Access);
+	void SetTextureAccessFinal(FRDGTextureRef Texture, ERHIAccess Access);
 
 	/** (External | Extracted only) Sets the access to transition to after execution at the end of the graph. Overwrites any previously set final access. */
-	RENDERCORE_API void SetBufferAccessFinal(FRDGBufferRef Buffer, ERHIAccess Access);
+	void SetBufferAccessFinal(FRDGBufferRef Buffer, ERHIAccess Access);
 
 	/** Configures the resource for external access for all subsequent passes, or until UseInternalAccessMode is called.
 	 *  Only read-only access states are allowed. When in external access mode, it is safe to access the underlying RHI
@@ -415,8 +415,8 @@ public:
 	}
 
 	/** Flag a resource that is produced by a pass but never used or extracted to not emit an 'unused' warning. */
-	RENDERCORE_API void RemoveUnusedTextureWarning(FRDGTextureRef Texture);
-	RENDERCORE_API void RemoveUnusedBufferWarning(FRDGBufferRef Buffer);
+	void RemoveUnusedTextureWarning(FRDGTextureRef Texture);
+	void RemoveUnusedBufferWarning(FRDGBufferRef Buffer);
 
 	/** Flushes all queued passes to an async task to perform setup work. */
 	RENDERCORE_API void FlushSetupQueue();
diff --git a/Engine/Source/Runtime/RenderCore/Public/RenderGraphResources.h b/Engine/Source/Runtime/RenderCore/Public/RenderGraphResources.h
--- a/Engine/Source/Runtime/RenderCore/Public/RenderGraphResources.h
+++ b/Engine/Source/Runtime/RenderCore/Public/RenderGraphResources.h
@@ -1354,15 +1354,10 @@ public:
 
 private:
 	RENDERCORE_API FRDGBuffer(const TCHAR* InName, const FRDGBufferDesc& InDesc, ERDGBufferFlags InFlags);
-
-	FRDGBuffer(const TCHAR* InName, const FRDGBufferDesc& InDesc, ERDGBufferFlags InFlags, FRDGBufferNumElementsCallback&& InNumElementsCallback)
-		: FRDGBuffer(InName, InDesc, InFlags)
-	{
-		NumElementsCallback = MoveTemp(InNumElementsCallback);
-	}
+	RENDERCORE_API FRDGBuffer(const TCHAR* InName, const FRDGBufferDesc& InDesc, ERDGBufferFlags InFlags, FRDGBufferNumElementsCallback&& InNumElementsCallback);
 
 	/** Finalizes any pending field of the buffer descriptor. */
-	void FinalizeDesc();
+	RENDERCORE_API void FinalizeDesc();
 
 	FRHIBuffer* GetRHIUnchecked() const
 	{
