Fixed RDG immediate not properly setting reserved buffer commit sizes.

diff --git a/Engine/Source/Runtime/RenderCore/Public/RenderGraphBuilder.inl b/Engine/Source/Runtime/RenderCore/Public/RenderGraphBuilder.inl
--- a/Engine/Source/Runtime/RenderCore/Public/RenderGraphBuilder.inl
+++ b/Engine/Source/Runtime/RenderCore/Public/RenderGraphBuilder.inl
@@ -418,8 +418,14 @@ inline void FRDGBuilder::QueueBufferUpload(FRDGBufferRef Buffer, FRDGBufferIniti
 inline void FRDGBuilder::QueueCommitReservedBuffer(FRDGBufferRef Buffer, uint64 CommitSizeInBytes)
 {
 	IF_RDG_ENABLE_DEBUG(UserValidation.ValidateCommitBuffer(Buffer, CommitSizeInBytes));
-
-	AsyncSetupQueue.Push(FAsyncSetupOp::ReservedBufferCommit(Buffer, CommitSizeInBytes));
+	if (IsImmediateMode())
+	{
+		Buffer->PendingCommitSize = CommitSizeInBytes;
+	}
+	else
+	{
+		AsyncSetupQueue.Push(FAsyncSetupOp::ReservedBufferCommit(Buffer, CommitSizeInBytes));
+	}
 	Buffer->PooledBuffer->SetCommittedSize(CommitSizeInBytes);
 }
 
