Disable async compute when using tiled GPUs. Render pass merging is not compatible with it due to the way transitions are moved.

diff --git a/Engine/Source/Runtime/RenderCore/Private/RenderGraphPrivate.h b/Engine/Source/Runtime/RenderCore/Private/RenderGraphPrivate.h
--- a/Engine/Source/Runtime/RenderCore/Private/RenderGraphPrivate.h
+++ b/Engine/Source/Runtime/RenderCore/Private/RenderGraphPrivate.h
@@ -227,12 +227,14 @@ FORCEINLINE bool IsImmediateMode()
 
 FORCEINLINE bool IsRenderPassMergeEnabled()
 {
-	return GRDGMergeRenderPasses != 0 && !IsImmediateMode() && RHIHasTiledGPU(GMaxRHIShaderPlatform) && !GRDGDebugFlushGPU;
+	return GRDGMergeRenderPasses != 0 && !IsImmediateMode() && FDataDrivenShaderPlatformInfo::IsValid(GMaxRHIShaderPlatform) && RHIHasTiledGPU(GMaxRHIShaderPlatform) && !GRDGDebugFlushGPU;
 }
 
 FORCEINLINE bool IsAsyncComputeSupported()
 {
-	return GRDGAsyncCompute > 0 && !IsImmediateMode() && GSupportsEfficientAsyncCompute && GRHISupportsSeparateDepthStencilCopyAccess && !GTriggerGPUProfile;
+	// Render pass merging and async compute are mutually exclusive since they are mutually dependent on each other (i.e. we can't
+	// merge render passes that involve async compute fences, since transitions would be relocated to the ends of the merged region).
+	return GRDGAsyncCompute > 0 && !IsImmediateMode() && !IsRenderPassMergeEnabled() && GSupportsEfficientAsyncCompute && GRHISupportsSeparateDepthStencilCopyAccess && !GTriggerGPUProfile;
 }
 
 extern bool IsParallelExecuteEnabled();
