Workaround for RHI validation error due to prologue pass issuing a fence along with all its transitions.

diff --git a/Engine/Source/Runtime/RenderCore/Private/RenderGraphBuilder.cpp b/Engine/Source/Runtime/RenderCore/Private/RenderGraphBuilder.cpp
--- a/Engine/Source/Runtime/RenderCore/Private/RenderGraphBuilder.cpp
+++ b/Engine/Source/Runtime/RenderCore/Private/RenderGraphBuilder.cpp
@@ -1570,8 +1570,12 @@ void FRDGBuilder::Compile()
 
 				FRDGBarrierBatchBegin& EpilogueBarriersToBeginForAsyncCompute = GraphicsForkPass->GetEpilogueBarriersToBeginForAsyncCompute(Allocators.Transition, TransitionCreateQueue);
 
+				// Workaround for RHI validation. The prologue pass issues its own separate transition for the prologue pass
+				// so that external access resources left in the all pipes state can be transitioned back to graphics.
+				const bool bSeparateTransitionNeeded = GraphicsForkPass == ProloguePass;
+
 				GraphicsForkPass->bGraphicsFork = 1;
-				EpilogueBarriersToBeginForAsyncCompute.SetUseCrossPipelineFence();
+				EpilogueBarriersToBeginForAsyncCompute.SetUseCrossPipelineFence(bSeparateTransitionNeeded);
 
 				AsyncComputePass->bAsyncComputeBegin = 1;
 				AsyncComputePass->GetPrologueBarriersToEnd(Allocators.Transition).AddDependency(&EpilogueBarriersToBeginForAsyncCompute);
@@ -1622,8 +1626,10 @@ void FRDGBuilder::Compile()
 
 				FRDGBarrierBatchBegin& EpilogueBarriersToBeginForGraphics = AsyncComputePass->GetEpilogueBarriersToBeginForGraphics(Allocators.Transition, TransitionCreateQueue);
 
+				const bool bSeparateTransitionNeeded = false;
+
 				AsyncComputePass->bAsyncComputeEnd = 1;
-				EpilogueBarriersToBeginForGraphics.SetUseCrossPipelineFence();
+				EpilogueBarriersToBeginForGraphics.SetUseCrossPipelineFence(bSeparateTransitionNeeded);
 
 				GraphicsJoinPass->bGraphicsJoin = 1;
 				GraphicsJoinPass->GetPrologueBarriersToEnd(Allocators.Transition).AddDependency(&EpilogueBarriersToBeginForGraphics);
diff --git a/Engine/Source/Runtime/RenderCore/Private/RenderGraphPass.cpp b/Engine/Source/Runtime/RenderCore/Private/RenderGraphPass.cpp
--- a/Engine/Source/Runtime/RenderCore/Private/RenderGraphPass.cpp
+++ b/Engine/Source/Runtime/RenderCore/Private/RenderGraphPass.cpp
@@ -222,10 +222,20 @@ void FRDGBarrierBatchBegin::CreateTransition(TConstArrayView<FRHITransitionInfo>
 {
 	check(bTransitionNeeded && !Transition);
 	Transition = RHICreateTransition(FRHITransitionCreateInfo(PipelinesToBegin, PipelinesToEnd, TransitionFlags, TransitionsRHI, Aliases));
+
+	if (bSeparateFenceTransitionNeeded)
+	{
+		SeparateFenceTransition = RHICreateTransition(FRHITransitionCreateInfo(PipelinesToBegin, PipelinesToEnd));
+	}
 }
 
 void FRDGBarrierBatchBegin::Submit(FRHIComputeCommandList& RHICmdList, ERHIPipeline Pipeline, FRDGTransitionQueue& TransitionsToBegin)
 {
+	if (SeparateFenceTransition)
+	{
+		TransitionsToBegin.Emplace(SeparateFenceTransition);
+	}
+
 	if (Transition)
 	{
 		TransitionsToBegin.Emplace(Transition);
@@ -304,6 +314,11 @@ void FRDGBarrierBatchEnd::Submit(FRHIComputeCommandList& RHICmdList, ERHIPipelin
 	{
 		if (Dependent->BarriersToEnd[Pipeline] == Id)
 		{
+			if (Dependent->SeparateFenceTransition)
+			{
+				Transitions.Emplace(Dependent->SeparateFenceTransition);
+			}
+
 			Transitions.Emplace(Dependent->Transition);
 		}
 	}
diff --git a/Engine/Source/Runtime/RenderCore/Public/RenderGraphPass.h b/Engine/Source/Runtime/RenderCore/Public/RenderGraphPass.h
--- a/Engine/Source/Runtime/RenderCore/Public/RenderGraphPass.h
+++ b/Engine/Source/Runtime/RenderCore/Public/RenderGraphPass.h
@@ -129,9 +129,16 @@ public:
 
 	RENDERCORE_API void AddAlias(FRDGViewableResource* Resource, const FRHITransientAliasingInfo& Info);
 
-	void SetUseCrossPipelineFence()
+	void SetUseCrossPipelineFence(bool bUseSeparateTransition)
 	{
-		TransitionFlags = ERHITransitionCreateFlags::None;
+		if (bUseSeparateTransition)
+		{
+			bSeparateFenceTransitionNeeded = true;
+		}
+		else
+		{
+			TransitionFlags = ERHITransitionCreateFlags::None;
+		}
 		bTransitionNeeded = true;
 	}
 
@@ -153,13 +160,15 @@ public:
 
 private:
 	const FRHITransition* Transition = nullptr;
+	const FRHITransition* SeparateFenceTransition = nullptr;
+	TRHIPipelineArray<FRDGBarrierBatchEndId> BarriersToEnd;
 	TArray<FRDGTransitionInfo, FRDGArrayAllocator> Transitions;
 	TArray<FRHITransientAliasingInfo, FRDGArrayAllocator> Aliases;
 	ERHITransitionCreateFlags TransitionFlags = ERHITransitionCreateFlags::NoFence;
 	ERHIPipeline PipelinesToBegin;
 	ERHIPipeline PipelinesToEnd;
-	TRHIPipelineArray<FRDGBarrierBatchEndId> BarriersToEnd;
 	bool bTransitionNeeded = false;
+	bool bSeparateFenceTransitionNeeded = false;
 
 #if RDG_ENABLE_DEBUG
 	FRDGPassesByPipeline DebugPasses;
