Removing a check that was crashing the editor.

diff --git a/Engine/Source/Runtime/RenderCore/Private/ShaderCodeArchive.cpp b/Engine/Source/Runtime/RenderCore/Private/ShaderCodeArchive.cpp
--- a/Engine/Source/Runtime/RenderCore/Private/ShaderCodeArchive.cpp
+++ b/Engine/Source/Runtime/RenderCore/Private/ShaderCodeArchive.cpp
@@ -2219,8 +2219,6 @@ void FIoStoreShaderCodeArchive::ReleasePreloadEntry(int32 ShaderGroupIndex
 			uint32 ShaderGroupLoadBytes = 0;
 			if (!PreloadEntry->bNeverToBePreloaded)
 			{
-				// make sure we don't leak any reference on the shader group
-				check(PreloadEntry->PreloadEvent.GetRefCount() <= 1)
 				PreloadEntry->IoRequest.Cancel();
 				PreloadEntry->IoRequest = FIoRequest();
 				PreloadEntry->PreloadEvent.SafeRelease();
