Wrap the preload shader group priority around a cvar.

diff --git a/Engine/Source/Runtime/RenderCore/Private/ShaderCodeArchive.cpp b/Engine/Source/Runtime/RenderCore/Private/ShaderCodeArchive.cpp
--- a/Engine/Source/Runtime/RenderCore/Private/ShaderCodeArchive.cpp
+++ b/Engine/Source/Runtime/RenderCore/Private/ShaderCodeArchive.cpp
@@ -77,6 +77,36 @@ static FAutoConsoleVariableRef CVarShaderCodeLibraryMaxShaderPreloadWaitTime(
 	ECVF_RenderThreadSafe | ECVF_ReadOnly
 );
 
+int32 GPreloadShaderPriority = 2;
+static FAutoConsoleVariableRef CVarPreloadShaderPriority(
+	TEXT("r.PreloadShaderPriority"),
+	GPreloadShaderPriority,
+	TEXT("Change PreloadShaderGroup I/O priority.\n")
+	TEXT("0-Min\n")
+	TEXT("1-Low\n")
+	TEXT("2-Medium (Default)\n")
+	TEXT("3-High\n")
+	TEXT("4-Max\n"),
+	ECVF_Default
+);
+
+int32 GetShaderCodeArchivePriority()
+{
+	switch (GPreloadShaderPriority)
+	{
+		case 4:
+			return IoDispatcherPriority_Max;
+		case 3:
+			return IoDispatcherPriority_High;
+		case 1:
+			return IoDispatcherPriority_Low;
+		case 0:
+			return IoDispatcherPriority_Min;
+		default:
+			return IoDispatcherPriority_Medium;
+	}
+}
+
 #if RHI_RAYTRACING	// this function is only needed to check if we need to avoid excluding raytracing shaders
 namespace
 {
@@ -2002,7 +2032,7 @@ bool FIoStoreShaderCodeArchive::PreloadShaderGroup(int32 ShaderGroupIndex, FGrap
 		if (UNLIKELY(AttachShaderReadRequestFuncPtr == nullptr))
 		{
 			FIoBatch IoBatch = IoDispatcher.NewBatch();
-			PreloadEntry.IoRequest = IoBatch.Read(Header.ShaderGroupIoHashes[ShaderGroupIndex], FIoReadOptions(), IoDispatcherPriority_Medium);
+			PreloadEntry.IoRequest = IoBatch.Read(Header.ShaderGroupIoHashes[ShaderGroupIndex], FIoReadOptions(), GetShaderCodeArchivePriority());
 			IoBatch.IssueAndDispatchSubsequents(PreloadEntry.PreloadEvent);
 		}
 		else
