[shaders] fix metal shader archive population, was still implementing a deprecated API

diff --git a/Engine/Source/Developer/Apple/MetalShaderFormat/Private/MetalShaderFormat.cpp b/Engine/Source/Developer/Apple/MetalShaderFormat/Private/MetalShaderFormat.cpp
--- a/Engine/Source/Developer/Apple/MetalShaderFormat/Private/MetalShaderFormat.cpp
+++ b/Engine/Source/Developer/Apple/MetalShaderFormat/Private/MetalShaderFormat.cpp
@@ -175,7 +175,7 @@ public:
 		const FString& OutputDir,
 		const FString& DebugOutputDir,
 		const FSerializedShaderArchive& InSerializedShaders,
-		const TArray<TArray<uint8>>& ShaderCode,
+		const TArray<FSharedBuffer>& ShaderCode,
 		TArray<FString>* OutputFiles) const override final
 	{
 		int32 NumShadersPerLibrary = 10000;
diff --git a/Engine/Source/Runtime/RenderCore/Private/ShaderCodeArchive.cpp b/Engine/Source/Runtime/RenderCore/Private/ShaderCodeArchive.cpp
--- a/Engine/Source/Runtime/RenderCore/Private/ShaderCodeArchive.cpp
+++ b/Engine/Source/Runtime/RenderCore/Private/ShaderCodeArchive.cpp
@@ -467,7 +467,7 @@ bool ShaderCodeArchive::CompressShaderWithOodle(uint8* OutCompressedShader, int6
 	}	
 }
 
-void FSerializedShaderArchive::DecompressShader(int32 Index, const TArray<TArray<uint8>>& ShaderCode, TArray<uint8>& OutDecompressedShader) const
+void FSerializedShaderArchive::DecompressShader(int32 Index, const TArray<FSharedBuffer>& ShaderCode, TArray<uint8>& OutDecompressedShader) const
 {
 	const FShaderCodeEntry& Entry = ShaderEntries[Index];
 	OutDecompressedShader.SetNum(Entry.UncompressedSize, EAllowShrinking::No);
@@ -477,7 +477,7 @@ void FSerializedShaderArchive::DecompressShader(int32 Index, const TArray<TArray
 	}
 	else
 	{
-		ShaderCodeArchive::DecompressShaderWithOodle(OutDecompressedShader.GetData(), Entry.UncompressedSize, ShaderCode[Index].GetData(), Entry.Size);
+		ShaderCodeArchive::DecompressShaderWithOodle(OutDecompressedShader.GetData(), Entry.UncompressedSize, reinterpret_cast<const uint8*>(ShaderCode[Index].GetData()), Entry.Size);
 	}
 }
 
diff --git a/Engine/Source/Runtime/RenderCore/Public/ShaderCodeArchive.h b/Engine/Source/Runtime/RenderCore/Public/ShaderCodeArchive.h
--- a/Engine/Source/Runtime/RenderCore/Public/ShaderCodeArchive.h
+++ b/Engine/Source/Runtime/RenderCore/Public/ShaderCodeArchive.h
@@ -256,7 +256,12 @@ public:
 	RENDERCORE_API void Finalize();
 #endif
 
-	RENDERCORE_API void DecompressShader(int32 Index, const TArray<TArray<uint8>>& ShaderCode, TArray<uint8>& OutDecompressedShader) const;
+	UE_DEPRECATED(5.5, "Use overload accepting a TArray<FSharedBuffer> for shader code")
+	RENDERCORE_API void DecompressShader(int32 Index, const TArray<TArray<uint8>>& ShaderCode, TArray<uint8>& OutDecompressedShader) const
+	{
+	}
+
+	RENDERCORE_API void DecompressShader(int32 Index, const TArray<FSharedBuffer>& ShaderCode, TArray<uint8>& OutDecompressedShader) const;
 
 	RENDERCORE_API void Serialize(FArchive& Ar);
 #if WITH_EDITOR
