Improved game thread frame sync to resolve crashes introduced in 36468180 and worked around in 36479070
 - Various legacy game thread code assumes the render thread will never be more than one frame behind. The original change in 36468180 switched from syncing the GT with the RT, to syncing the GT with the RHIT. That left the render thread "floating" in the center of the pipeline, and led to cases where resources are deleted too soon.
 - New approach is to always sync with the GT with the N-1 RT frame, so the GT is never too far ahead of the RT. This maintains compatibility with the legacy GT code paths. In addition to the GT->RT sync, we also sync with the RHIT to prevent the engine running ahead, which was the original bug that 36468180 was fixing.
 - "r.GTSyncType" mode 0 now allows for 1 frame of GT->RT overlap, and 2 frames of GT->RHIT overlap.
 - For debugging purposes, "r.GTSyncType" can also be made negative, which increases the number of GT->RHIT overlap frames, e.g. "r.GTSyncType -3" gives 5 frames of GT->RHIT overlap. While this is not overly useful in a shipped title, it can be used to prove the correctness of the rendering pipeline.
 - Merged the FDeferredCleanupInterface / FPendingCleanupObjects processing into the FFrameEndSync code path, plus made FFrameEndSync a static singleton. This allows us to manage the N frames of overlap between the GT and RHIT in a central place, and correctly cleanup deferred resources when the RT fences have passed.

In future, we will need to revisit this as part of the frame pacing initiative. Game thread code should be written to not require render thread fences for correctness / threadsafety.

diff --git a/Engine/Source/Editor/UnrealEd/Private/Commandlets/CompileShadersTestBedCommandlet.cpp b/Engine/Source/Editor/UnrealEd/Private/Commandlets/CompileShadersTestBedCommandlet.cpp
--- a/Engine/Source/Editor/UnrealEd/Private/Commandlets/CompileShadersTestBedCommandlet.cpp
+++ b/Engine/Source/Editor/UnrealEd/Private/Commandlets/CompileShadersTestBedCommandlet.cpp
@@ -231,7 +231,7 @@ int32 UCompileShadersTestBedCommandlet::Main(const FString& Params)
 						}
 
 						// Flush rendering commands to release any RHI resources (shaders and shader maps).
-						// Delete any FPendingCleanupObjects (shader maps).
+						// Delete any FDeferredCleanupInterface (shader maps).
 						FlushRenderingCommands();
 
 						if (CurrentOutstandingJobs < MaxOutstandingJobs)
@@ -264,7 +264,7 @@ int32 UCompileShadersTestBedCommandlet::Main(const FString& Params)
 					}
 
 					// Flush rendering commands to release any RHI resources (shaders and shader maps).
-					// Delete any FPendingCleanupObjects (shader maps).
+					// Delete any FDeferredCleanupInterface (shader maps).
 					FlushRenderingCommands();
 
 					if (CurrentOutstandingJobs < MaxOutstandingJobs)
diff --git a/Engine/Source/Editor/UnrealEd/Private/Commandlets/DumpLightFunctionMaterialInfo.cpp b/Engine/Source/Editor/UnrealEd/Private/Commandlets/DumpLightFunctionMaterialInfo.cpp
--- a/Engine/Source/Editor/UnrealEd/Private/Commandlets/DumpLightFunctionMaterialInfo.cpp
+++ b/Engine/Source/Editor/UnrealEd/Private/Commandlets/DumpLightFunctionMaterialInfo.cpp
@@ -122,7 +122,7 @@ int32 UDumpLightFunctionMaterialInfoCommandlet::Main(const FString& Params)
 						}
 
 						// Flush rendering commands to release any RHI resources (shaders and shader maps).
-						// Delete any FPendingCleanupObjects (shader maps).
+						// Delete any FDeferredCleanupInterface (shader maps).
 						FlushRenderingCommands();
 
 						if (CurrentOutstandingJobs < MaxOutstandingJobs)
@@ -157,7 +157,7 @@ int32 UDumpLightFunctionMaterialInfoCommandlet::Main(const FString& Params)
 					}
 
 					// Flush rendering commands to release any RHI resources (shaders and shader maps).
-					// Delete any FPendingCleanupObjects (shader maps).
+					// Delete any FDeferredCleanupInterface (shader maps).
 					FlushRenderingCommands();
 
 					if (CurrentOutstandingJobs < MaxOutstandingJobs)
diff --git a/Engine/Source/Editor/UnrealEd/Private/Commandlets/DumpMaterialInfo.cpp b/Engine/Source/Editor/UnrealEd/Private/Commandlets/DumpMaterialInfo.cpp
--- a/Engine/Source/Editor/UnrealEd/Private/Commandlets/DumpMaterialInfo.cpp
+++ b/Engine/Source/Editor/UnrealEd/Private/Commandlets/DumpMaterialInfo.cpp
@@ -664,7 +664,7 @@ static void DumpMaterials(
 					}
 
 					// Flush rendering commands to release any RHI resources (shaders and shader maps).
-					// Delete any FPendingCleanupObjects (shader maps).
+					// Delete any FDeferredCleanupInterface (shader maps).
 					FlushRenderingCommands();
 
 					if (CurrentOutstandingJobs < MaxOutstandingJobs)
@@ -697,7 +697,7 @@ static void DumpMaterials(
 				}
 
 				// Flush rendering commands to release any RHI resources (shaders and shader maps).
-				// Delete any FPendingCleanupObjects (shader maps).
+				// Delete any FDeferredCleanupInterface (shader maps).
 				FlushRenderingCommands();
 
 				if (CurrentOutstandingJobs < MaxOutstandingJobs)
diff --git a/Engine/Source/Editor/UnrealEd/Private/CookOnTheFlyServer.cpp b/Engine/Source/Editor/UnrealEd/Private/CookOnTheFlyServer.cpp
--- a/Engine/Source/Editor/UnrealEd/Private/CookOnTheFlyServer.cpp
+++ b/Engine/Source/Editor/UnrealEd/Private/CookOnTheFlyServer.cpp
@@ -2119,7 +2119,7 @@ void UCookOnTheFlyServer::PollFlushRenderingCommands()
 	UE_SCOPED_COOKTIMER_AND_DURATION(CookByTheBook_TickCommandletStats, DetailedCookStats::TickLoopFlushRenderingCommandsTimeSec);
 
 	// Flush rendering commands to release any RHI resources (shaders and shader maps).
-	// Delete any FPendingCleanupObjects (shader maps).
+	// Delete any FDeferredCleanupInterface (shader maps).
 	FlushRenderingCommands();
 }
 
diff --git a/Engine/Source/Editor/UnrealEd/Private/FeedbackContextEditor.cpp b/Engine/Source/Editor/UnrealEd/Private/FeedbackContextEditor.cpp
--- a/Engine/Source/Editor/UnrealEd/Private/FeedbackContextEditor.cpp
+++ b/Engine/Source/Editor/UnrealEd/Private/FeedbackContextEditor.cpp
@@ -564,15 +564,8 @@ void FFeedbackContextEditor::ProgressReported( const float TotalProgressInterp,
 	static double LastTimePendingCleanupObjectsWhereDeleted;
 	if( FPlatformTime::Seconds() - LastTimePendingCleanupObjectsWhereDeleted > 1 )
 	{
-		// Get list of objects that are pending cleanup.
-		FPendingCleanupObjects* PendingCleanupObjects = GetPendingCleanupObjects();
-		if (!PendingCleanupObjects->IsEmpty())
-		{
-			// Flush rendering commands in the queue.
-			FlushRenderingCommands();
-		}
-		// It is now safe to delete the pending clean objects.
-		delete PendingCleanupObjects;
+		FFrameEndSync::Sync();
+
 		// Keep track of time this operation was performed so we don't do it too often.
 		LastTimePendingCleanupObjectsWhereDeleted = FPlatformTime::Seconds();
 	}
diff --git a/Engine/Source/Runtime/Engine/Private/UnrealEngine.cpp b/Engine/Source/Runtime/Engine/Private/UnrealEngine.cpp
--- a/Engine/Source/Runtime/Engine/Private/UnrealEngine.cpp
+++ b/Engine/Source/Runtime/Engine/Private/UnrealEngine.cpp
@@ -449,23 +449,6 @@ FSystemResolution GSystemResolution;
 
 ENGINE_API int32 GUnbuiltHLODCount = 0;
 
-TAutoConsoleVariable<int32> CVarAllowOneFrameThreadLag(
-	TEXT("r.OneFrameThreadLag"),
-	1,
-	TEXT("Whether to allow the rendering thread to lag one frame behind the game thread (0: disabled, otherwise enabled)")
-);
-
-TAutoConsoleVariable<int32> CVarGTSyncType(
-	TEXT("r.GTSyncType"),
-	0,
-	TEXT("Determines how the game thread syncs with the render thread, RHI thread and GPU.\n")
-	TEXT("Syncing to the GPU swap chain flip allows for lower frame latency.\n")
-	TEXT(" 0 - Sync the game thread with the N-2 RHI thread frame when the RHI thread is enabled, or the N-1 render thread frame when it is disabled (default).\n")
-	TEXT(" 1 - Sync the game thread with the N-1 RHI thread frame.\n")
-	TEXT(" 2 - Sync the game thread with the GPU swap chain flip (only on supported platforms).\n"),
-	ECVF_Default
-);
-
 static FAutoConsoleVariable CVarSystemResolution(
 	TEXT("r.SetRes"),
 	TEXT("1280x720w"),
@@ -747,8 +730,10 @@ void CalculateFPSTimings()
 	double CurrentTime = FPlatformTime::Seconds();
 	float FrameTimeMS = (float)((CurrentTime - LastTime) * 1000.0);
 
+	static auto CVarGTSyncType = IConsoleManager::Get().FindConsoleVariable(TEXT("r.GTSyncType"));
 	static auto CVarVsync = IConsoleManager::Get().FindConsoleVariable(TEXT("r.VSync"));
-	if (CVarGTSyncType.GetValueOnGameThread() == 2 && CVarVsync->GetInt() != 0)
+
+	if (CVarGTSyncType->GetInt() == 2 && CVarVsync->GetInt() != 0)
 	{
 		float RHIFrameTime = RHIGetFrameTime();
 		if (RHIFrameTime != 0)
@@ -13272,7 +13257,6 @@ DEFINE_STAT(STAT_RedrawViewports);
 DEFINE_STAT(STAT_UpdateLevelStreaming);
 DEFINE_STAT(STAT_RHITickTime);
 DEFINE_STAT(STAT_IntentionalHitch);
-DEFINE_STAT(STAT_FrameSyncTime);
 DEFINE_STAT(STAT_DeferredTickTime);
 
 /** Input stat */
@@ -13306,115 +13290,6 @@ UFont* GetStatsFont()
 	return GEngine->GetSmallFont();
 }
 
-
-FFrameEndSync::FFrameEndSync()
-{
-	// FFrameEndSync instances are often used as static local vars. we need to cleanup them on engine exit to avoid static destruciton order problem
-	CleanupDelegate = FCoreDelegates::OnEnginePreExit.AddRaw(this, &FFrameEndSync::Cleanup);
-}
-
-FFrameEndSync::~FFrameEndSync()
-{
-	// if it's destroyed before the engine exit, remove the delegate so it doesn't use the instance after destruction
-	if (CleanupDelegate.IsValid())
-	{
-		FCoreDelegates::OnEnginePreExit.Remove(CleanupDelegate);
-	}
-}
-
-void FFrameEndSync::Cleanup()
-{
-	// forcing static deallocation order:
-	// fences can hold task completion handles that need to be freed before their allocator is destroyed. wating for fences does the job
-	for (FRenderCommandFence& Fence : Fences)
-	{
-		Fence.Wait();
-	}
-
-	FCoreDelegates::OnEnginePreExit.Remove(CleanupDelegate);
-	// notify the destructor
-	CleanupDelegate = {};
-}
-
-void FFrameEndSync::Sync()
-{
-	check(IsInGameThread());			
-
-#if !UE_BUILD_SHIPPING && PLATFORM_SUPPORTS_FLIP_TRACKING
-	// Set the FrameDebugInfo on platforms that have accurate frame tracking.
-	ENQUEUE_RENDER_COMMAND(FrameDebugInfo)(
-		[CurrentFrameCounter = GFrameCounter, CurrentInputTime = GInputTime](FRHICommandListImmediate& RHICmdList)
-	{
-		RHICmdList.EnqueueLambda(
-			[CurrentFrameCounter, CurrentInputTime](FRHICommandListImmediate&)
-		{
-			// Set the FrameCount and InputTime for input latency stats and flip debugging.
-			RHISetFrameDebugInfo(GRHIPresentCounter - 1, CurrentFrameCounter, CurrentInputTime);
-		});
-	});
-#endif
-
-	using ESyncDepth = FRenderCommandFence::ESyncDepth;
-
-	ESyncDepth SyncDepth;
-	int32 NumFramesOverlap;
-
-	uint32 const GTSyncType = (uint32)FMath::Max(CVarGTSyncType.GetValueOnAnyThread(), 0);
-
-	if (CVarAllowOneFrameThreadLag.GetValueOnAnyThread() <= 0)
-	{
-		// Legacy r.OneFrameThreadLag cvar is set to zero... the gamethread should not start frame N+1 until all work for frame N has completed.
-		// Depending on r.GTSyncType, sync with either the RHI thread or the swapchain. It is no longer correct to only sync with the render thread (pre UE 5.5 behaviour).
-		
-		SyncDepth = GTSyncType == 2
-			? ESyncDepth::Swapchain
-			: ESyncDepth::RHIThread;
-
-		NumFramesOverlap = 0;
-	}
-	else
-	{
-		bool bHasRHIThread = IsRunningRHIInSeparateThread();
-
-		switch (GTSyncType)
-		{
-		// HACK - only allow 1 frame of overlap while we investigate crashes caused by the GetPendingCleanupObjects() cleanup mechanism
-		case 0:  SyncDepth = ESyncDepth::RHIThread; NumFramesOverlap = 1/*bHasRHIThread ? 2 : 1*/; break;
-		case 1:  SyncDepth = ESyncDepth::RHIThread; NumFramesOverlap = 1; break;
-		case 2:  SyncDepth = ESyncDepth::Swapchain; NumFramesOverlap = 1; break;
-		default: SyncDepth = ESyncDepth::Swapchain; NumFramesOverlap = 1; break;
-		}
-	}
-
-	if (SyncDepth == ESyncDepth::Swapchain)
-	{
-		// Swapchain sync mode does not work when vsync is disabled. Fallback to RHI thread sync in that case.
-		static auto CVarVsync = IConsoleManager::Get().FindConsoleVariable(TEXT("r.VSync"));
-		check(CVarVsync != nullptr);
-
-		if (CVarVsync->GetInt() == 0)
-		{
-			SyncDepth = ESyncDepth::RHIThread;
-		}
-	}
-
-	Fences.Emplace_GetRef().BeginFence(SyncDepth);
-
-	bool bEmptyGameThreadTasks = !FTaskGraphInterface::Get().IsThreadProcessingTasks(ENamedThreads::GameThread);
-
-	if (bEmptyGameThreadTasks)
-	{
-		// need to process gamethread tasks at least once a frame no matter what
-		FTaskGraphInterface::Get().ProcessThreadUntilIdle(ENamedThreads::GameThread);
-	}
-
-	while (Fences.Num() > NumFramesOverlap)
-	{
-		Fences[0].Wait(bEmptyGameThreadTasks); // here we also opportunistically execute game thread tasks while we wait
-		Fences.RemoveAt(0);
-	}
-}
-
 FString appGetStartupMap(const TCHAR* CommandLine)
 {
 	FURL DefaultURL;
diff --git a/Engine/Source/Runtime/Engine/Public/EngineStats.h b/Engine/Source/Runtime/Engine/Public/EngineStats.h
--- a/Engine/Source/Runtime/Engine/Public/EngineStats.h
+++ b/Engine/Source/Runtime/Engine/Public/EngineStats.h
@@ -27,7 +27,6 @@ DECLARE_CYCLE_STAT_EXTERN(TEXT("RedrawViewports"),STAT_RedrawViewports,STATGROUP
 DECLARE_CYCLE_STAT_EXTERN(TEXT("Update Level Streaming"),STAT_UpdateLevelStreaming,STATGROUP_Engine, ENGINE_API);
 DECLARE_CYCLE_STAT_EXTERN(TEXT("RHI Game Tick"),STAT_RHITickTime,STATGROUP_Engine, ENGINE_API);
 DECLARE_CYCLE_STAT_EXTERN(TEXT("Debug Hitch"),STAT_IntentionalHitch,STATGROUP_Engine, );
-DECLARE_CYCLE_STAT_EXTERN(TEXT("Frame Sync Time"),STAT_FrameSyncTime,STATGROUP_Engine, ENGINE_API);
 DECLARE_CYCLE_STAT_EXTERN(TEXT("Deferred Tick Time"),STAT_DeferredTickTime,STATGROUP_Engine, ENGINE_API);
 
 /** Unit time stats */
diff --git a/Engine/Source/Runtime/Engine/Public/UnrealEngine.h b/Engine/Source/Runtime/Engine/Public/UnrealEngine.h
--- a/Engine/Source/Runtime/Engine/Public/UnrealEngine.h
+++ b/Engine/Source/Runtime/Engine/Public/UnrealEngine.h
@@ -364,48 +364,12 @@ ENGINE_API void CalculateFPSTimings();
 /** @return The font to use for rendering stats display. */
 extern ENGINE_API UFont* GetStatsFont();
 
-/*-----------------------------------------------------------------------------
-	Frame end sync object implementation.
------------------------------------------------------------------------------*/
-
-/**
- * Special helper class for frame end sync. It respects a passed in option to allow one frame
- * of lag between the game and the render thread by using two events in round robin fashion.
- */
-class FFrameEndSync
-{
-	/** Pair of fences. */
-	TArray<FRenderCommandFence, TInlineAllocator<3>> Fences;
-
-	/** cleanup delegate for engine pre-exit */
-	FDelegateHandle CleanupDelegate;
-
-public:
-	ENGINE_API FFrameEndSync();
-	ENGINE_API ~FFrameEndSync();
-
-	/**
-	 * Syncs the game thread with the render thread. Depending on passed in bool this will be a total
-	 * sync or a one frame lag.
-	 */
-	ENGINE_API void Sync();
-
-private:
-	void Cleanup();
-};
-
-
 /** Public interface to FEngineLoop so we can call it from editor or editor code */
 class IEngineLoop
 {
 public:
-
 	virtual int32 Init() = 0;
-
 	virtual void Tick() = 0;
-
-	/** Removes references to any objects pending cleanup by deleting them. */
-	virtual void ClearPendingCleanupObjects() = 0;
 };
 
 /**
diff --git a/Engine/Source/Runtime/Launch/Private/LaunchEngineLoop.cpp b/Engine/Source/Runtime/Launch/Private/LaunchEngineLoop.cpp
--- a/Engine/Source/Runtime/Launch/Private/LaunchEngineLoop.cpp
+++ b/Engine/Source/Runtime/Launch/Private/LaunchEngineLoop.cpp
@@ -4343,9 +4343,6 @@ int32 FEngineLoop::PreInitPostStartupScreen(const TCHAR* CmdLine)
 	}
 #endif // !UE_BUILD_SHIPPING
 
-	// initialize the pointer, as it is deleted before being assigned in the first frame
-	PendingCleanupObjects = nullptr;
-
 	// Initialize profile visualizers.
 #if !(UE_BUILD_SHIPPING || UE_BUILD_TEST)
 	FModuleManager::Get().LoadModule(TEXT("ProfileVisualizer"));
@@ -4988,7 +4985,8 @@ void FEngineLoop::Exit()
 	TRACE_CPUPROFILER_EVENT_SCOPE(FEngineLoop::Exit);
 	TRACE_BOOKMARK(TEXT("EngineLoop.Exit"));
 
-	ClearPendingCleanupObjects();
+	// Flush render commands to delete any pending FDeferredCleanupInterface objects.
+	FlushRenderingCommands();
 
 	GIsRunning	= 0;
 	GLogConsole	= nullptr;
@@ -6064,24 +6062,14 @@ void FEngineLoop::Tick()
 			TotalTickTime += FApp::GetDeltaTime();
 		}
 
-		// Find the objects which need to be cleaned up the next frame.
-		FPendingCleanupObjects* PreviousPendingCleanupObjects = PendingCleanupObjects;
-		PendingCleanupObjects = GetPendingCleanupObjects();
-
-		{
-			SCOPE_CYCLE_COUNTER(STAT_FrameSyncTime);
-			// this could be perhaps moved down to get greater parallelism
-			// Sync game and render/RHI threads.
-			static FFrameEndSync FrameEndSync;
-			FrameEndSync.Sync();
-		}
+		// This could be perhaps moved down to get greater parallelism
+		// Sync game and render/RHI threads.
+		FFrameEndSync::Sync();
 
 		// tick core ticker, threads & deferred commands
 		{
 			SCOPE_CYCLE_COUNTER(STAT_DeferredTickTime);
 			CSV_SCOPED_TIMING_STAT_EXCLUSIVE(DeferredTickTime);
-			// Delete the objects which were enqueued for deferred cleanup before the previous frame.
-			delete PreviousPendingCleanupObjects;
 
 #if WITH_COREUOBJECT
 			DeleteLoaders(); // destroy all linkers pending delete
@@ -6166,14 +6154,6 @@ void FEngineLoop::Tick()
 #endif
 }
 
-
-void FEngineLoop::ClearPendingCleanupObjects()
-{
-	FlushRenderingCommands();
-	delete PendingCleanupObjects;
-	PendingCleanupObjects = nullptr;
-}
-
 #endif // WITH_ENGINE
 
 
diff --git a/Engine/Source/Runtime/Launch/Public/LaunchEngineLoop.h b/Engine/Source/Runtime/Launch/Public/LaunchEngineLoop.h
--- a/Engine/Source/Runtime/Launch/Public/LaunchEngineLoop.h
+++ b/Engine/Source/Runtime/Launch/Public/LaunchEngineLoop.h
@@ -10,7 +10,6 @@
 
 class FEngineService;
 class FTraceService;
-class FPendingCleanupObjects;
 class ISessionService;
 class FSlateRenderer;
 
@@ -122,9 +121,6 @@ public:
 	/** Advances the main loop. */
 	virtual void Tick() override;
 
-	/** Removes references to any objects pending cleanup by deleting them. */
-	virtual void ClearPendingCleanupObjects() override;
-
 #endif // WITH_ENGINE
 
 	/** RHI post-init initialization */
@@ -174,13 +170,6 @@ protected:
 	/** Holds the number of cycles in the last frame. */
 	uint32 LastFrameCycles;
 
-#if WITH_ENGINE
-
-	/** Holds the objects which need to be cleaned up when the rendering thread finishes the previous frame. */
-	FPendingCleanupObjects* PendingCleanupObjects;
-
-#endif //WITH_ENGINE
-
 private:
 
 
diff --git a/Engine/Source/Runtime/RenderCore/Private/RenderingThread.cpp b/Engine/Source/Runtime/RenderCore/Private/RenderingThread.cpp
--- a/Engine/Source/Runtime/RenderCore/Private/RenderingThread.cpp
+++ b/Engine/Source/Runtime/RenderCore/Private/RenderingThread.cpp
@@ -734,9 +734,6 @@ static void StopRenderingThread()
 
 	GStopRenderingThreadDelegate.Broadcast();
 
-	// Get the list of objects which need to be cleaned up when the rendering thread is done with them.
-	FPendingCleanupObjects* PendingCleanupObjects = GetPendingCleanupObjects();
-
 	// Make sure we're not in the middle of streaming textures.
 	SuspendTextureStreamingRenderTasks();
 
@@ -787,9 +784,6 @@ static void StopRenderingThread()
 	delete GRenderingThreadRunnable;
 	GRenderingThreadRunnable = nullptr;
 
-	// Delete the pending cleanup objects which were in use by the rendering thread.
-	delete PendingCleanupObjects;
-
 	// Update can now resume with renderthread being the gamethread.
 	ResumeTextureStreamingRenderTasks();
 
@@ -1017,16 +1011,29 @@ void FRenderCommandFence::BeginFence(ESyncDepth SyncDepth)
 	{
 		return;
 	}
+
+	check(IsInGameThread());
 	
-	if (GRenderCommandFenceBundlerState.Event && IsInGameThread())
+	if (GRenderCommandFenceBundlerState.Event && SyncDepth == ESyncDepth::RenderThread)
 	{
-		check(SyncDepth == ESyncDepth::RenderThread);
+		// Case for game->render thread syncs when fence bundling is enabled. These are used
+		// throughout the engine when resources are destroyed. The fence bundling is an optimization
+		// to avoid the overhead of hundreds of individual fences.
+		// We aren't syncing any deeper than the render thread, so just use the bundled fence event.
 		CompletionTask = *GRenderCommandFenceBundlerState.Event;
 		return;
 	}
 
 	TRACE_CPUPROFILER_EVENT_SCOPE(FRenderCommandFence::BeginFence);
-	UE::Tasks::FTaskEvent Event { UE_SOURCE_LOCATION };
+	UE::Tasks::FTaskEvent Event{ UE_SOURCE_LOCATION };
+
+	if (GRenderCommandFenceBundlerState.Event)
+	{
+		// Render command fences are bundled, but we're syncing deeper than the render thread.
+		// Flush the fence bundler so we can insert an RHIThread (or deeper) fence in the right location.
+		Event.AddPrerequisites(*GRenderCommandFenceBundlerState.Event);
+		FlushRenderCommandFenceBundler();
+	}
 
 	if (GRenderCommandPipeMode == ERenderCommandPipeMode::All)
 	{
@@ -1324,16 +1331,9 @@ void FlushRenderingCommands()
 		RHICmdList.ImmediateFlush(EImmediateFlushType::FlushRHIThread);
 	});
 
-	// Find the objects which may be cleaned up once the rendering thread command queue has been flushed.
-	FPendingCleanupObjects* PendingCleanupObjects = GetPendingCleanupObjects();
-
 	// Issue a fence command to the rendering thread and wait for it to complete.
-	FRenderCommandFence Fence;
-	Fence.BeginFence();
-	Fence.Wait();
-
-	// Delete the objects which were enqueued for deferred cleanup before the command queue flush.
-	delete PendingCleanupObjects;
+	// This also deletes the objects which were enqueued for deferred cleanup before the command queue flush.
+	FFrameEndSync::Sync(/* bFullSync = */ true);
 
 	FCoreRenderDelegates::OnFlushRenderingCommandsEnd.Broadcast();
 }
@@ -1370,96 +1370,226 @@ static FAutoConsoleVariableRef CVarEnablePendingCleanupObjectsCommandBatching(
 	TEXT("Enable batching PendingCleanupObjects destruction.")
 );
 
-#if WITH_EDITOR || IS_PROGRAM
+TAutoConsoleVariable<int32> CVarAllowOneFrameThreadLag(
+	TEXT("r.OneFrameThreadLag"),
+	1,
+	TEXT("Whether to allow the rendering thread to lag one frame behind the game thread (0: disabled, otherwise enabled)")
+);
 
-// mainly concerned about the cooker here, but anyway, the editor can run without a frame for a very long time (hours) and we do not have enough lock free links. 
+TAutoConsoleVariable<int32> CVarGTSyncType(
+	TEXT("r.GTSyncType"),
+	0,
+	TEXT("Determines how the game thread syncs with the render thread, RHI thread and GPU.\n")
+	TEXT("Syncing to the GPU swap chain flip allows for lower frame latency.\n")
+	TEXT(" <= 0 - Sync the game thread with the N-1 render thread frame. Then sync with the N-m RHI thread frame where m is (2 + (-r.GTSyncType)) (i.e. negative values increase the amount of RHI thread overlap) (default = 0).\n")
+	TEXT("    1 - Sync the game thread with the N-1 RHI thread frame.\n")
+	TEXT("    2 - Sync the game thread with the GPU swap chain flip (only on supported platforms).\n"),
+	ECVF_Default
+);
 
-/** The set of deferred cleanup objects which are pending cleanup. */
-TArray<FDeferredCleanupInterface*> PendingCleanupObjectsList;
-FCriticalSection PendingCleanupObjectsListLock;
+DECLARE_CYCLE_STAT(TEXT("Frame Sync Time"), STAT_FrameSyncTime, STATGROUP_RHI);
 
-FPendingCleanupObjects::FPendingCleanupObjects()
+namespace FFrameEndSync
 {
-	check(IsInGameThread());
-	{
-		FScopeLock Lock(&PendingCleanupObjectsListLock);
-		Exchange(CleanupArray, PendingCleanupObjectsList);
-	}
-}
+	using ESyncDepth = FRenderCommandFence::ESyncDepth;
 
-FPendingCleanupObjects::~FPendingCleanupObjects()
-{
-	if (CleanupArray.Num())
+	class FDeferredCleanupInterfaceArray
 	{
-		QUICK_SCOPE_CYCLE_COUNTER(STAT_FPendingCleanupObjects_Destruct);
+	// Mainly concerned about the cooker here, but anyway, the editor can run without 
+	// a frame for a very long time (hours) and we do not have enough lock free links. 
+	#define USE_STANDARD_ARRAY (WITH_EDITOR || IS_PROGRAM)
+
+	#if USE_STANDARD_ARRAY
+		FCriticalSection CS;
+		TArray<FDeferredCleanupInterface*> Objects;
+	#else
+		TLockFreePointerListUnordered<FDeferredCleanupInterface, PLATFORM_CACHE_LINE_SIZE> Objects;
+	#endif
 
-		const bool bBatchingEnabled = bEnablePendingCleanupObjectsCommandBatching;
-		if (bBatchingEnabled)
+	public:
+		void Add(FDeferredCleanupInterface* Object)
 		{
-			StartRenderCommandFenceBundler();
+		#if USE_STANDARD_ARRAY
+			FScopeLock Lock(&CS);
+			Objects.Add(Object);
+		#else
+			Objects.Push(Object);
+		#endif
 		}
-		for (int32 ObjectIndex = 0; ObjectIndex < CleanupArray.Num(); ObjectIndex++)
+
+		void PopAll(TArray<FDeferredCleanupInterface*>& OutObjects)
 		{
-			delete CleanupArray[ObjectIndex];
+		#if USE_STANDARD_ARRAY
+			FScopeLock Lock(&CS);
+			OutObjects = MoveTemp(Objects);
+		#else
+			Objects.PopAll(OutObjects);
+		#endif
 		}
-		if (bBatchingEnabled)
+	#undef USE_STANDARD_ARRAY
+	};
+
+	// Wrapper to hold the objects which need to be cleaned up when the corresponding rendering frame finishes.
+	struct FPendingCleanupObjects : private TArray<FDeferredCleanupInterface*>
+	{
+		FPendingCleanupObjects(FDeferredCleanupInterfaceArray& InObjects)
 		{
-			StopRenderCommandFenceBundler();
+			check(IsInGameThread());
+			InObjects.PopAll(*this);
 		}
-	}
-}
 
-void BeginCleanup(FDeferredCleanupInterface* CleanupObject)
-{
+		~FPendingCleanupObjects()
+		{
+			check(IsInGameThread());
+
+			if (Num())
+			{
+				QUICK_SCOPE_CYCLE_COUNTER(STAT_FPendingCleanupObjects_Destruct);
+
+				const bool bBatchingEnabled = bEnablePendingCleanupObjectsCommandBatching;
+				if (bBatchingEnabled)
+				{
+					StartRenderCommandFenceBundler();
+				}
+
+				for (FDeferredCleanupInterface* Object : *this)
+				{
+					delete Object;
+				}
+
+				if (bBatchingEnabled)
+				{
+					StopRenderCommandFenceBundler();
+				}
+			}
+		}
+	};
+
+	// The set of deferred cleanup objects marked for deletion via BeginCleanup().
+	// These will be enqueued for deletion at the next Sync() and deleted one frame later.
+	FDeferredCleanupInterfaceArray PendingObjects;
+
+	struct FRenderThreadState
 	{
-		FScopeLock Lock(&PendingCleanupObjectsListLock);
-		PendingCleanupObjectsList.Add(CleanupObject);
-	}
-}
+		// Legacy game code assumes the game thread will never get further than 1 frame ahead of the render thread.
+		// This fence is used to sync the game thread with the N-1 render thread frame.
+		FRenderCommandFence Fence;
 
-#else
+		// Once we've synced the game and render threads, we can also delete the deferred cleanup objects for that frame.
+		FPendingCleanupObjects Objects;
 
-/** The set of deferred cleanup objects which are pending cleanup. */
-static TLockFreePointerListUnordered<FDeferredCleanupInterface, PLATFORM_CACHE_LINE_SIZE>	PendingCleanupObjectsList;
+		FRenderThreadState()
+			: Objects(PendingObjects)
+		{
+			Fence.BeginFence(ESyncDepth::RenderThread);
+		}
 
-FPendingCleanupObjects::FPendingCleanupObjects()
-{
-	check(IsInGameThread());
-	PendingCleanupObjectsList.PopAll(CleanupArray);
-}
+		~FRenderThreadState()
+		{
+			Fence.Wait(true);
+		}
+	};
+	TArray<FRenderThreadState, TInlineAllocator<2>> RenderThreadStates;
 
-FPendingCleanupObjects::~FPendingCleanupObjects()
-{
-	if (CleanupArray.Num())
+	// Additional fences to await. These sync with either the RHI thread or swapchain,
+	// and are used to prevent the game thread running too far ahead of presented frames.
+	TArray<FRenderCommandFence, TInlineAllocator<3>> PipelineFences;
+
+	void Sync(bool bFullSync)
 	{
-		QUICK_SCOPE_CYCLE_COUNTER(STAT_FPendingCleanupObjects_Destruct);
+		// The "r.OneFrameThreadLag" cvar forces a full sync, meaning the game thread will
+		// not start work until all the rendering work for the previous frame has completed.
+		bFullSync |= CVarAllowOneFrameThreadLag.GetValueOnAnyThread() <= 0;
+
+		SCOPE_CYCLE_COUNTER(STAT_FrameSyncTime);
+
+		check(IsInGameThread());			
+
+	#if !UE_BUILD_SHIPPING && PLATFORM_SUPPORTS_FLIP_TRACKING
+		// Set the FrameDebugInfo on platforms that have accurate frame tracking.
+		ENQUEUE_RENDER_COMMAND(FrameDebugInfo)(
+			[CurrentFrameCounter = GFrameCounter, CurrentInputTime = GInputTime](FRHICommandListImmediate& RHICmdList)
+		{
+			RHICmdList.EnqueueLambda(
+				[CurrentFrameCounter, CurrentInputTime](FRHICommandListImmediate&)
+			{
+				// Set the FrameCount and InputTime for input latency stats and flip debugging.
+				RHISetFrameDebugInfo(GRHIPresentCounter - 1, CurrentFrameCounter, CurrentInputTime);
+			});
+		});
+	#endif
+
+		// Always sync with the render thread (either current frame, or N-1 frame)
+		RenderThreadStates.Emplace();
+		while (RenderThreadStates.Num() > (bFullSync ? 0 : 1))
+		{
+			RenderThreadStates.RemoveAt(0);
+		}
+
+		// Insert an additional fence based on how we want to sync with the RHI thread / swapchain
+		ESyncDepth SyncDepth;
+		int32 NumFramesOverlap;
+
+		int32 const GTSyncType = CVarGTSyncType.GetValueOnAnyThread();
 
-		const bool bBatchingEnabled = bEnablePendingCleanupObjectsCommandBatching;
-		if (bBatchingEnabled)
+		if (bFullSync)
 		{
-			StartRenderCommandFenceBundler();
+			SyncDepth = GTSyncType >= 2
+				? ESyncDepth::Swapchain
+				: ESyncDepth::RHIThread;
+
+			NumFramesOverlap = 0;
+		}
+		else if (GTSyncType >= 2)
+		{
+			SyncDepth = ESyncDepth::Swapchain;
+			NumFramesOverlap = 1;
 		}
-		for (int32 ObjectIndex = 0; ObjectIndex < CleanupArray.Num(); ObjectIndex++)
+		else if (GTSyncType == 1)
+		{
+			SyncDepth = ESyncDepth::RHIThread;
+			NumFramesOverlap = 1;
+		}
+		else
 		{
-			delete CleanupArray[ObjectIndex];
+			check(GTSyncType <= 0);
+
+			// Modes <= 0 allows N frames of overlap with the RHI thread.
+			SyncDepth = ESyncDepth::RHIThread;
+			NumFramesOverlap = 2 + (-GTSyncType);
 		}
-		if (bBatchingEnabled)
+
+		if (SyncDepth == ESyncDepth::Swapchain)
 		{
-			StopRenderCommandFenceBundler();
+			// Swapchain sync mode does not work when vsync is disabled. Fallback to RHI thread sync in that case.
+			static auto CVarVsync = IConsoleManager::Get().FindConsoleVariable(TEXT("r.VSync"));
+			check(CVarVsync != nullptr);
+
+			if (CVarVsync->GetInt() == 0)
+			{
+				SyncDepth = ESyncDepth::RHIThread;
+			}
+		}
+
+		PipelineFences.Emplace_GetRef().BeginFence(SyncDepth);
+
+		if (!FTaskGraphInterface::Get().IsThreadProcessingTasks(ENamedThreads::GameThread))
+		{
+			// need to process gamethread tasks at least once a frame no matter what
+			FTaskGraphInterface::Get().ProcessThreadUntilIdle(ENamedThreads::GameThread);
+		}
+
+		while (PipelineFences.Num() > NumFramesOverlap)
+		{
+			PipelineFences[0].Wait(true);
+			PipelineFences.RemoveAt(0);
 		}
 	}
 }
 
 void BeginCleanup(FDeferredCleanupInterface* CleanupObject)
 {
-	PendingCleanupObjectsList.Push(CleanupObject);
-}
-
-#endif
-
-FPendingCleanupObjects* GetPendingCleanupObjects()
-{
-	return new FPendingCleanupObjects;
+	FFrameEndSync::PendingObjects.Add(CleanupObject);
 }
 
 static void HandleRHIThreadEnableChanged(const TArray<FString>& Args)
diff --git a/Engine/Source/Runtime/RenderCore/Public/RenderDeferredCleanup.h b/Engine/Source/Runtime/RenderCore/Public/RenderDeferredCleanup.h
--- a/Engine/Source/Runtime/RenderCore/Public/RenderDeferredCleanup.h
+++ b/Engine/Source/Runtime/RenderCore/Public/RenderDeferredCleanup.h
@@ -13,26 +13,7 @@ public:
 	virtual ~FDeferredCleanupInterface() {}
 };
 
-/**
- * A set of cleanup objects which are pending deletion.
- */
-class FPendingCleanupObjects
-{
-	TArray<FDeferredCleanupInterface*> CleanupArray;
-
-public:
-	inline bool IsEmpty() const { return CleanupArray.IsEmpty(); }
-	FPendingCleanupObjects();
-	RENDERCORE_API ~FPendingCleanupObjects();
-};
-
 /**
  * Adds the specified deferred cleanup object to the current set of pending cleanup objects.
  */
 extern RENDERCORE_API void BeginCleanup(FDeferredCleanupInterface* CleanupObject);
-
-/**
- * Transfers ownership of the current set of pending cleanup objects to the caller.  A new set is created for subsequent BeginCleanup calls.
- * @return A pointer to the set of pending cleanup objects.  The called is responsible for deletion.
- */
-extern RENDERCORE_API FPendingCleanupObjects* GetPendingCleanupObjects();
diff --git a/Engine/Source/Runtime/RenderCore/Public/RenderingThread.h b/Engine/Source/Runtime/RenderCore/Public/RenderingThread.h
--- a/Engine/Source/Runtime/RenderCore/Public/RenderingThread.h
+++ b/Engine/Source/Runtime/RenderCore/Public/RenderingThread.h
@@ -742,3 +742,13 @@ using FStopRenderingThreadDelegate = FStopRenderingThread::FDelegate;
 extern RENDERCORE_API FDelegateHandle RegisterStopRenderingThreadDelegate(const FStopRenderingThreadDelegate& InDelegate);
 
 extern RENDERCORE_API void UnregisterStopRenderingThreadDelegate(FDelegateHandle InDelegateHandle);
+
+namespace FFrameEndSync
+{
+	/**
+	 * Syncs the game thread based on progress throughout the rendering pipeline
+	 * @param bFullSync - when true, blocks the caller until all rendering work is completed,
+	 * otherwise only blocks until the N - m frame has completed, where m is driven by various config.
+	 */
+	RENDERCORE_API void Sync(bool bFullSync = false);
+};
diff --git a/Engine/Source/Runtime/SlateNullRenderer/Private/SlateNullRenderer.cpp b/Engine/Source/Runtime/SlateNullRenderer/Private/SlateNullRenderer.cpp
--- a/Engine/Source/Runtime/SlateNullRenderer/Private/SlateNullRenderer.cpp
+++ b/Engine/Source/Runtime/SlateNullRenderer/Private/SlateNullRenderer.cpp
@@ -5,7 +5,7 @@
 #include "Rendering/ShaderResourceManager.h"
 #include "Rendering/SlateDrawBuffer.h"
 #if UE_SLATE_NULL_RENDERER_WITH_ENGINE
-#include "UnrealEngine.h"
+#include "RenderingThread.h"
 #endif
 
 
@@ -144,7 +144,6 @@ void FSlateNullRenderer::ClearScenes()
 void FSlateNullRenderer::Sync() const
 {
 #if UE_SLATE_NULL_RENDERER_WITH_ENGINE
-	static FFrameEndSync FrameEndSync;
-	FrameEndSync.Sync();
+	FFrameEndSync::Sync();
 #endif
 }
diff --git a/Engine/Source/Runtime/SlateRHIRenderer/Private/SlateRHIRenderer.cpp b/Engine/Source/Runtime/SlateRHIRenderer/Private/SlateRHIRenderer.cpp
--- a/Engine/Source/Runtime/SlateRHIRenderer/Private/SlateRHIRenderer.cpp
+++ b/Engine/Source/Runtime/SlateRHIRenderer/Private/SlateRHIRenderer.cpp
@@ -1451,8 +1451,7 @@ void FSlateRHIRenderer::FlushCommands() const
 
 void FSlateRHIRenderer::Sync() const
 {
-	static FFrameEndSync FrameEndSync;
-	FrameEndSync.Sync();
+	FFrameEndSync::Sync();
 }
 
 void FSlateRHIRenderer::BeginFrame() const
