Don't check for 64-bit atomic support of the local machine when adding Nanite config to the shader map DDC key. Fixes issues with mismatched keys when the cooker doesn't have support. Also fixes potential issues when `r.Nanite` is configured differently for the cooking platform than the target platform.

diff --git a/Engine/Source/Runtime/RenderCore/Private/Shader.cpp b/Engine/Source/Runtime/RenderCore/Private/Shader.cpp
--- a/Engine/Source/Runtime/RenderCore/Private/Shader.cpp
+++ b/Engine/Source/Runtime/RenderCore/Private/Shader.cpp
@@ -2754,7 +2754,7 @@ void ShaderMapAppendKey(EShaderPlatform Platform, FShaderKeyGenerator& KeyGen)
 		KeyGen.Append(TEXT("SKPassThrough0"));
 	}
 
-	if (UseNanite(Platform))
+	if (DoesRuntimeSupportNanite(Platform, false, true))
 	{
 		static const auto CVarAllowSpline = IConsoleManager::Get().FindConsoleVariable(TEXT("r.Nanite.AllowSplineMeshes"));
 		static const auto CVarAllowSkinned = IConsoleManager::Get().FindConsoleVariable(TEXT("r.Nanite.AllowSkinnedMeshes"));
