Fix Substrate shader compilation issue when running -featureleveles31 with material having glints.

The glints code causes unrolling issue. Glints are normally disabled throught the platform settings, but when running with -featureleveles31, these settings are not taken into account. This CL disabled explicitely glints on platform lower than SM5.

diff --git a/Engine/Source/Runtime/RenderCore/Private/RenderUtils.cpp b/Engine/Source/Runtime/RenderCore/Private/RenderUtils.cpp
--- a/Engine/Source/Runtime/RenderCore/Private/RenderUtils.cpp
+++ b/Engine/Source/Runtime/RenderCore/Private/RenderUtils.cpp
@@ -1998,7 +1998,11 @@ namespace Substrate
 	bool IsGlintEnabled(EShaderPlatform InPlatform)
 	{
 		static FShaderPlatformCachedIniValue<int32> CVar(TEXT("r.Substrate.Glints"));
-		return IsSubstrateEnabled() && CVar.Get(InPlatform);
+		// Force disable glint on ES31/SM4 platforms as it causes compilation error (e.g. too large local light unroll in forward shaders).
+		// Usually this effect is explicitly disabled in XXXPlatform.ini files. However when running preview (e.g. -featureleveles31), 
+		// these settings are not taken into account causing compilation issue.
+		const bool bSupported = FDataDrivenShaderPlatformInfo::GetMaxFeatureLevel(InPlatform) >= ERHIFeatureLevel::SM5;
+		return IsSubstrateEnabled() && CVar.Get(InPlatform) && bSupported;
 	}
 
 	uint32 GlintLUTIndex()
