Media: Fixed color space transformation on HDR Electra videos.

The (optional) display mastering color space (usually used as a hints for the display devices) was incorrectly used as the sample source color space. Since we are linearizing to the working color space by default, display HDR metadata can be ignored. We also take this opportunity to cleanup the media sample color space interface, especially now that FColorSpace is in engine core. Furthermore, tonemapping is applied by on Electra HDR videos so that their (unlit or composited) default appearance resembles other media players.

This change also introduces two news console variables:
* Electra.HDR.WhiteLevel
* Electra.HDR.ToneMapMethod

diff --git a/Engine/Plugins/Media/AjaMedia/Source/AjaMedia/Private/Player/AjaMediaPlayer.cpp b/Engine/Plugins/Media/AjaMedia/Source/AjaMedia/Private/Player/AjaMediaPlayer.cpp
--- a/Engine/Plugins/Media/AjaMedia/Source/AjaMedia/Private/Player/AjaMediaPlayer.cpp
+++ b/Engine/Plugins/Media/AjaMedia/Source/AjaMedia/Private/Player/AjaMediaPlayer.cpp
@@ -644,21 +644,21 @@ bool FAjaMediaPlayer::OnInputFrameReceived(const AJA::AJAInputFrameData& InInput
 		
 		if (bOverrideSourceColorSpace)
 		{
-			ColorFormat.ColorSpace = (UE::Color::EColorSpace) OverrideSourceColorSpace;
+			ColorFormat.ColorSpaceType = (UE::Color::EColorSpace) OverrideSourceColorSpace;
 		}
 		else
 		{
 			switch (HDROptions.Gamut)
 			{
 			case EAjaHDRMetadataGamut::Rec709:
-				ColorFormat.ColorSpace = UE::Color::EColorSpace::sRGB;
+				ColorFormat.ColorSpaceType = UE::Color::EColorSpace::sRGB;
 				break;
 			case EAjaHDRMetadataGamut::Rec2020:
-				ColorFormat.ColorSpace = UE::Color::EColorSpace::Rec2020;
+				ColorFormat.ColorSpaceType = UE::Color::EColorSpace::Rec2020;
 				break;
 			default:
 				checkNoEntry();
-				ColorFormat.ColorSpace = UE::Color::EColorSpace::sRGB;
+				ColorFormat.ColorSpaceType = UE::Color::EColorSpace::sRGB;
 				break;
 			}
 		}
diff --git a/Engine/Plugins/Media/BlackmagicMedia/Source/BlackmagicMedia/Private/Player/BlackmagicMediaPlayer.cpp b/Engine/Plugins/Media/BlackmagicMedia/Source/BlackmagicMedia/Private/Player/BlackmagicMediaPlayer.cpp
--- a/Engine/Plugins/Media/BlackmagicMedia/Source/BlackmagicMedia/Private/Player/BlackmagicMediaPlayer.cpp
+++ b/Engine/Plugins/Media/BlackmagicMedia/Source/BlackmagicMedia/Private/Player/BlackmagicMediaPlayer.cpp
@@ -329,21 +329,21 @@ namespace BlackmagicMediaPlayerHelpers
 		
 					if (MediaPlayer->bOverrideSourceColorSpace)
 					{
-						ColorFormat.ColorSpace = (UE::Color::EColorSpace) MediaPlayer->OverrideSourceColorSpace;
+						ColorFormat.ColorSpaceType = (UE::Color::EColorSpace) MediaPlayer->OverrideSourceColorSpace;
 					}
 					else
 					{
 						switch (HDROptions.Gamut)
 						{
 						case EBlackmagicHDRMetadataGamut::Rec709:
-							ColorFormat.ColorSpace = UE::Color::EColorSpace::sRGB;
+							ColorFormat.ColorSpaceType = UE::Color::EColorSpace::sRGB;
 							break;
 						case EBlackmagicHDRMetadataGamut::Rec2020:
-							ColorFormat.ColorSpace = UE::Color::EColorSpace::Rec2020;
+							ColorFormat.ColorSpaceType = UE::Color::EColorSpace::Rec2020;
 							break;
 						default:
 							checkNoEntry();
-							ColorFormat.ColorSpace = UE::Color::EColorSpace::sRGB;
+							ColorFormat.ColorSpaceType = UE::Color::EColorSpace::sRGB;
 							break;
 						}
 					}
diff --git a/Engine/Plugins/Media/ElectraUtil/Source/ElectraSamples/Private/ElectraTextureSampleAVF.cpp b/Engine/Plugins/Media/ElectraUtil/Source/ElectraSamples/Private/ElectraTextureSampleAVF.cpp
--- a/Engine/Plugins/Media/ElectraUtil/Source/ElectraSamples/Private/ElectraTextureSampleAVF.cpp
+++ b/Engine/Plugins/Media/ElectraUtil/Source/ElectraSamples/Private/ElectraTextureSampleAVF.cpp
@@ -157,7 +157,7 @@ public:
 } // namespace anonymous
 
 
-void FElectraMediaTexConvApple::ConvertTexture(FTextureRHIRef & InDstTexture, CVImageBufferRef InImageBufferRef, bool bFullRange, EMediaTextureSampleFormat Format, const FMatrix44f& YUVMtx, const FMatrix44d& GamutToXYZMtx, UE::Color::EEncoding EncodingType, float NormalizationFactor)
+void FElectraMediaTexConvApple::ConvertTexture(FTextureRHIRef & InDstTexture, CVImageBufferRef InImageBufferRef, bool bFullRange, EMediaTextureSampleFormat Format, const FMatrix44f& YUVMtx, const UE::Color::FColorSpace& SourceColorSpace, UE::Color::EEncoding EncodingType, float NormalizationFactor)
 {
 	FRHICommandListImmediate& RHICmdList = FRHICommandListImmediate::Get();
 
@@ -238,9 +238,9 @@ void FElectraMediaTexConvApple::ConvertTexture(FTextureRHIRef & InDstTexture, CV
 					GraphicsPSOInit.BoundShaderState.VertexDeclarationRHI = GMediaVertexDeclaration.VertexDeclarationRHI;
 					GraphicsPSOInit.PrimitiveType = PT_TriangleStrip;
 
-					// Setup conversion from Rec2020 to current working color space
+					// Setup conversion from source color space (i.e. Rec2020) to current working color space
 					const UE::Color::FColorSpace& Working = UE::Color::FColorSpace::GetWorking();
-					FMatrix44f ColorSpaceMtx = FMatrix44f(Working.GetXYZToRgb().GetTransposed() * GamutToXYZMtx);
+					FMatrix44f ColorSpaceMtx = UE::Color::Transpose<float>(UE::Color::FColorSpaceTransform(SourceColorSpace, Working));
 					ColorSpaceMtx = ColorSpaceMtx.ApplyScale(NormalizationFactor);
 
 					if (Format == EMediaTextureSampleFormat::CharNV12)
@@ -453,7 +453,7 @@ bool FElectraTextureSample::Convert(FRHICommandListImmediate& RHICmdList, FTextu
 		{
 			PinnedTexConv->ConvertTexture(InDstTexture, VideoDecoderOutputApple->GetImageBuffer(),
 				VideoDecoderOutput->GetColorimetry().IsValid() ? VideoDecoderOutput->GetColorimetry()->GetMPEGDefinition()->VideoFullRangeFlag != 0 : true,
-				GetFormat(), GetSampleToRGBMatrix(), GetGamutToXYZMatrix(), GetEncodingType(), GetHDRNitsNormalizationFactor());
+				GetFormat(), GetSampleToRGBMatrix(), GetSourceColorSpace(), GetEncodingType(), GetHDRNitsNormalizationFactor());
 			return true;
 		}
 	}
diff --git a/Engine/Plugins/Media/ElectraUtil/Source/ElectraSamples/Private/ElectraTextureSampleAndroid.cpp b/Engine/Plugins/Media/ElectraUtil/Source/ElectraSamples/Private/ElectraTextureSampleAndroid.cpp
--- a/Engine/Plugins/Media/ElectraUtil/Source/ElectraSamples/Private/ElectraTextureSampleAndroid.cpp
+++ b/Engine/Plugins/Media/ElectraUtil/Source/ElectraSamples/Private/ElectraTextureSampleAndroid.cpp
@@ -959,7 +959,7 @@ void  FElectraTextureSample::CopyFromExternalTextureVulkan(FRHICommandListImmedi
 			}
 			else
 			{
-				ColorSpaceMtx = FMatrix44f(Working.GetXYZToRgb().GetTransposed() * GetGamutToXYZMatrix());
+				ColorSpaceMtx = UE::Color::Transpose<float>(UE::Color::FColorSpaceTransform(GetSourceColorSpace(), Working));
 			}
 	
 			float NF = GetHDRNitsNormalizationFactor();
@@ -1093,7 +1093,7 @@ bool FElectraTextureSample::ConvertCpuOutputPath(FRHICommandListImmediate& RHICm
 
 	// Setup conversion from Rec2020 to current working color space
 	const UE::Color::FColorSpace& Working = UE::Color::FColorSpace::GetWorking();
-	FMatrix44f ColorSpaceMtx = FMatrix44f(Working.GetXYZToRgb().GetTransposed() * GetGamutToXYZMatrix());
+	FMatrix44f ColorSpaceMtx = UE::Color::Transpose<float>(UE::Color::FColorSpaceTransform(GetSourceColorSpace(), Working));
 	if (GetEncodingType() == UE::Color::EEncoding::ST2084)
 	{
 		// Normalize output (e.g. 80 or 100 nits == 1.0)
diff --git a/Engine/Plugins/Media/ElectraUtil/Source/ElectraSamples/Private/ElectraTextureSampleBase.cpp b/Engine/Plugins/Media/ElectraUtil/Source/ElectraSamples/Private/ElectraTextureSampleBase.cpp
--- a/Engine/Plugins/Media/ElectraUtil/Source/ElectraSamples/Private/ElectraTextureSampleBase.cpp
+++ b/Engine/Plugins/Media/ElectraUtil/Source/ElectraSamples/Private/ElectraTextureSampleBase.cpp
@@ -4,6 +4,7 @@
 
 #if !UE_SERVER
 
+#include "HAL/IConsoleManager.h"
 #include "HAL/Platform.h"
 #include "ElectraTextureSample.h"
 #include "ElectraTextureSampleUtils.h"
@@ -27,11 +28,28 @@ static TOptional<FTimecode> CreateTimecodeFromMPEGDefinition(TOptional<FFrameRat
 
 }
 
+static TAutoConsoleVariable<float> CVarElectraHdrWhiteLevel(
+	TEXT("Electra.HDR.WhiteLevel"),
+	MediaTextureSample::kLinearToNitsScale_BT2408,
+	TEXT("White level as a linear to nits scale factor.\n")
+	TEXT("(default: 203.0)"),
+	ECVF_Default);
+
+static TAutoConsoleVariable<int32> CVarElectraHdrToneMapMethod(
+	TEXT("Electra.HDR.ToneMapMethod"),
+	static_cast<int32>(MediaShaders::EToneMapMethod::Hable),
+	TEXT("Tone mapping method applied on source HDR media:\n")
+	TEXT(" 0: None\n")
+	TEXT(" 1: Hable (default)\n")
+	TEXT(" 2: SimpleReinhard\n"),
+	ECVF_Default);
+
 // -------------------------------------------------------------------------------------------------------------------------------------------------------
 
 void IElectraTextureSampleBase::Initialize(FVideoDecoderOutput* InVideoDecoderOutput)
 {
 	VideoDecoderOutput = StaticCastSharedPtr<FVideoDecoderOutput, IDecoderOutputPoolable, ESPMode::ThreadSafe>(InVideoDecoderOutput->AsShared());
+	Colorimetry = VideoDecoderOutput->GetColorimetry();
 	HDRInfo = VideoDecoderOutput->GetHDRInformation();
 	TSharedPtr<const IVideoDecoderTimecode, ESPMode::ThreadSafe> TimecodePtr = VideoDecoderOutput->GetTimecode();
 
@@ -82,7 +100,7 @@ void IElectraTextureSampleBase::Initialize(FVideoDecoderOutput* InVideoDecoderOu
 	const FMatrix* Mtx = nullptr;
 
 	// Defaults in case no HDR info is present
-	bDisplayColorSpaceValid = false;
+	DisplayMasteringColorSpace.Reset();
 	DisplayMasteringLuminanceMin = -1.0f;
 	DisplayMasteringLuminanceMax = -1.0f;
 	MaxCLL = 0;
@@ -106,11 +124,10 @@ void IElectraTextureSampleBase::Initialize(FVideoDecoderOutput* InVideoDecoderOu
 				ColorVolume->display_primaries_x[2] <= ColorVolume->display_primaries_x[1] && 													// Blue's X is smaller or same than Green's
 				ColorVolume->display_primaries_x[1] <= ColorVolume->display_primaries_x[0]) 													// Red's X is greater or same than Green's
 			{
-				DisplayColorSpace = UE::Color::FColorSpace(FVector2d(ColorVolume->display_primaries_x[0], ColorVolume->display_primaries_y[0]),
+				DisplayMasteringColorSpace = UE::Color::FColorSpace(FVector2d(ColorVolume->display_primaries_x[0], ColorVolume->display_primaries_y[0]),
 														   FVector2d(ColorVolume->display_primaries_x[1], ColorVolume->display_primaries_y[1]),
 														   FVector2d(ColorVolume->display_primaries_x[2], ColorVolume->display_primaries_y[2]),
 														   FVector2d(ColorVolume->white_point_x, ColorVolume->white_point_y));
-				bDisplayColorSpaceValid = true;
 			}
 
 			DisplayMasteringLuminanceMin = ColorVolume->min_display_mastering_luminance;
@@ -125,8 +142,8 @@ void IElectraTextureSampleBase::Initialize(FVideoDecoderOutput* InVideoDecoderOu
 		}
 	}
 
-	// The sample color space is always defined by the color primaries value
-	SampleColorSpace = UE::Color::FColorSpace(ElectraColorimetryUtils::TranslateMPEGColorPrimaries(ColorPrimaries));
+	// The sample source color space is always defined by the color primaries value
+	SourceColorSpace = UE::Color::FColorSpace(ElectraColorimetryUtils::TranslateMPEGColorPrimaries(ColorPrimaries));
 
 	// Select the YUV-RGB conversion matrix to use
 	switch (ElectraColorimetryUtils::TranslateMPEGMatrixCoefficients(MatrixCoefficients))
@@ -265,34 +282,26 @@ FMatrix44f IElectraTextureSampleBase::GetSampleToRGBMatrix() const
 	return SampleToRgbMtx;
 }
 
-FMatrix44d IElectraTextureSampleBase::GetGamutToXYZMatrix() const
-{
-	return SampleColorSpace.GetRgbToXYZ().GetTransposed();
-}
-
-FVector2d IElectraTextureSampleBase::GetWhitePoint() const
+const UE::Color::FColorSpace& IElectraTextureSampleBase::GetSourceColorSpace() const
 {
-	return bDisplayColorSpaceValid ? DisplayColorSpace.GetWhiteChromaticity() : FVector2d(-1.0, -1.0);
+	return SourceColorSpace;
 }
 
-FVector2d IElectraTextureSampleBase::GetDisplayPrimaryRed() const
-{
-	return bDisplayColorSpaceValid ? DisplayColorSpace.GetRedChromaticity() : FVector2d(-1.0, -1.0);
-}
-
-FVector2d IElectraTextureSampleBase::GetDisplayPrimaryGreen() const
-{
-	return bDisplayColorSpaceValid ? DisplayColorSpace.GetGreenChromaticity() : FVector2d(-1.0, -1.0);
-}
-
-FVector2d IElectraTextureSampleBase::GetDisplayPrimaryBlue() const
+UE::Color::EEncoding IElectraTextureSampleBase::GetEncodingType() const
 {
-	return bDisplayColorSpaceValid ? DisplayColorSpace.GetBlueChromaticity() : FVector2d(-1.0, -1.0);
+	return ColorEncoding;
 }
 
-UE::Color::EEncoding IElectraTextureSampleBase::GetEncodingType() const
+float IElectraTextureSampleBase::GetHDRNitsNormalizationFactor() const
 {
-	return ColorEncoding;
+	if (GetEncodingType() == UE::Color::EEncoding::sRGB || GetEncodingType() == UE::Color::EEncoding::Linear)
+	{
+		return 1.0f;
+	}
+	else
+	{
+		return 1.0f / CVarElectraHdrWhiteLevel->GetFloat();
+	}
 }
 
 bool IElectraTextureSampleBase::GetDisplayMasteringLuminance(float& OutMin, float& OutMax) const
@@ -307,6 +316,11 @@ bool IElectraTextureSampleBase::GetDisplayMasteringLuminance(float& OutMin, floa
 	return true;
 }
 
+TOptional<UE::Color::FColorSpace> IElectraTextureSampleBase::GetDisplayMasteringColorSpace() const
+{
+	return DisplayMasteringColorSpace;
+}
+
 bool IElectraTextureSampleBase::GetMaxLuminanceLevels(uint16& OutCLL, uint16& OutFALL) const
 {
 	if (MaxCLL == 0 && MaxFALL == 0)
@@ -319,4 +333,18 @@ bool IElectraTextureSampleBase::GetMaxLuminanceLevels(uint16& OutCLL, uint16& Ou
 	return true;
 }
 
+MediaShaders::EToneMapMethod IElectraTextureSampleBase::GetToneMapMethod() const
+{
+	if (GetEncodingType() == UE::Color::EEncoding::sRGB || GetEncodingType() == UE::Color::EEncoding::Linear)
+	{
+		return MediaShaders::EToneMapMethod::None;
+	}
+	else
+	{
+		const int32 ToneMapMethod = FMath::Clamp(CVarElectraHdrToneMapMethod->GetInt(), 0, static_cast<int32>(MediaShaders::EToneMapMethod::MAX) - 1);
+
+		return static_cast<MediaShaders::EToneMapMethod>(ToneMapMethod);
+	}
+}
+
 #endif
diff --git a/Engine/Plugins/Media/ElectraUtil/Source/ElectraSamples/Public/Apple/ElectraTextureSample.h b/Engine/Plugins/Media/ElectraUtil/Source/ElectraSamples/Public/Apple/ElectraTextureSample.h
--- a/Engine/Plugins/Media/ElectraUtil/Source/ElectraSamples/Public/Apple/ElectraTextureSample.h
+++ b/Engine/Plugins/Media/ElectraUtil/Source/ElectraSamples/Public/Apple/ElectraTextureSample.h
@@ -73,7 +73,10 @@ public:
     ~FElectraMediaTexConvApple();
 
 #if WITH_ENGINE
-    void ConvertTexture(FTextureRHIRef & InDstTexture, CVImageBufferRef InImageBufferRef, bool bFullRange, EMediaTextureSampleFormat Format, const FMatrix44f& YUVMtx, const FMatrix44d& GamutToXYZMtx, UE::Color::EEncoding EncodingType, float NormalizationFactor);
+	void ConvertTexture(FTextureRHIRef& InDstTexture, CVImageBufferRef InImageBufferRef, bool bFullRange, EMediaTextureSampleFormat Format, const FMatrix44f& YUVMtx, const UE::Color::FColorSpace& SourceColorSpace, UE::Color::EEncoding EncodingType, float NormalizationFactor);
+
+	UE_DEPRECATED(5.5, "This ConvertTexture function is deprecated, please use the version with FColorSpace instead.")
+	void ConvertTexture(FTextureRHIRef& InDstTexture, CVImageBufferRef InImageBufferRef, bool bFullRange, EMediaTextureSampleFormat Format, const FMatrix44f& YUVMtx, const FMatrix44d& GamutToXYZMtx, UE::Color::EEncoding EncodingType, float NormalizationFactor) { }
 #endif
 
 private:
diff --git a/Engine/Plugins/Media/ElectraUtil/Source/ElectraSamples/Public/IElectraTextureSample.h b/Engine/Plugins/Media/ElectraUtil/Source/ElectraSamples/Public/IElectraTextureSample.h
--- a/Engine/Plugins/Media/ElectraUtil/Source/ElectraSamples/Public/IElectraTextureSample.h
+++ b/Engine/Plugins/Media/ElectraUtil/Source/ElectraSamples/Public/IElectraTextureSample.h
@@ -59,14 +59,13 @@ public:
 	virtual bool GetFullRange() const override;
 
 	virtual FMatrix44f GetSampleToRGBMatrix() const override;
-	virtual FMatrix44d GetGamutToXYZMatrix() const override;
-	virtual FVector2d GetWhitePoint() const override;
-	virtual FVector2d GetDisplayPrimaryRed() const override;
-	virtual FVector2d GetDisplayPrimaryGreen() const override;
-	virtual FVector2d GetDisplayPrimaryBlue() const override;
+	virtual const UE::Color::FColorSpace& GetSourceColorSpace() const override;
 	virtual UE::Color::EEncoding GetEncodingType() const override;
+	virtual float GetHDRNitsNormalizationFactor() const override;
 	virtual bool GetDisplayMasteringLuminance(float& OutMin, float& OutMax) const override;
+	virtual TOptional<UE::Color::FColorSpace> GetDisplayMasteringColorSpace() const override;
 	virtual bool GetMaxLuminanceLevels(uint16& OutCLL, uint16& OutFALL) const override;
+	virtual MediaShaders::EToneMapMethod GetToneMapMethod() const override;
 
 	virtual void SetTime(const FMediaTimeStamp& InTime);
 
@@ -92,9 +91,8 @@ protected:
 
 	/** Precomputed colorimetric data */
 	UE::Color::EEncoding ColorEncoding;
-	UE::Color::FColorSpace SampleColorSpace;
-	UE::Color::FColorSpace DisplayColorSpace;
-	bool bDisplayColorSpaceValid;
+	UE::Color::FColorSpace SourceColorSpace;
+	TOptional<UE::Color::FColorSpace> DisplayMasteringColorSpace;
 	float DisplayMasteringLuminanceMin;
 	float DisplayMasteringLuminanceMax;
 	uint16 MaxCLL;
diff --git a/Engine/Plugins/Media/MediaIOFramework/Source/MediaIOCore/Private/Shared/MediaIOCoreTextureSampleBase.cpp b/Engine/Plugins/Media/MediaIOFramework/Source/MediaIOCore/Private/Shared/MediaIOCoreTextureSampleBase.cpp
--- a/Engine/Plugins/Media/MediaIOFramework/Source/MediaIOCore/Private/Shared/MediaIOCoreTextureSampleBase.cpp
+++ b/Engine/Plugins/Media/MediaIOFramework/Source/MediaIOCore/Private/Shared/MediaIOCoreTextureSampleBase.cpp
@@ -121,8 +121,8 @@ bool FMediaIOCoreTextureSampleBase::SetProperties(uint32 InStride, uint32 InWidt
 	Duration = FTimespan(ETimespan::TicksPerSecond * InFrameRate.AsInterval());
 	Timecode = InTimecode;
 	Encoding = InColorFormatArgs.Encoding;
-	ColorSpace = InColorFormatArgs.ColorSpace;
-	ColorSpaceStruct = UE::Color::FColorSpace(ColorSpace);
+	ColorSpaceType = InColorFormatArgs.ColorSpaceType;
+	ColorSpaceStruct = UE::Color::FColorSpace(ColorSpaceType);
 
 	return true;
 }
@@ -226,7 +226,7 @@ void FMediaIOCoreTextureSampleBase::CopyConfiguration(const TSharedPtr<FMediaIOC
 	Time = SourceSample->Time;
 	Timecode = SourceSample->Timecode;
 	Encoding = SourceSample->Encoding;
-	ColorSpace = SourceSample->ColorSpace;
+	ColorSpaceType = SourceSample->ColorSpaceType;
 	ColorSpaceStruct = SourceSample->ColorSpaceStruct;
 	ColorConversionSettings = SourceSample->ColorConversionSettings;
 	CachedOCIOResources = SourceSample->CachedOCIOResources;
@@ -348,7 +348,7 @@ void FMediaIOCoreTextureSampleBase::ShutdownPoolable()
 
 const FMatrix& FMediaIOCoreTextureSampleBase::GetYUVToRGBMatrix() const
 {
-	switch (ColorSpace)
+	switch (ColorSpaceType)
 	{
 	case UE::Color::EColorSpace::sRGB:
 		return MediaShaders::YuvToRgbRec709Scaled;
@@ -368,29 +368,9 @@ bool FMediaIOCoreTextureSampleBase::IsOutputSrgb() const
 	return Encoding == UE::Color::EEncoding::sRGB;
 }
 
-FMatrix44d FMediaIOCoreTextureSampleBase::GetGamutToXYZMatrix() const
+const UE::Color::FColorSpace& FMediaIOCoreTextureSampleBase::GetSourceColorSpace() const
 {
-	return ColorSpaceStruct.GetRgbToXYZ().GetTransposed();
-}
-
-FVector2d FMediaIOCoreTextureSampleBase::GetWhitePoint() const
-{
-	return ColorSpaceStruct.GetWhiteChromaticity();
-}
-
-FVector2d FMediaIOCoreTextureSampleBase::GetDisplayPrimaryRed() const
-{
-	return ColorSpaceStruct.GetRedChromaticity();
-}
-
-FVector2d FMediaIOCoreTextureSampleBase::GetDisplayPrimaryGreen() const
-{
-	return ColorSpaceStruct.GetGreenChromaticity();
-} 
-
-FVector2d FMediaIOCoreTextureSampleBase::GetDisplayPrimaryBlue() const
-{
-	return ColorSpaceStruct.GetBlueChromaticity();
+	return ColorSpaceStruct;
 }
 
 UE::Color::EEncoding FMediaIOCoreTextureSampleBase::GetEncodingType() const
@@ -409,7 +389,7 @@ UE::Color::EColorSpace FMediaIOCoreTextureSampleBase::GetColorSpaceType() const
 	{
 		return UE::Color::EColorSpace::None;
 	}
-	return ColorSpace;
+	return ColorSpaceType;
 }
 
 float FMediaIOCoreTextureSampleBase::GetHDRNitsNormalizationFactor() const
@@ -419,5 +399,5 @@ float FMediaIOCoreTextureSampleBase::GetHDRNitsNormalizationFactor() const
     	return 1.0f;
     }
 	
-	return (GetEncodingType() == UE::Color::EEncoding::sRGB || GetEncodingType() == UE::Color::EEncoding::Linear) ? 1.0f : kMediaSample_HDR_NitsNormalizationFactor;
+	return IMediaTextureSample::GetHDRNitsNormalizationFactor();
 }
\ No newline at end of file
diff --git a/Engine/Plugins/Media/MediaIOFramework/Source/MediaIOCore/Public/MediaIOCoreTextureSampleBase.h b/Engine/Plugins/Media/MediaIOFramework/Source/MediaIOCore/Public/MediaIOCoreTextureSampleBase.h
--- a/Engine/Plugins/Media/MediaIOFramework/Source/MediaIOCore/Public/MediaIOCoreTextureSampleBase.h
+++ b/Engine/Plugins/Media/MediaIOFramework/Source/MediaIOCore/Public/MediaIOCoreTextureSampleBase.h
@@ -28,18 +28,18 @@ namespace UE::MediaIOCore
 			if (bIsSRGBInput)
 			{
 				Encoding = UE::Color::EEncoding::sRGB;
-				ColorSpace = UE::Color::EColorSpace::sRGB;
+				ColorSpaceType = UE::Color::EColorSpace::sRGB;
 			}
 			else
 			{
 				Encoding = UE::Color::EEncoding::Linear;
-				ColorSpace = UE::Color::EColorSpace::sRGB;
+				ColorSpaceType = UE::Color::EColorSpace::sRGB;
 			}
 		}
 
-		FColorFormatArgs(UE::Color::EEncoding InEncoding, UE::Color::EColorSpace InColorSpace)
+		FColorFormatArgs(UE::Color::EEncoding InEncoding, UE::Color::EColorSpace InColorSpaceType)
 			: Encoding(InEncoding)
-			, ColorSpace(InColorSpace)
+			, ColorSpaceType(InColorSpaceType)
 		{
 		}
 		
@@ -48,7 +48,7 @@ namespace UE::MediaIOCore
 		UE::Color::EEncoding Encoding = UE::Color::EEncoding::Linear;
 
 		/** Color space of the texture. */
-		UE::Color::EColorSpace ColorSpace = UE::Color::EColorSpace::sRGB;
+		UE::Color::EColorSpace ColorSpaceType = UE::Color::EColorSpace::sRGB;
 	};
 }
 
@@ -314,11 +314,7 @@ public:
 
 	virtual bool IsOutputSrgb() const override;
 	
-	virtual FMatrix44d GetGamutToXYZMatrix() const override;
-	virtual FVector2d GetWhitePoint() const override;
-	virtual FVector2d GetDisplayPrimaryRed() const override;
-	virtual FVector2d GetDisplayPrimaryGreen() const override;
-	virtual FVector2d GetDisplayPrimaryBlue() const override;
+	virtual const UE::Color::FColorSpace& GetSourceColorSpace() const override;
 	virtual UE::Color::EEncoding GetEncodingType() const override;
 	virtual float GetHDRNitsNormalizationFactor() const override;
 	
@@ -457,7 +453,7 @@ protected:
 	UE::Color::EEncoding Encoding = UE::Color::EEncoding::Linear;
 
 	/** Color space enum of the incoming texture. */
-	UE::Color::EColorSpace ColorSpace = UE::Color::EColorSpace::sRGB;
+	UE::Color::EColorSpace ColorSpaceType = UE::Color::EColorSpace::sRGB;
 
 	/** Color space structure of the incoming texture. Used for retrieving chromaticities. */
 	UE::Color::FColorSpace ColorSpaceStruct = UE::Color::FColorSpace(UE::Color::EColorSpace::sRGB);
diff --git a/Engine/Plugins/VirtualProduction/Rivermax/RivermaxMedia/Source/RivermaxMedia/Private/RivermaxMediaPlayer.cpp b/Engine/Plugins/VirtualProduction/Rivermax/RivermaxMedia/Source/RivermaxMedia/Private/RivermaxMediaPlayer.cpp
--- a/Engine/Plugins/VirtualProduction/Rivermax/RivermaxMedia/Source/RivermaxMedia/Private/RivermaxMediaPlayer.cpp
+++ b/Engine/Plugins/VirtualProduction/Rivermax/RivermaxMedia/Source/RivermaxMedia/Private/RivermaxMediaPlayer.cpp
@@ -443,7 +443,7 @@ namespace UE::RivermaxMedia
 				ProxySampleDummy = MakeShared<FRivermaxMediaTextureSample>();
 				ProxySampleDummy->SetTexture(CreateIntermediateRenderTarget(RHICmdList, Resolution, NewSample->GetPixelFormat(), NewSample->IsOutputSrgb()));
 				UE::MediaIOCore::FColorFormatArgs ColorArgs;
-				ColorArgs.ColorSpace = bOverrideSourceColorSpace ? (UE::Color::EColorSpace)OverrideSourceColorSpace : UE::Color::EColorSpace::None;
+				ColorArgs.ColorSpaceType = bOverrideSourceColorSpace ? (UE::Color::EColorSpace)OverrideSourceColorSpace : UE::Color::EColorSpace::None;
 				ColorArgs.Encoding = bOverrideSourceEncoding ? (UE::Color::EEncoding)OverrideSourceEncoding : UE::Color::EEncoding::None;
 
 				ProxySampleDummy->SetProperties(NewSample->GetStride(), VideoTrackFormat.Dim.X, VideoTrackFormat.Dim.X, NewSample->GetFormat(), FTimespan(0), FFrameRate(), FTimecode(), ColorArgs);
diff --git a/Engine/Shaders/Private/MediaShaders.usf b/Engine/Shaders/Private/MediaShaders.usf
--- a/Engine/Shaders/Private/MediaShaders.usf
+++ b/Engine/Shaders/Private/MediaShaders.usf
@@ -7,6 +7,7 @@
 #include "Common.ush"
 #include "ColorUtils.ush"
 #include "GammaCorrectionCommon.ush"
+#include "TonemapCommon.ush"
 
 /* Vertex shader
 =============================================================================*/
@@ -69,7 +70,11 @@ void BMPConvertPS(
 #define UE_Color_EEncoding_ProTune		15u
 #define UE_Color_EEncoding_VLog			16u
 
-float3 ApplyColorProcessingRGB(float3 Color, float4x4 CSMtx, uint EOTF)
+#define ToneMapMethod_None           0u
+#define ToneMapMethod_Hable          1u
+#define ToneMapMethod_SimpleReinhard 2u
+
+float3 ApplyColorProcessingRGB(float3 Color, float4x4 CSMtx, uint EOTF, uint ToneMapMethod = 0)
 {
 	// Go linear
 	switch(EOTF)
@@ -86,11 +91,18 @@ float3 ApplyColorProcessingRGB(float3 Color, float4x4 CSMtx, uint EOTF)
 	Color.y = dot(CSMtx[1], Col4);
 	Color.z = dot(CSMtx[2], Col4);
 
+	switch(ToneMapMethod)
+	{
+	case ToneMapMethod_Hable:			Color = HableToneMap(Color);  break;
+	case ToneMapMethod_SimpleReinhard:	Color = SimpleReinhardToneMap(Color);  break;
+	default: break;
+	}
+
 	return Color;
 }
 
 
-float3 ApplyColorProcessing(float3 YUV, float4x4 YUVMtx, float4x4 CSMtx, uint EOTF)
+float3 ApplyColorProcessing(float3 YUV, float4x4 YUVMtx, float4x4 CSMtx, uint EOTF, uint ToneMapMethod = 0)
 {
 	float3 Color;
 	float4 YUV4 = float4(YUV, 1);
@@ -101,7 +113,7 @@ float3 ApplyColorProcessing(float3 YUV, float4x4 YUVMtx, float4x4 CSMtx, uint EO
 	// Make sure we safely stay above zero
 	Color = max(Color, 0.0f);
 
-	return ApplyColorProcessingRGB(Color, CSMtx, EOTF);
+	return ApplyColorProcessingRGB(Color, CSMtx, EOTF, ToneMapMethod);
 }
 
 
@@ -119,7 +131,7 @@ void NV12ConvertPS(
 	YUV.yz = NV12ConvertUB.SwapChroma != 0u ? YUV.zy : YUV.yz;
 
 	// Output with solid alpha
-	OutColor = float4(ApplyColorProcessing(YUV, NV12ConvertUB.ColorTransform, NV12ConvertUB.CSTransform, NV12ConvertUB.EOTF), 1.0f); 
+	OutColor = float4(ApplyColorProcessing(YUV, NV12ConvertUB.ColorTransform, NV12ConvertUB.CSTransform, NV12ConvertUB.EOTF, NV12ConvertUB.ToneMapMethod), 1.0f); 
 }
 
 
@@ -146,7 +158,7 @@ void NV12ConvertAsBytesPS(
 	YUV.z = NV12ConvertAsBytesUB.Texture.Sample(NV12ConvertAsBytesUB.SamplerP, float2(VPos, ScaledUV.y)).x;
 	YUV.yz = NV12ConvertAsBytesUB.SwapChroma != 0u ? YUV.zy : YUV.yz;
 
-	OutColor = float4(ApplyColorProcessing(YUV, NV12ConvertAsBytesUB.ColorTransform, NV12ConvertAsBytesUB.CSTransform, NV12ConvertAsBytesUB.EOTF), 1.0f); 
+	OutColor = float4(ApplyColorProcessing(YUV, NV12ConvertAsBytesUB.ColorTransform, NV12ConvertAsBytesUB.CSTransform, NV12ConvertAsBytesUB.EOTF, NV12ConvertAsBytesUB.ToneMapMethod), 1.0f); 
 }
 
 
@@ -184,7 +196,7 @@ void RGBConvertPS(
 	float2 ScaledUV = InUV * RGBConvertUB.UVScale;
 	float4 Color = RGBConvertUB.Texture.Sample(RGBConvertUB.Sampler, ScaledUV);
 
-	OutColor = float4(ApplyColorProcessingRGB(Color.xyz, RGBConvertUB.CSTransform, RGBConvertUB.EOTF),
+	OutColor = float4(ApplyColorProcessingRGB(Color.xyz, RGBConvertUB.CSTransform, RGBConvertUB.EOTF, RGBConvertUB.ToneMapMethod),
 					  dot(RGBConvertUB.CSTransform[3], Color));
 }
 
@@ -208,7 +220,7 @@ void YCoCgConvertPS(
 	float Y = CoCgSY.w;
 	float3 Col = float3(Y + Co - Cg, Y + Cg, Y - Co - Cg);
 
-	OutColor = float4(ApplyColorProcessingRGB(Col, YCoCgConvertUB.CSTransform, YCoCgConvertUB.EOTF), 1.0f);
+	OutColor = float4(ApplyColorProcessingRGB(Col, YCoCgConvertUB.CSTransform, YCoCgConvertUB.EOTF, YCoCgConvertUB.ToneMapMethod), 1.0f);
 }
 
 // shader to convert a luma (Y) and chroma (CbCr) texture to RGBA
@@ -288,7 +300,7 @@ void YUVv216ConvertPS(
 	CbY0CrY1 = (YUVv216ConvertUB.SwapChroma != 0u) ? CbY0CrY1.zyxw : CbY0CrY1;						// Cb/Cr swap
 	float3 YUV = (uint(InUV.x * YUVv216ConvertUB.OutputDimX) & 1u) ? CbY0CrY1.wxz : CbY0CrY1.yxz;	// Decode to YCbCr for odd/even pixels
 
-	OutColor = float4(ApplyColorProcessing(YUV, YUVv216ConvertUB.ColorTransform, YUVv216ConvertUB.CSTransform, YUVv216ConvertUB.EOTF), 1.0f); 
+	OutColor = float4(ApplyColorProcessing(YUV, YUVv216ConvertUB.ColorTransform, YUVv216ConvertUB.CSTransform, YUVv216ConvertUB.EOTF, YUVv216ConvertUB.ToneMapMethod), 1.0f); 
 }
 
 
@@ -337,7 +349,7 @@ void YUVv210ConvertPS(
 		YUV = float3(In3.z, In2.z, In3.y);
 	}
 
-	OutColor = float4(ApplyColorProcessing(YUV, YUVv210ConvertUB.ColorTransform, YUVv210ConvertUB.CSTransform, YUVv210ConvertUB.EOTF), 1.0f); 
+	OutColor = float4(ApplyColorProcessing(YUV, YUVv210ConvertUB.ColorTransform, YUVv210ConvertUB.CSTransform, YUVv210ConvertUB.EOTF, YUVv210ConvertUB.ToneMapMethod), 1.0f); 
 }
 
 
@@ -350,7 +362,7 @@ void YUVY416ConvertPS(
 	float4 AVYU = YUVY416ConvertUB.SRV_Y.Sample(YUVY416ConvertUB.BilinearClampedSamplerUV, InUV);
 	AVYU = (YUVY416ConvertUB.bIsARGBFmt != 0u) ? AVYU.zyxw : AVYU;	// PF_ARGB8 vs. PF_ABGR16 ordering
 
-	OutColor = float4(ApplyColorProcessing(AVYU.yzw, YUVY416ConvertUB.ColorTransform, YUVY416ConvertUB.CSTransform, YUVY416ConvertUB.EOTF), AVYU.x); 
+	OutColor = float4(ApplyColorProcessing(AVYU.yzw, YUVY416ConvertUB.ColorTransform, YUVY416ConvertUB.CSTransform, YUVY416ConvertUB.EOTF, YUVY416ConvertUB.ToneMapMethod), AVYU.x); 
 }
 
 
@@ -366,7 +378,7 @@ void ARGB16BigConvertPS(
 
 	float3 Color = In.yzw;
 
-	OutColor = float4(ApplyColorProcessingRGB(Color, ARGB16BigConvertUB.CSTransform, ARGB16BigConvertUB.EOTF), In.x);
+	OutColor = float4(ApplyColorProcessingRGB(Color, ARGB16BigConvertUB.CSTransform, ARGB16BigConvertUB.EOTF, ARGB16BigConvertUB.ToneMapMethod), In.x);
 }
 
 
@@ -437,7 +449,7 @@ void P010ConvertPS(
 	YUV.x = P010ConvertUB.SRV_Y.Sample(P010ConvertUB.SamplerB, ScaledUV).x;
 	YUV.yz = P010ConvertUB.SRV_UV.Sample(P010ConvertUB.SamplerP, ScaledUV).xy;
 
-	OutColor = float4(ApplyColorProcessing(YUV, P010ConvertUB.ColorTransform, P010ConvertUB.CSTransform, P010ConvertUB.EOTF), 1.0f); 
+	OutColor = float4(ApplyColorProcessing(YUV, P010ConvertUB.ColorTransform, P010ConvertUB.CSTransform, P010ConvertUB.EOTF, P010ConvertUB.ToneMapMethod), 1.0f); 
 }
 
 // shader to convert a P010 texture to RGBA (P010 data, packed into a G16 texture; manually addressed)
@@ -462,7 +474,7 @@ void P010ConvertAsUINT16sPS(
 	YUV.y = P010ConvertAsUINT16sUB.Texture.Sample(P010ConvertAsUINT16sUB.SamplerP, float2(UPos, ScaledUV.y)).x;
 	YUV.z = P010ConvertAsUINT16sUB.Texture.Sample(P010ConvertAsUINT16sUB.SamplerP, float2(VPos, ScaledUV.y)).x;
 
-	OutColor = float4(ApplyColorProcessing(YUV, P010ConvertAsUINT16sUB.ColorTransform, P010ConvertAsUINT16sUB.CSTransform, P010ConvertAsUINT16sUB.EOTF), 1.0f); 
+	OutColor = float4(ApplyColorProcessing(YUV, P010ConvertAsUINT16sUB.ColorTransform, P010ConvertAsUINT16sUB.CSTransform, P010ConvertAsUINT16sUB.EOTF, P010ConvertAsUINT16sUB.ToneMapMethod), 1.0f); 
 }
 
 
@@ -490,7 +502,7 @@ void P010_2101010ConvertPS(
 	YUV.y = dot(P010_2101010ConvertUB.SRV_U.Sample(P010_2101010ConvertUB.SamplerP, ScaledUV_UV).xyz, SelUV);
 	YUV.z = dot(P010_2101010ConvertUB.SRV_V.Sample(P010_2101010ConvertUB.SamplerP, ScaledUV_UV).xyz, SelUV);
 
-	OutColor = float4(ApplyColorProcessing(YUV, P010_2101010ConvertUB.ColorTransform, P010_2101010ConvertUB.CSTransform, P010_2101010ConvertUB.EOTF), 1.0f); 
+	OutColor = float4(ApplyColorProcessing(YUV, P010_2101010ConvertUB.ColorTransform, P010_2101010ConvertUB.CSTransform, P010_2101010ConvertUB.EOTF, P010_2101010ConvertUB.ToneMapMethod), 1.0f); 
 #else
 	OutColor = float4(1,0,1,1);
 #endif // COMPILER_GLSL_ES3_1
@@ -598,7 +610,7 @@ void VYUConvertPS(
 	float3 YUV = VYUConvertUB.Texture.Sample(VYUConvertUB.Sampler, ScaledUV).gbr;
 
 	// Output with solid alpha
-	OutColor = float4(ApplyColorProcessing(YUV, VYUConvertUB.ColorTransform, VYUConvertUB.CSTransform, VYUConvertUB.EOTF), 1.0f); 
+	OutColor = float4(ApplyColorProcessing(YUV, VYUConvertUB.ColorTransform, VYUConvertUB.CSTransform, VYUConvertUB.EOTF, VYUConvertUB.ToneMapMethod), 1.0f); 
 }
 
 // Following defines must reflect values in EMediaCaptureConversionOperation in MediaOutput
diff --git a/Engine/Shaders/Private/TonemapCommon.ush b/Engine/Shaders/Private/TonemapCommon.ush
--- a/Engine/Shaders/Private/TonemapCommon.ush
+++ b/Engine/Shaders/Private/TonemapCommon.ush
@@ -615,3 +615,31 @@ float3 WhiteBalance(float3 LinearColor, float WhiteTemp, float WhiteTint, bool b
 
 	return mul(WhiteBalanceMat, LinearColor);
 }
+
+float3 HableToneMapCurve(float3 X)
+{
+	const float A = 0.15f;
+	const float B = 0.50f;
+	const float C = 0.10f;
+	const float D = 0.20f;
+	const float E = 0.02f;
+	const float F = 0.30f;
+
+	return ((X*(A*X+C*B)+D*E)/(X*(A*X+B)+D*F))-E/F;
+}
+
+float3 HableToneMap(float3 LinearColor)
+{
+	LinearColor = max(0, LinearColor);
+
+	const float3 W = 11.2f;
+	const float3 WhiteScale = 1.0f / HableToneMapCurve(W);
+	const float ExposureBias = 2.0f;
+
+	return HableToneMapCurve(LinearColor * ExposureBias) * WhiteScale;
+}
+
+float3 SimpleReinhardToneMap(float3 LinearColor)
+{
+	return LinearColor / (1.0f + LinearColor);
+}
diff --git a/Engine/Source/Runtime/Media/Public/IMediaTextureSample.h b/Engine/Source/Runtime/Media/Public/IMediaTextureSample.h
--- a/Engine/Source/Runtime/Media/Public/IMediaTextureSample.h
+++ b/Engine/Source/Runtime/Media/Public/IMediaTextureSample.h
@@ -153,7 +153,16 @@ enum class EMediaOrientation
 	CW270
 };
 
-static constexpr float kMediaSample_HDR_NitsNormalizationFactor = 1.0f / 100.0f;
+namespace MediaTextureSample
+{
+	/* This the reference white level for mapping UE scene-referred colors to nits (see TonemapCommon.ush). */
+	static constexpr float kLinearToNitsScale_UE = 100.0f;
+
+	/* This the reference white level for mapping SDR 1.0 to nits, as defined by ITU-R Report BT.2408. */
+	static constexpr float kLinearToNitsScale_BT2408 = 203.0f;
+}
+
+static constexpr float kMediaSample_HDR_NitsNormalizationFactor = 1.0f / MediaTextureSample::kLinearToNitsScale_UE;
 
 /**
  * Interface for media texture samples.
@@ -424,36 +433,39 @@ public:
 		return FMatrix44f(MediaShaders::YuvToRgbRec709Scaled * Pre);	// assumes sRGB & video range
 	}
 
-	/**
-	 * Get Colorspace conversion matrix to convert to CIE1931 XYZ space
-	 *
-	 * @return Conversion Matrix
-	 */
+	/*
+	* Get sample source color space (defaults to the sRGB/Rec709 gamut)
+	*/
+	virtual const UE::Color::FColorSpace& GetSourceColorSpace() const
+	{
+		return UE::Color::FColorSpace::GetSRGB();
+	}
+
+	UE_DEPRECATED(5.5, "GetGamutToXYZMatrix is deprecated, please use GetSourceColorSpace instead.")
 	virtual FMatrix44d GetGamutToXYZMatrix() const
 	{
 		return FMatrix44d(GamutToXYZMatrix(EDisplayColorGamut::sRGB_D65));
 	}
 
-	/**
-	 * Get white point of color space of the data contain in the sample
-	 *
-	 * @return White point
-	 */
+	UE_DEPRECATED(5.5, "GetWhitePoint is deprecated, please use GetSourceColorSpace instead.")
 	virtual FVector2d GetWhitePoint() const
 	{
 		return UE::Color::GetWhitePoint(UE::Color::EWhitePoint::CIE1931_D65);
 	}
 
+	UE_DEPRECATED(5.5, "GetDisplayPrimaryRed is deprecated, please use GetMasteringDisplayColorSpace instead.")
 	virtual FVector2d GetDisplayPrimaryRed() const
 	{
 		return FVector2d(0.64, 0.33);
 	}
 
+	UE_DEPRECATED(5.5, "GetDisplayPrimaryGreen is deprecated, please use GetMasteringDisplayColorSpace instead.")
 	virtual FVector2d GetDisplayPrimaryGreen() const
 	{
 		return FVector2d(0.30, 0.60);
 	}
 
+	UE_DEPRECATED(5.5, "GetDisplayPrimaryBlue is deprecated, please use GetMasteringDisplayColorSpace instead.")
 	virtual FVector2d GetDisplayPrimaryBlue() const
 	{
 		return FVector2d(0.15, 0.06);
@@ -483,6 +495,14 @@ public:
 		return false;
 	}
 
+	/**
+	 * Get display mastering color space
+	 */
+	virtual TOptional<UE::Color::FColorSpace> GetDisplayMasteringColorSpace() const
+	{
+		return TOptional<UE::Color::FColorSpace>();
+	}
+
 	/**
 	 * Get maximum luminance information
 	 */
@@ -491,6 +511,14 @@ public:
 		return false;
 	}
 
+	/**
+	 * Get an optional tonemapping method, for application on HDR inputs.
+	 */
+	virtual MediaShaders::EToneMapMethod GetToneMapMethod() const
+	{
+		return MediaShaders::EToneMapMethod::None;
+	}
+
 	/**
 	 * Reset sample to empty state
 	 */
diff --git a/Engine/Source/Runtime/MediaAssets/Private/Misc/MediaTextureResource.cpp b/Engine/Source/Runtime/MediaAssets/Private/Misc/MediaTextureResource.cpp
--- a/Engine/Source/Runtime/MediaAssets/Private/Misc/MediaTextureResource.cpp
+++ b/Engine/Source/Runtime/MediaAssets/Private/Misc/MediaTextureResource.cpp
@@ -902,13 +902,8 @@ bool FMediaTextureResource::RequiresConversion(const TSharedPtr<IMediaTextureSam
 	}
 
 	// Color space different?
-	const UE::Color::FColorSpace& Working = OverrideColorSpace.IsValid() ? *OverrideColorSpace : UE::Color::FColorSpace::GetWorking();
-	const float Tolerance = 1.e-7f;
-	if (Sample->ShouldApplyColorConversion() &&
-		(!Sample->GetDisplayPrimaryRed().Equals(Working.GetRedChromaticity(), Tolerance) ||
-		!Sample->GetDisplayPrimaryGreen().Equals(Working.GetGreenChromaticity(), Tolerance) ||
-		!Sample->GetDisplayPrimaryBlue().Equals(Working.GetBlueChromaticity(), Tolerance) ||
-		!Sample->GetWhitePoint().Equals(Working.GetWhiteChromaticity(), Tolerance)))
+	const UE::Color::FColorSpace& DestinationCS = OverrideColorSpace.IsValid() ? *OverrideColorSpace : UE::Color::FColorSpace::GetWorking();
+	if (Sample->ShouldApplyColorConversion() && !Sample->GetSourceColorSpace().Equals(DestinationCS))
 	{
 		// Yes! We need to convert...
 		return true;
@@ -1028,15 +1023,16 @@ void FMediaTextureResource::ConvertSample(FRHICommandListImmediate& RHICmdList,
 
 void FMediaTextureResource::GetColorSpaceConversionMatrixForSample(const TSharedPtr<IMediaTextureSample, ESPMode::ThreadSafe> Sample, FMatrix44f& ColorSpaceMtx)
 {
-	const UE::Color::FColorSpace& Working = OverrideColorSpace.IsValid() ? *OverrideColorSpace : UE::Color::FColorSpace::GetWorking();
+	const UE::Color::FColorSpace& DestinationCS = OverrideColorSpace.IsValid() ? *OverrideColorSpace : UE::Color::FColorSpace::GetWorking();
 	
-	if (Sample->GetMediaTextureSampleColorConverter())
+	if (Sample->GetMediaTextureSampleColorConverter() || Sample->GetSourceColorSpace().Equals(DestinationCS))
 	{
 		ColorSpaceMtx = FMatrix44f::Identity;
 	}
 	else
 	{
-		ColorSpaceMtx = FMatrix44f(Working.GetXYZToRgb().GetTransposed() * Sample->GetGamutToXYZMatrix());
+		// Apply the color space transformation from source to destination (including the Bradford chromatic adaptation).
+		ColorSpaceMtx = UE::Color::Transpose<float>(UE::Color::FColorSpaceTransform(Sample->GetSourceColorSpace(), DestinationCS));
 	}
 	
 	float NF = Sample->GetHDRNitsNormalizationFactor();
@@ -1112,14 +1108,14 @@ void FMediaTextureResource::GetColorSpaceConversionMatrixForSample(const TShared
 						FIntPoint TexDim = InputTexture->GetSizeXY();
 						TempSRV0 = RHICmdList.CreateShaderResourceView(InputTexture, 0, 1, PF_G8);								// note: RHI does provide "magic" to select Y vs. UV planes based on the pixel format (D3D/DXGI)
 						TempSRV1 = RHICmdList.CreateShaderResourceView(InputTexture, 0, 1, PF_R8G8);
-						SetShaderParametersLegacyPS(RHICmdList, ConvertShader, TexDim, TempSRV0, TempSRV1, OutputDim, YUVMtx, Sample->GetEncodingType(), ColorSpaceMtx, SampleFormat == EMediaTextureSampleFormat::CharNV21);
+						SetShaderParametersLegacyPS(RHICmdList, ConvertShader, TexDim, TempSRV0, TempSRV1, OutputDim, YUVMtx, Sample->GetEncodingType(), ColorSpaceMtx, SampleFormat == EMediaTextureSampleFormat::CharNV21, Sample->GetToneMapMethod());
 					}
 					else
 					{
 						TShaderMapRef<FNV12ConvertAsBytesPS> ConvertShader(ShaderMap);
 						GraphicsPSOInit.BoundShaderState.PixelShaderRHI = ConvertShader.GetPixelShader();
 						SetGraphicsPipelineState(RHICmdList, GraphicsPSOInit, 0);
-						SetShaderParametersLegacyPS(RHICmdList, ConvertShader, InputTexture, OutputDim, YUVMtx, Sample->GetEncodingType(), ColorSpaceMtx, SampleFormat == EMediaTextureSampleFormat::CharNV21);
+						SetShaderParametersLegacyPS(RHICmdList, ConvertShader, InputTexture, OutputDim, YUVMtx, Sample->GetEncodingType(), ColorSpaceMtx, SampleFormat == EMediaTextureSampleFormat::CharNV21, Sample->GetToneMapMethod());
 					}
 				}
 				break;
@@ -1140,7 +1136,7 @@ void FMediaTextureResource::GetColorSpaceConversionMatrixForSample(const TShared
 
 						FShaderResourceViewRHIRef Y_SRV = RHICmdList.CreateShaderResourceView(InputTexture, 0, 1, PF_G16);		// note: RHI does provide "magic" to select Y vs. UV planes based on the pixel format (D3D/DXGI)
 						FShaderResourceViewRHIRef UV_SRV = RHICmdList.CreateShaderResourceView(InputTexture, 0, 1, PF_G16R16);
-						SetShaderParametersLegacyPS(RHICmdList, ConvertShader, TexDim, Y_SRV, UV_SRV, OutputDim, YUVMtx, ColorSpaceMtx, Sample->GetEncodingType());
+						SetShaderParametersLegacyPS(RHICmdList, ConvertShader, TexDim, Y_SRV, UV_SRV, OutputDim, YUVMtx, ColorSpaceMtx, Sample->GetEncodingType(), Sample->GetToneMapMethod());
 					}
 					else
 					{
@@ -1148,7 +1144,7 @@ void FMediaTextureResource::GetColorSpaceConversionMatrixForSample(const TShared
 						GraphicsPSOInit.BoundShaderState.PixelShaderRHI = ConvertShader.GetPixelShader();
 						SetGraphicsPipelineState(RHICmdList, GraphicsPSOInit, 0);
 
-						SetShaderParametersLegacyPS(RHICmdList, ConvertShader, TexDim, InputTexture, OutputDim, YUVMtx, ColorSpaceMtx, Sample->GetEncodingType());
+						SetShaderParametersLegacyPS(RHICmdList, ConvertShader, TexDim, InputTexture, OutputDim, YUVMtx, ColorSpaceMtx, Sample->GetEncodingType(), Sample->GetToneMapMethod());
 					}
 				}
 				break;
@@ -1170,8 +1166,8 @@ void FMediaTextureResource::GetColorSpaceConversionMatrixForSample(const TShared
 					SetShaderParametersLegacyPS(RHICmdList, ConvertShader, InputTexture, OutputDim, YUVMtx, Sample->GetEncodingType(), ColorSpaceMtx,
 												SampleFormat!= EMediaTextureSampleFormat::CharYUY2 && SampleFormat != EMediaTextureSampleFormat::CharYVYU,	// Y or Cb first
 												InputTexture->GetFormat() == PF_B8G8R8A8,																	// ARGB vs. ABGR (memory order)
-												SampleFormat == EMediaTextureSampleFormat::CharYVYU															// Cb / Cr swap
-												);
+												SampleFormat == EMediaTextureSampleFormat::CharYVYU,														// Cb / Cr swap
+												Sample->GetToneMapMethod());
 				}
 				break;
 
@@ -1186,7 +1182,7 @@ void FMediaTextureResource::GetColorSpaceConversionMatrixForSample(const TShared
 					GraphicsPSOInit.BoundShaderState.PixelShaderRHI = ConvertShader.GetPixelShader();
 					SetGraphicsPipelineState(RHICmdList, GraphicsPSOInit, 0);
 					SetShaderParametersLegacyPS(RHICmdList, ConvertShader, InputTexture, OutputDim, YUVMtx, Sample->GetEncodingType(), ColorSpaceMtx,
-												true);
+												true, Sample->GetToneMapMethod());
 				}
 				break;
 
@@ -1205,7 +1201,7 @@ void FMediaTextureResource::GetColorSpaceConversionMatrixForSample(const TShared
 					FShaderResourceViewRHIRef SRV = RHICmdList.CreateShaderResourceView(InputTexture, 0, 1, (Sample->GetFormat() == EMediaTextureSampleFormat::Y416) ? PF_A16B16G16R16 : PF_A32B32G32R32F);
 
 					SetShaderParametersLegacyPS(RHICmdList, ConvertShader, SRV, YUVMtx, Sample->GetEncodingType(), ColorSpaceMtx,
-												InputTexture->GetFormat() == PF_A8R8G8B8);
+												InputTexture->GetFormat() == PF_A8R8G8B8, Sample->GetToneMapMethod());
 				}
 				break;
 
@@ -1238,7 +1234,7 @@ void FMediaTextureResource::GetColorSpaceConversionMatrixForSample(const TShared
 					TShaderMapRef<FRGBConvertPS> ConvertShader(ShaderMap);
 					GraphicsPSOInit.BoundShaderState.PixelShaderRHI = ConvertShader.GetPixelShader();
 					SetGraphicsPipelineState(RHICmdList, GraphicsPSOInit, 0);
-					SetShaderParametersLegacyPS(RHICmdList, ConvertShader, InputTexture, OutputDim, Encoding, ColorSpaceMtx);
+					SetShaderParametersLegacyPS(RHICmdList, ConvertShader, InputTexture, OutputDim, Encoding, ColorSpaceMtx, Sample->GetToneMapMethod());
 				}
 				break;
 
@@ -1252,7 +1248,7 @@ void FMediaTextureResource::GetColorSpaceConversionMatrixForSample(const TShared
 					SetGraphicsPipelineState(RHICmdList, GraphicsPSOInit, 0);
 					FShaderResourceViewRHIRef SRV = RHICmdList.CreateShaderResourceView(InputTexture, 0, 1, PF_R16G16B16A16_UINT);
 
-					SetShaderParametersLegacyPS(RHICmdList, ConvertShader, SRV, OutputDim, Sample->GetEncodingType(), ColorSpaceMtx);
+					SetShaderParametersLegacyPS(RHICmdList, ConvertShader, SRV, OutputDim, Sample->GetEncodingType(), ColorSpaceMtx, Sample->GetToneMapMethod());
 				}
 				break;
 
@@ -1291,7 +1287,7 @@ void FMediaTextureResource::GetColorSpaceConversionMatrixForSample(const TShared
 					TShaderMapRef<FRGBConvertPS> ConvertShader(ShaderMap);
 					GraphicsPSOInit.BoundShaderState.PixelShaderRHI = ConvertShader.GetPixelShader();
 					SetGraphicsPipelineState(RHICmdList, GraphicsPSOInit, 0);
-					SetShaderParametersLegacyPS(RHICmdList, ConvertShader, InputTexture, OutputDim, Encoding, ColorSpaceMtx);
+					SetShaderParametersLegacyPS(RHICmdList, ConvertShader, InputTexture, OutputDim, Encoding, ColorSpaceMtx, Sample->GetToneMapMethod());
 				}
 				break;
 
@@ -1328,7 +1324,7 @@ void FMediaTextureResource::GetColorSpaceConversionMatrixForSample(const TShared
 					TShaderMapRef<FYCoCgConvertPS> ConvertShader(ShaderMap);
 					GraphicsPSOInit.BoundShaderState.PixelShaderRHI = ConvertShader.GetPixelShader();
 					SetGraphicsPipelineState(RHICmdList, GraphicsPSOInit, 0);
-					SetShaderParametersLegacyPS(RHICmdList, ConvertShader, InputTexture, OutputDim, Encoding, ColorSpaceMtx);
+					SetShaderParametersLegacyPS(RHICmdList, ConvertShader, InputTexture, OutputDim, Encoding, ColorSpaceMtx, Sample->GetToneMapMethod());
 				}
 				break;
 
@@ -1341,7 +1337,7 @@ void FMediaTextureResource::GetColorSpaceConversionMatrixForSample(const TShared
 					TShaderMapRef<FVYUConvertPS> ConvertShader(ShaderMap);
 					GraphicsPSOInit.BoundShaderState.PixelShaderRHI = ConvertShader.GetPixelShader();
 					SetGraphicsPipelineState(RHICmdList, GraphicsPSOInit, 0);
-					SetShaderParametersLegacyPS(RHICmdList, ConvertShader, InputTexture, OutputDim, YUVMtx, Sample->GetEncodingType(), ColorSpaceMtx);
+					SetShaderParametersLegacyPS(RHICmdList, ConvertShader, InputTexture, OutputDim, YUVMtx, Sample->GetEncodingType(), ColorSpaceMtx, Sample->GetToneMapMethod());
 				}
 				break;
 
diff --git a/Engine/Source/Runtime/MediaAssets/Private/Misc/MediaTextureResource.h b/Engine/Source/Runtime/MediaAssets/Private/Misc/MediaTextureResource.h
--- a/Engine/Source/Runtime/MediaAssets/Private/Misc/MediaTextureResource.h
+++ b/Engine/Source/Runtime/MediaAssets/Private/Misc/MediaTextureResource.h
@@ -258,7 +258,7 @@ private:
 	/** Cached FRenderParams, used when JustInTimeRender() gets called. */
 	TUniquePtr<FRenderParams> JustInTimeRenderParams;
 
-	/** Colorspace to override standard proejct "working color space' */
+	/** Destination colorspace to override standard project "working color space'. Used primarily by Slate which remains sRGB-only. */
 	TUniquePtr<UE::Color::FColorSpace> OverrideColorSpace;
 
 	/** Used to keep track of whether we should re-create the output target because the intermediate target has changed. */
diff --git a/Engine/Source/Runtime/RenderCore/Private/MediaShaders.cpp b/Engine/Source/Runtime/RenderCore/Private/MediaShaders.cpp
--- a/Engine/Source/Runtime/RenderCore/Private/MediaShaders.cpp
+++ b/Engine/Source/Runtime/RenderCore/Private/MediaShaders.cpp
@@ -204,6 +204,7 @@ SHADER_PARAMETER(FMatrix44f, ColorTransform)
 SHADER_PARAMETER(FMatrix44f, CSTransform)
 SHADER_PARAMETER(uint32, OutputWidth)
 SHADER_PARAMETER(uint32, EOTF)
+SHADER_PARAMETER(uint32, ToneMapMethod)
 SHADER_PARAMETER(uint32, SwapChroma)
 SHADER_PARAMETER(FVector2f, UVScale)
 SHADER_PARAMETER_SRV(Texture2D, SRV_Y)
@@ -215,7 +216,7 @@ END_GLOBAL_SHADER_PARAMETER_STRUCT()
 IMPLEMENT_GLOBAL_SHADER_PARAMETER_STRUCT(FNV12ConvertUB, "NV12ConvertUB");
 IMPLEMENT_SHADER_TYPE(, FNV12ConvertPS, TEXT("/Engine/Private/MediaShaders.usf"), TEXT("NV12ConvertPS"), SF_Pixel);
 
-void FNV12ConvertPS::SetParameters(FRHIBatchedShaderParameters& BatchedParameters, const FIntPoint & TexDim, FShaderResourceViewRHIRef SRV_Y, FShaderResourceViewRHIRef SRV_UV, const FIntPoint& OutputDimensions, const FMatrix44f& ColorTransform, UE::Color::EEncoding Encoding, const FMatrix44f& CSTransform, bool bSwapChroma)
+void FNV12ConvertPS::SetParameters(FRHIBatchedShaderParameters& BatchedParameters, const FIntPoint & TexDim, FShaderResourceViewRHIRef SRV_Y, FShaderResourceViewRHIRef SRV_UV, const FIntPoint& OutputDimensions, const FMatrix44f& ColorTransform, UE::Color::EEncoding Encoding, const FMatrix44f& CSTransform, bool bSwapChroma, MediaShaders::EToneMapMethod ToneMapMethod)
 {
 	// Ensure shader code assumptions about the layout of this enum are correct
 	static_assert(int(UE::Color::EEncoding::Linear) == 1, "Enum mismatch with use in shader code!");
@@ -242,7 +243,8 @@ void FNV12ConvertPS::SetParameters(FRHIBatchedShaderParameters& BatchedParameter
 		UB.OutputWidth = OutputDimensions.X;
 		UB.SamplerB = TStaticSamplerState<SF_Bilinear>::GetRHI();
 		UB.SamplerP = TStaticSamplerState<SF_Point>::GetRHI();
-		UB.EOTF = int(Encoding);
+		UB.EOTF = static_cast<uint32>(Encoding);
+		UB.ToneMapMethod = static_cast<uint32>(ToneMapMethod);
 		UB.SRV_Y = SRV_Y;
 		UB.SRV_UV = SRV_UV;
 		UB.UVScale = FVector2f((float)OutputDimensions.X / (float)TexDim.X, (float)OutputDimensions.Y / (float)TexDim.Y);
@@ -262,6 +264,7 @@ SHADER_PARAMETER(FMatrix44f, ColorTransform)
 SHADER_PARAMETER(FMatrix44f, CSTransform)
 SHADER_PARAMETER(uint32, OutputWidth)
 SHADER_PARAMETER(uint32, EOTF)
+SHADER_PARAMETER(uint32, ToneMapMethod)
 SHADER_PARAMETER(uint32, SwapChroma)
 SHADER_PARAMETER(FVector2f, UVScale)
 SHADER_PARAMETER_TEXTURE(Texture2D, Texture)
@@ -272,7 +275,7 @@ END_GLOBAL_SHADER_PARAMETER_STRUCT()
 IMPLEMENT_GLOBAL_SHADER_PARAMETER_STRUCT(FNV12ConvertAsBytesUB, "NV12ConvertAsBytesUB");
 IMPLEMENT_SHADER_TYPE(, FNV12ConvertAsBytesPS, TEXT("/Engine/Private/MediaShaders.usf"), TEXT("NV12ConvertAsBytesPS"), SF_Pixel);
 
-void FNV12ConvertAsBytesPS::SetParameters(FRHIBatchedShaderParameters& BatchedParameters, TRefCountPtr<FRHITexture> NV12Texture, const FIntPoint& OutputDimensions, const FMatrix44f& ColorTransform, UE::Color::EEncoding Encoding, const FMatrix44f& CSTransform, bool bSwapChroma)
+void FNV12ConvertAsBytesPS::SetParameters(FRHIBatchedShaderParameters& BatchedParameters, TRefCountPtr<FRHITexture> NV12Texture, const FIntPoint& OutputDimensions, const FMatrix44f& ColorTransform, UE::Color::EEncoding Encoding, const FMatrix44f& CSTransform, bool bSwapChroma, MediaShaders::EToneMapMethod ToneMapMethod)
 {
 	FNV12ConvertAsBytesUB UB;
 	{
@@ -281,7 +284,8 @@ void FNV12ConvertAsBytesPS::SetParameters(FRHIBatchedShaderParameters& BatchedPa
 		UB.OutputWidth = OutputDimensions.X;
 		UB.SamplerB = TStaticSamplerState<SF_Bilinear>::GetRHI();
 		UB.SamplerP = TStaticSamplerState<SF_Point>::GetRHI();
-		UB.EOTF = int(Encoding);
+		UB.EOTF = static_cast<uint32>(Encoding);
+		UB.ToneMapMethod = static_cast<uint32>(ToneMapMethod);
 		UB.Texture = NV12Texture;
 		UB.UVScale = FVector2f((float)OutputDimensions.X / (float)NV12Texture->GetSizeX(), (float)OutputDimensions.Y / (float)NV12Texture->GetSizeY());
 		UB.SwapChroma = bSwapChroma;
@@ -334,6 +338,7 @@ BEGIN_GLOBAL_SHADER_PARAMETER_STRUCT(FP010ConvertUB, )
 SHADER_PARAMETER(FMatrix44f, ColorTransform)
 SHADER_PARAMETER(FMatrix44f, CSTransform)
 SHADER_PARAMETER(uint32, EOTF)
+SHADER_PARAMETER(uint32, ToneMapMethod)
 SHADER_PARAMETER(FVector2f, UVScale)
 SHADER_PARAMETER_SRV(Texture2D, SRV_Y)
 SHADER_PARAMETER_SRV(Texture2D, SRV_UV)
@@ -344,7 +349,7 @@ END_GLOBAL_SHADER_PARAMETER_STRUCT()
 IMPLEMENT_GLOBAL_SHADER_PARAMETER_STRUCT(FP010ConvertUB, "P010ConvertUB");
 IMPLEMENT_SHADER_TYPE(, FP010ConvertPS, TEXT("/Engine/Private/MediaShaders.usf"), TEXT("P010ConvertPS"), SF_Pixel);
 
-void FP010ConvertPS::SetParameters(FRHIBatchedShaderParameters& BatchedParameters, const FIntPoint& TexDim, FShaderResourceViewRHIRef SRV_Y, FShaderResourceViewRHIRef SRV_UV, const FIntPoint& OutputDimensions, const FMatrix44f& ColorTransform, const FMatrix44f& CSTransform, UE::Color::EEncoding Encoding)
+void FP010ConvertPS::SetParameters(FRHIBatchedShaderParameters& BatchedParameters, const FIntPoint& TexDim, FShaderResourceViewRHIRef SRV_Y, FShaderResourceViewRHIRef SRV_UV, const FIntPoint& OutputDimensions, const FMatrix44f& ColorTransform, const FMatrix44f& CSTransform, UE::Color::EEncoding Encoding, MediaShaders::EToneMapMethod ToneMapMethod)
 {
 	FP010ConvertUB UB;
 	{
@@ -352,7 +357,8 @@ void FP010ConvertPS::SetParameters(FRHIBatchedShaderParameters& BatchedParameter
 		UB.CSTransform = CSTransform;
 		UB.SamplerB = TStaticSamplerState<SF_Bilinear>::GetRHI();
 		UB.SamplerP = TStaticSamplerState<SF_Point>::GetRHI();
-		UB.EOTF = int(Encoding);
+		UB.EOTF = static_cast<uint32>(Encoding);
+		UB.ToneMapMethod = static_cast<uint32>(ToneMapMethod);
 		UB.SRV_Y = SRV_Y;
 		UB.SRV_UV = SRV_UV;
 		UB.UVScale = FVector2f((float)OutputDimensions.X / (float)TexDim.X, (float)OutputDimensions.Y / (float)TexDim.Y);
@@ -371,6 +377,7 @@ SHADER_PARAMETER(FMatrix44f, ColorTransform)
 SHADER_PARAMETER(FMatrix44f, CSTransform)
 SHADER_PARAMETER(uint32, OutputWidth)
 SHADER_PARAMETER(uint32, EOTF)
+SHADER_PARAMETER(uint32, ToneMapMethod)
 SHADER_PARAMETER(FVector2f, UVScale)
 SHADER_PARAMETER_TEXTURE(Texture2D, Texture)
 SHADER_PARAMETER_SAMPLER(SamplerState, SamplerB)
@@ -380,7 +387,7 @@ END_GLOBAL_SHADER_PARAMETER_STRUCT()
 IMPLEMENT_GLOBAL_SHADER_PARAMETER_STRUCT(FP010ConvertAsUINT16sUB, "P010ConvertAsUINT16sUB");
 IMPLEMENT_SHADER_TYPE(, FP010ConvertAsUINT16sPS, TEXT("/Engine/Private/MediaShaders.usf"), TEXT("P010ConvertAsUINT16sPS"), SF_Pixel);
 
-void FP010ConvertAsUINT16sPS::SetParameters(FRHIBatchedShaderParameters& BatchedParameters, const FIntPoint& TexDim, TRefCountPtr<FRHITexture> NV12Texture, const FIntPoint& OutputDimensions, const FMatrix44f& ColorTransform, const FMatrix44f& CSTransform, UE::Color::EEncoding Encoding)
+void FP010ConvertAsUINT16sPS::SetParameters(FRHIBatchedShaderParameters& BatchedParameters, const FIntPoint& TexDim, TRefCountPtr<FRHITexture> NV12Texture, const FIntPoint& OutputDimensions, const FMatrix44f& ColorTransform, const FMatrix44f& CSTransform, UE::Color::EEncoding Encoding, MediaShaders::EToneMapMethod ToneMapMethod)
 {
 	FP010ConvertAsUINT16sUB UB;
 	{
@@ -389,7 +396,8 @@ void FP010ConvertAsUINT16sPS::SetParameters(FRHIBatchedShaderParameters& Batched
 		UB.OutputWidth = OutputDimensions.X;
 		UB.SamplerB = TStaticSamplerState<SF_Bilinear>::GetRHI();
 		UB.SamplerP = TStaticSamplerState<SF_Point>::GetRHI();
-		UB.EOTF = int(Encoding);
+		UB.EOTF = static_cast<uint32>(Encoding);
+		UB.ToneMapMethod = static_cast<uint32>(ToneMapMethod);
 		UB.Texture = NV12Texture;
 		UB.UVScale = FVector2f((float)OutputDimensions.X / (float)TexDim.X, (float)OutputDimensions.Y / (float)TexDim.Y);
 	}
@@ -406,6 +414,7 @@ BEGIN_GLOBAL_SHADER_PARAMETER_STRUCT(FP010_2101010ConvertUB, )
 SHADER_PARAMETER(FMatrix44f, ColorTransform)
 SHADER_PARAMETER(FMatrix44f, CSTransform)
 SHADER_PARAMETER(uint32, EOTF)
+SHADER_PARAMETER(uint32, ToneMapMethod)
 SHADER_PARAMETER(FVector2f, UVScaleY)
 SHADER_PARAMETER(FVector2f, UVScaleUV)
 SHADER_PARAMETER_SRV(Texture2D, SRV_Y)
@@ -419,14 +428,15 @@ END_GLOBAL_SHADER_PARAMETER_STRUCT()
 IMPLEMENT_GLOBAL_SHADER_PARAMETER_STRUCT(FP010_2101010ConvertUB, "P010_2101010ConvertUB");
 IMPLEMENT_SHADER_TYPE(, FP010_2101010ConvertPS, TEXT("/Engine/Private/MediaShaders.usf"), TEXT("P010_2101010ConvertPS"), SF_Pixel);
 
-void FP010_2101010ConvertPS::SetParameters(FRHIBatchedShaderParameters& BatchedParameters, const FIntPoint& TexDim, FShaderResourceViewRHIRef SRV_Y, FShaderResourceViewRHIRef SRV_U, FShaderResourceViewRHIRef SRV_V, const FIntPoint& OutputDimensions, const FMatrix44f& ColorTransform, const FMatrix44f& CSTransform, UE::Color::EEncoding Encoding)
+void FP010_2101010ConvertPS::SetParameters(FRHIBatchedShaderParameters& BatchedParameters, const FIntPoint& TexDim, FShaderResourceViewRHIRef SRV_Y, FShaderResourceViewRHIRef SRV_U, FShaderResourceViewRHIRef SRV_V, const FIntPoint& OutputDimensions, const FMatrix44f& ColorTransform, const FMatrix44f& CSTransform, UE::Color::EEncoding Encoding, MediaShaders::EToneMapMethod ToneMapMethod)
 {
 	FP010_2101010ConvertUB UB;
 	{
 		UB.ColorTransform = ColorTransform;
 		UB.CSTransform = CSTransform;
 		UB.SamplerP = TStaticSamplerState<SF_Point>::GetRHI();
-		UB.EOTF = int(Encoding);
+		UB.EOTF = static_cast<uint32>(Encoding);
+		UB.ToneMapMethod = static_cast<uint32>(ToneMapMethod);
 		UB.SRV_Y = SRV_Y;
 		UB.SRV_U = SRV_U;
 		UB.SRV_V = SRV_V;
@@ -454,20 +464,22 @@ SHADER_PARAMETER_TEXTURE(Texture2D, Texture)
 SHADER_PARAMETER_SAMPLER(SamplerState, Sampler)
 SHADER_PARAMETER(FMatrix44f, CSTransform)
 SHADER_PARAMETER(uint32, EOTF)				// 0 = linear, 1=sRGB, 2=ST2084
+SHADER_PARAMETER(uint32, ToneMapMethod)
 END_GLOBAL_SHADER_PARAMETER_STRUCT()
 
 IMPLEMENT_GLOBAL_SHADER_PARAMETER_STRUCT(FRGBConvertUB, "RGBConvertUB");
 IMPLEMENT_SHADER_TYPE(, FRGBConvertPS, TEXT("/Engine/Private/MediaShaders.usf"), TEXT("RGBConvertPS"), SF_Pixel);
 
 
-void FRGBConvertPS::SetParameters(FRHIBatchedShaderParameters& BatchedParameters, TRefCountPtr<FRHITexture> RGBTexture, const FIntPoint& OutputDimensions, UE::Color::EEncoding Encoding, const FMatrix44f& CSTransform)
+void FRGBConvertPS::SetParameters(FRHIBatchedShaderParameters& BatchedParameters, TRefCountPtr<FRHITexture> RGBTexture, const FIntPoint& OutputDimensions, UE::Color::EEncoding Encoding, const FMatrix44f& CSTransform, MediaShaders::EToneMapMethod ToneMapMethod)
 {
 	FRGBConvertUB UB;
 	{
 		UB.Sampler = TStaticSamplerState<SF_Bilinear>::GetRHI();
 		UB.Texture = RGBTexture;
 		UB.UVScale = FVector2f((float)OutputDimensions.X / (float)RGBTexture->GetSizeX(), (float)OutputDimensions.Y / (float)RGBTexture->GetSizeY());
-		UB.EOTF = int(Encoding);
+		UB.EOTF = static_cast<uint32>(Encoding);
+		UB.ToneMapMethod = static_cast<uint32>(ToneMapMethod);
 		UB.CSTransform = CSTransform;
 	}
 
@@ -485,20 +497,22 @@ SHADER_PARAMETER_TEXTURE(Texture2D, Texture)
 SHADER_PARAMETER_SAMPLER(SamplerState, Sampler)
 SHADER_PARAMETER(FMatrix44f, CSTransform)
 SHADER_PARAMETER(uint32, EOTF)				// 0 = linear, 1=sRGB, 2=ST2084
+SHADER_PARAMETER(uint32, ToneMapMethod)
 END_GLOBAL_SHADER_PARAMETER_STRUCT()
 
 IMPLEMENT_GLOBAL_SHADER_PARAMETER_STRUCT(FYCoCgConvertUB, "YCoCgConvertUB");
 IMPLEMENT_SHADER_TYPE(, FYCoCgConvertPS, TEXT("/Engine/Private/MediaShaders.usf"), TEXT("YCoCgConvertPS"), SF_Pixel);
 
 
-void FYCoCgConvertPS::SetParameters(FRHIBatchedShaderParameters& BatchedParameters, TRefCountPtr<FRHITexture> RGBTexture, const FIntPoint& OutputDimensions, UE::Color::EEncoding Encoding, const FMatrix44f& CSTransform)
+void FYCoCgConvertPS::SetParameters(FRHIBatchedShaderParameters& BatchedParameters, TRefCountPtr<FRHITexture> RGBTexture, const FIntPoint& OutputDimensions, UE::Color::EEncoding Encoding, const FMatrix44f& CSTransform, MediaShaders::EToneMapMethod ToneMapMethod)
 {
 	FYCoCgConvertUB UB;
 	{
 		UB.Sampler = TStaticSamplerState<SF_Bilinear>::GetRHI();
 		UB.Texture = RGBTexture;
 		UB.UVScale = FVector2f((float)OutputDimensions.X / (float)RGBTexture->GetSizeX(), (float)OutputDimensions.Y / (float)RGBTexture->GetSizeY());
-		UB.EOTF = int(Encoding);
+		UB.EOTF = static_cast<uint32>(Encoding);
+		UB.ToneMapMethod = static_cast<uint32>(ToneMapMethod);
 		UB.CSTransform = CSTransform;
 	}
 
@@ -623,6 +637,7 @@ BEGIN_GLOBAL_SHADER_PARAMETER_STRUCT(FYUVv216ConvertUB, )
 SHADER_PARAMETER(FMatrix44f, ColorTransform)
 SHADER_PARAMETER(FMatrix44f, CSTransform)
 SHADER_PARAMETER(uint32, EOTF)
+SHADER_PARAMETER(uint32, ToneMapMethod)
 SHADER_PARAMETER(uint32, IsCbY0CrY1)
 SHADER_PARAMETER(uint32, IsARGBFmt)
 SHADER_PARAMETER(uint32, SwapChroma)
@@ -635,13 +650,14 @@ IMPLEMENT_GLOBAL_SHADER_PARAMETER_STRUCT(FYUVv216ConvertUB, "YUVv216ConvertUB");
 IMPLEMENT_SHADER_TYPE(, FYUVv216ConvertPS, TEXT("/Engine/Private/MediaShaders.usf"), TEXT("YUVv216ConvertPS"), SF_Pixel);
 
 
-void FYUVv216ConvertPS::SetParameters(FRHIBatchedShaderParameters& BatchedParameters, TRefCountPtr<FRHITexture> YUVTexture, const FIntPoint& OutputDimensions, const FMatrix44f& ColorTransform, UE::Color::EEncoding Encoding, const FMatrix44f& CSTransform, bool bIsCbY0CrY1, bool bIsARGBFmt, bool bSwapChroma)
+void FYUVv216ConvertPS::SetParameters(FRHIBatchedShaderParameters& BatchedParameters, TRefCountPtr<FRHITexture> YUVTexture, const FIntPoint& OutputDimensions, const FMatrix44f& ColorTransform, UE::Color::EEncoding Encoding, const FMatrix44f& CSTransform, bool bIsCbY0CrY1, bool bIsARGBFmt, bool bSwapChroma, MediaShaders::EToneMapMethod ToneMapMethod)
 {
 	FYUVv216ConvertUB UB;
 	{
 		UB.ColorTransform = ColorTransform;
 		UB.CSTransform = CSTransform;
-		UB.EOTF = int(Encoding);
+		UB.EOTF = static_cast<uint32>(Encoding);
+		UB.ToneMapMethod = static_cast<uint32>(ToneMapMethod);
 		UB.IsCbY0CrY1 = bIsCbY0CrY1;
 		UB.IsARGBFmt = bIsARGBFmt;
 		UB.SwapChroma = bSwapChroma;
@@ -662,6 +678,7 @@ BEGIN_GLOBAL_SHADER_PARAMETER_STRUCT(FYUVv210ConvertUB, )
 SHADER_PARAMETER(FMatrix44f, ColorTransform)
 SHADER_PARAMETER(FMatrix44f, CSTransform)
 SHADER_PARAMETER(uint32, EOTF)
+SHADER_PARAMETER(uint32, ToneMapMethod)
 SHADER_PARAMETER(uint32, IsCbY0CrY1)
 SHADER_PARAMETER(float, OutputDimX)
 SHADER_PARAMETER(float, OutputDimY)
@@ -672,13 +689,14 @@ IMPLEMENT_GLOBAL_SHADER_PARAMETER_STRUCT(FYUVv210ConvertUB, "YUVv210ConvertUB");
 IMPLEMENT_SHADER_TYPE(, FYUVv210ConvertPS, TEXT("/Engine/Private/MediaShaders.usf"), TEXT("YUVv210ConvertPS"), SF_Pixel);
 
 
-void FYUVv210ConvertPS::SetParameters(FRHIBatchedShaderParameters& BatchedParameters, TRefCountPtr<FRHITexture> YUVTexture, const FIntPoint& OutputDimensions, const FMatrix44f& ColorTransform, UE::Color::EEncoding Encoding, const FMatrix44f& CSTransform, bool bIsCbY0CrY1)
+void FYUVv210ConvertPS::SetParameters(FRHIBatchedShaderParameters& BatchedParameters, TRefCountPtr<FRHITexture> YUVTexture, const FIntPoint& OutputDimensions, const FMatrix44f& ColorTransform, UE::Color::EEncoding Encoding, const FMatrix44f& CSTransform, bool bIsCbY0CrY1, MediaShaders::EToneMapMethod ToneMapMethod)
 {
 	FYUVv210ConvertUB UB;
 	{
 		UB.ColorTransform = ColorTransform;
 		UB.CSTransform = CSTransform;
-		UB.EOTF = int(Encoding);
+		UB.EOTF = static_cast<uint32>(Encoding);
+		UB.ToneMapMethod = static_cast<uint32>(ToneMapMethod);
 		UB.IsCbY0CrY1 = bIsCbY0CrY1;
 		UB.OutputDimX = (float)OutputDimensions.X;
 		UB.OutputDimY = (float)OutputDimensions.Y;
@@ -699,6 +717,7 @@ SHADER_PARAMETER(FMatrix44f, CSTransform)
 SHADER_PARAMETER_SRV(Texture2D, SRV_Y)
 SHADER_PARAMETER_SAMPLER(SamplerState, BilinearClampedSamplerUV)
 SHADER_PARAMETER(uint32, EOTF)				// 0 = linear, 1=sRGB, 2=ST2084
+SHADER_PARAMETER(uint32, ToneMapMethod)
 SHADER_PARAMETER(uint32, bIsARGBFmt)
 END_GLOBAL_SHADER_PARAMETER_STRUCT()
 
@@ -706,7 +725,7 @@ IMPLEMENT_GLOBAL_SHADER_PARAMETER_STRUCT(FYUVY416ConvertUB, "YUVY416ConvertUB");
 IMPLEMENT_SHADER_TYPE(, FYUVY416ConvertPS, TEXT("/Engine/Private/MediaShaders.usf"), TEXT("YUVY416ConvertPS"), SF_Pixel);
 
 
-void FYUVY416ConvertPS::SetParameters(FRHIBatchedShaderParameters& BatchedParameters, FShaderResourceViewRHIRef SRV_Y, const FMatrix44f& ColorTransform, UE::Color::EEncoding Encoding, const FMatrix44f& CSTransform, bool bIsARGBFmt)
+void FYUVY416ConvertPS::SetParameters(FRHIBatchedShaderParameters& BatchedParameters, FShaderResourceViewRHIRef SRV_Y, const FMatrix44f& ColorTransform, UE::Color::EEncoding Encoding, const FMatrix44f& CSTransform, bool bIsARGBFmt, MediaShaders::EToneMapMethod ToneMapMethod)
 {
 	FYUVY416ConvertUB UB;
 	{
@@ -714,7 +733,8 @@ void FYUVY416ConvertPS::SetParameters(FRHIBatchedShaderParameters& BatchedParame
 		UB.BilinearClampedSamplerUV = TStaticSamplerState<SF_Bilinear, AM_Clamp, AM_Clamp, AM_Clamp>::GetRHI();
 		UB.ColorTransform = ColorTransform;
 		UB.CSTransform = CSTransform;
-		UB.EOTF = int(Encoding);
+		UB.EOTF = static_cast<uint32>(Encoding);
+		UB.ToneMapMethod = static_cast<uint32>(ToneMapMethod);
 		UB.bIsARGBFmt = bIsARGBFmt;
 	}
 
@@ -731,13 +751,14 @@ SHADER_PARAMETER(uint32, OutputWidth)
 SHADER_PARAMETER(uint32, OutputHeight)
 SHADER_PARAMETER_SRV(Texture2D<uint4>, Texture)
 SHADER_PARAMETER(uint32, EOTF)				// 0 = linear, 1=sRGB, 2=ST2084
+SHADER_PARAMETER(uint32, ToneMapMethod)
 END_GLOBAL_SHADER_PARAMETER_STRUCT()
 
 IMPLEMENT_GLOBAL_SHADER_PARAMETER_STRUCT(FARGB16BigConvertUB, "ARGB16BigConvertUB");
 IMPLEMENT_SHADER_TYPE(, FARGB16BigConvertPS, TEXT("/Engine/Private/MediaShaders.usf"), TEXT("ARGB16BigConvertPS"), SF_Pixel);
 
 
-void FARGB16BigConvertPS::SetParameters(FRHIBatchedShaderParameters& BatchedParameters, FShaderResourceViewRHIRef SRV, const FIntPoint& OutputDimensions, UE::Color::EEncoding Encoding, const FMatrix44f& CSTransform)
+void FARGB16BigConvertPS::SetParameters(FRHIBatchedShaderParameters& BatchedParameters, FShaderResourceViewRHIRef SRV, const FIntPoint& OutputDimensions, UE::Color::EEncoding Encoding, const FMatrix44f& CSTransform, MediaShaders::EToneMapMethod ToneMapMethod)
 {
 	FARGB16BigConvertUB UB;
 	{
@@ -745,7 +766,8 @@ void FARGB16BigConvertPS::SetParameters(FRHIBatchedShaderParameters& BatchedPara
 		UB.OutputWidth = OutputDimensions.X;
 		UB.OutputHeight = OutputDimensions.Y;
 		UB.CSTransform = CSTransform;
-		UB.EOTF = int(Encoding);
+		UB.EOTF = static_cast<uint32>(Encoding);
+		UB.ToneMapMethod = static_cast<uint32>(ToneMapMethod);
 	}
 
 	TUniformBufferRef<FARGB16BigConvertUB> Data = TUniformBufferRef<FARGB16BigConvertUB>::CreateUniformBufferImmediate(UB, UniformBuffer_SingleFrame);
@@ -769,7 +791,7 @@ IMPLEMENT_GLOBAL_SHADER_PARAMETER_STRUCT(FYUY2ConvertUB, "YUY2ConvertUB");
 IMPLEMENT_SHADER_TYPE(, FYUY2ConvertPS, TEXT("/Engine/Private/MediaShaders.usf"), TEXT("YUY2ConvertPS"), SF_Pixel);
 
 
-void FYUY2ConvertPS::SetParameters(FRHIBatchedShaderParameters& BatchedParameters, TRefCountPtr<FRHITexture> YUY2Texture, const FIntPoint& OutputDimensions, const FMatrix& ColorTransform, const FVector& YUVOffset, bool SrgbToLinear)
+void FYUY2ConvertPS::SetParameters(FRHIBatchedShaderParameters& BatchedParameters, TRefCountPtr<FRHITexture> YUY2Texture, const FIntPoint& OutputDimensions, const FMatrix& ColorTransform, const FVector& YUVOffset, bool SrgbToLinear, MediaShaders::EToneMapMethod ToneMapMethod)
 {
 	FYUY2ConvertUB UB;
 	{
@@ -946,6 +968,7 @@ BEGIN_GLOBAL_SHADER_PARAMETER_STRUCT(FVYUConvertUB, )
 SHADER_PARAMETER(FMatrix44f, ColorTransform)
 SHADER_PARAMETER(FMatrix44f, CSTransform)
 SHADER_PARAMETER(uint32, EOTF)
+SHADER_PARAMETER(uint32, ToneMapMethod)
 SHADER_PARAMETER(FVector2f, UVScale)
 SHADER_PARAMETER_TEXTURE(Texture2D, Texture)
 SHADER_PARAMETER_SAMPLER(SamplerState, Sampler)
@@ -955,13 +978,14 @@ IMPLEMENT_GLOBAL_SHADER_PARAMETER_STRUCT(FVYUConvertUB, "VYUConvertUB");
 IMPLEMENT_SHADER_TYPE(, FVYUConvertPS, TEXT("/Engine/Private/MediaShaders.usf"), TEXT("VYUConvertPS"), SF_Pixel);
 
 
-void FVYUConvertPS::SetParameters(FRHIBatchedShaderParameters& BatchedParameters, TRefCountPtr<FRHITexture> VYUTexture, const FIntPoint& OutputDimensions, const FMatrix44f& ColorTransform, UE::Color::EEncoding Encoding, const FMatrix44f& CSTransform)
+void FVYUConvertPS::SetParameters(FRHIBatchedShaderParameters& BatchedParameters, TRefCountPtr<FRHITexture> VYUTexture, const FIntPoint& OutputDimensions, const FMatrix44f& ColorTransform, UE::Color::EEncoding Encoding, const FMatrix44f& CSTransform, MediaShaders::EToneMapMethod ToneMapMethod)
 {
 	FVYUConvertUB UB;
 	{
 		UB.ColorTransform = ColorTransform;
 		UB.CSTransform = CSTransform;
-		UB.EOTF = int(Encoding);
+		UB.EOTF = static_cast<uint32>(Encoding);
+		UB.ToneMapMethod = static_cast<uint32>(ToneMapMethod);
 		UB.UVScale = FVector2f((float)OutputDimensions.X / (float)VYUTexture->GetSizeX(), (float)OutputDimensions.Y / (float)VYUTexture->GetSizeY());
 		UB.Texture = VYUTexture;
 		UB.Sampler = TStaticSamplerState<SF_Bilinear>::GetRHI();
diff --git a/Engine/Source/Runtime/RenderCore/Public/MediaShaders.h b/Engine/Source/Runtime/RenderCore/Public/MediaShaders.h
--- a/Engine/Source/Runtime/RenderCore/Public/MediaShaders.h
+++ b/Engine/Source/Runtime/RenderCore/Public/MediaShaders.h
@@ -44,6 +44,15 @@ namespace UE {
 
 namespace MediaShaders
 {
+	enum class EToneMapMethod : uint8
+	{
+		None = 0,
+		Hable = 1,
+		SimpleReinhard = 2,
+
+		MAX
+	};
+
 	/** Color transform from YUV to Rec601 without range scaling. */
 	RENDERCORE_API extern const FMatrix YuvToRgbRec601Unscaled;
 
@@ -239,7 +248,7 @@ public:
 		: FGlobalShader(Initializer)
 	{ }
 
-	RENDERCORE_API void SetParameters(FRHIBatchedShaderParameters& BatchedParameters, const FIntPoint & TexDim, FShaderResourceViewRHIRef SRV_Y, FShaderResourceViewRHIRef SRV_UV, const FIntPoint& OutputDimensions, const FMatrix44f& ColorTransform, UE::Color::EEncoding Encoding, const FMatrix44f& CSTransform, bool bSwapChroma);
+	RENDERCORE_API void SetParameters(FRHIBatchedShaderParameters& BatchedParameters, const FIntPoint & TexDim, FShaderResourceViewRHIRef SRV_Y, FShaderResourceViewRHIRef SRV_UV, const FIntPoint& OutputDimensions, const FMatrix44f& ColorTransform, UE::Color::EEncoding Encoding, const FMatrix44f& CSTransform, bool bSwapChroma, MediaShaders::EToneMapMethod ToneMapMethod = MediaShaders::EToneMapMethod::None);
 };
 
 
@@ -263,7 +272,7 @@ public:
 		: FGlobalShader(Initializer)
 	{ }
 
-	RENDERCORE_API void SetParameters(FRHIBatchedShaderParameters& BatchedParameters, TRefCountPtr<FRHITexture> NV12Texture, const FIntPoint& OutputDimensions, const FMatrix44f& ColorTransform, UE::Color::EEncoding Encoding, const FMatrix44f& CSTransform, bool bSwapChroma);
+	RENDERCORE_API void SetParameters(FRHIBatchedShaderParameters& BatchedParameters, TRefCountPtr<FRHITexture> NV12Texture, const FIntPoint& OutputDimensions, const FMatrix44f& ColorTransform, UE::Color::EEncoding Encoding, const FMatrix44f& CSTransform, bool bSwapChroma, MediaShaders::EToneMapMethod ToneMapMethod = MediaShaders::EToneMapMethod::None);
 };
 
 
@@ -305,7 +314,7 @@ public:
 		: FGlobalShader(Initializer)
 	{ }
 
-	RENDERCORE_API void SetParameters(FRHIBatchedShaderParameters& BatchedParameters, const FIntPoint& TexDim, FShaderResourceViewRHIRef SRV_Y, FShaderResourceViewRHIRef SRV_UV, const FIntPoint& OutputDimensions, const FMatrix44f& ColorTransform, const FMatrix44f& CSTransform, UE::Color::EEncoding Encoding);
+	RENDERCORE_API void SetParameters(FRHIBatchedShaderParameters& BatchedParameters, const FIntPoint& TexDim, FShaderResourceViewRHIRef SRV_Y, FShaderResourceViewRHIRef SRV_UV, const FIntPoint& OutputDimensions, const FMatrix44f& ColorTransform, const FMatrix44f& CSTransform, UE::Color::EEncoding Encoding, MediaShaders::EToneMapMethod ToneMapMethod = MediaShaders::EToneMapMethod::None);
 };
 
 /**
@@ -324,7 +333,7 @@ public:
 		: FGlobalShader(Initializer)
 	{ }
 
-	RENDERCORE_API void SetParameters(FRHIBatchedShaderParameters& BatchedParameters, const FIntPoint& TexDim, TRefCountPtr<FRHITexture> NV12Texture, const FIntPoint& OutputDimensions, const FMatrix44f& ColorTransform, const FMatrix44f& CSTransform, UE::Color::EEncoding Encoding);
+	RENDERCORE_API void SetParameters(FRHIBatchedShaderParameters& BatchedParameters, const FIntPoint& TexDim, TRefCountPtr<FRHITexture> NV12Texture, const FIntPoint& OutputDimensions, const FMatrix44f& ColorTransform, const FMatrix44f& CSTransform, UE::Color::EEncoding Encoding, MediaShaders::EToneMapMethod ToneMapMethod = MediaShaders::EToneMapMethod::None);
 };
 
 /**
@@ -343,7 +352,7 @@ public:
 		: FGlobalShader(Initializer)
 	{ }
 
-	RENDERCORE_API void SetParameters(FRHIBatchedShaderParameters& BatchedParameters, const FIntPoint& TexDim, FShaderResourceViewRHIRef SRV_Y, FShaderResourceViewRHIRef SRV_U, FShaderResourceViewRHIRef SRV_V, const FIntPoint& OutputDimensions, const FMatrix44f& ColorTransform, const FMatrix44f& CSTransform, UE::Color::EEncoding Encoding);
+	RENDERCORE_API void SetParameters(FRHIBatchedShaderParameters& BatchedParameters, const FIntPoint& TexDim, FShaderResourceViewRHIRef SRV_Y, FShaderResourceViewRHIRef SRV_U, FShaderResourceViewRHIRef SRV_V, const FIntPoint& OutputDimensions, const FMatrix44f& ColorTransform, const FMatrix44f& CSTransform, UE::Color::EEncoding Encoding, MediaShaders::EToneMapMethod ToneMapMethod = MediaShaders::EToneMapMethod::None);
 };
 
 /**
@@ -365,7 +374,7 @@ public:
 	{ }
 
 
-	RENDERCORE_API void SetParameters(FRHIBatchedShaderParameters& BatchedParameters, TRefCountPtr<FRHITexture> RGBTexture, const FIntPoint& OutputDimensions, UE::Color::EEncoding Encoding, const FMatrix44f& CSTransform);
+	RENDERCORE_API void SetParameters(FRHIBatchedShaderParameters& BatchedParameters, TRefCountPtr<FRHITexture> RGBTexture, const FIntPoint& OutputDimensions, UE::Color::EEncoding Encoding, const FMatrix44f& CSTransform, MediaShaders::EToneMapMethod ToneMapMethod = MediaShaders::EToneMapMethod::None);
 };
 
 /**
@@ -387,7 +396,7 @@ public:
 	{ }
 
 
-	RENDERCORE_API void SetParameters(FRHIBatchedShaderParameters& BatchedParameters, TRefCountPtr<FRHITexture> YCoCgTexture, const FIntPoint& OutputDimensions, UE::Color::EEncoding Encoding, const FMatrix44f& CSTransform);
+	RENDERCORE_API void SetParameters(FRHIBatchedShaderParameters& BatchedParameters, TRefCountPtr<FRHITexture> YCoCgTexture, const FIntPoint& OutputDimensions, UE::Color::EEncoding Encoding, const FMatrix44f& CSTransform, MediaShaders::EToneMapMethod ToneMapMethod = MediaShaders::EToneMapMethod::None);
 };
 
 
@@ -492,7 +501,7 @@ public:
 		: FGlobalShader(Initializer)
 	{ }
 
-	RENDERCORE_API void SetParameters(FRHIBatchedShaderParameters& BatchedParameters, TRefCountPtr<FRHITexture> YUVTexture, const FIntPoint& OutputDimensions, const FMatrix44f& ColorTransform, UE::Color::EEncoding Encoding, const FMatrix44f& CSTransform, bool bIsCbY0CrY1, bool bIsARGBFmt, bool bSwapChroma);
+	RENDERCORE_API void SetParameters(FRHIBatchedShaderParameters& BatchedParameters, TRefCountPtr<FRHITexture> YUVTexture, const FIntPoint& OutputDimensions, const FMatrix44f& ColorTransform, UE::Color::EEncoding Encoding, const FMatrix44f& CSTransform, bool bIsCbY0CrY1, bool bIsARGBFmt, bool bSwapChroma, MediaShaders::EToneMapMethod ToneMapMethod = MediaShaders::EToneMapMethod::None);
 };
 
 
@@ -514,7 +523,7 @@ public:
 		: FGlobalShader(Initializer)
 	{ }
 
-	RENDERCORE_API void SetParameters(FRHIBatchedShaderParameters& BatchedParameters, TRefCountPtr<FRHITexture> YUVTexture, const FIntPoint& OutputDimensions, const FMatrix44f& ColorTransform, UE::Color::EEncoding Encoding, const FMatrix44f& CSTransform, bool bIsCbY0CrY1);
+	RENDERCORE_API void SetParameters(FRHIBatchedShaderParameters& BatchedParameters, TRefCountPtr<FRHITexture> YUVTexture, const FIntPoint& OutputDimensions, const FMatrix44f& ColorTransform, UE::Color::EEncoding Encoding, const FMatrix44f& CSTransform, bool bIsCbY0CrY1, MediaShaders::EToneMapMethod ToneMapMethod = MediaShaders::EToneMapMethod::None);
 };
 
 
@@ -535,7 +544,7 @@ public:
 		: FGlobalShader(Initializer)
 	{ }
 
-	RENDERCORE_API void SetParameters(FRHIBatchedShaderParameters& BatchedParameters, FShaderResourceViewRHIRef SRV_Y, const FMatrix44f& ColorTransform, UE::Color::EEncoding Encoding, const FMatrix44f& CSTransform, bool bIsARGBFmt);
+	RENDERCORE_API void SetParameters(FRHIBatchedShaderParameters& BatchedParameters, FShaderResourceViewRHIRef SRV_Y, const FMatrix44f& ColorTransform, UE::Color::EEncoding Encoding, const FMatrix44f& CSTransform, bool bIsARGBFmt, MediaShaders::EToneMapMethod ToneMapMethod = MediaShaders::EToneMapMethod::None);
 };
 
 /**
@@ -558,7 +567,7 @@ public:
 		: FGlobalShader(Initializer)
 	{ }
 
-	RENDERCORE_API void SetParameters(FRHIBatchedShaderParameters& BatchedParameters, TRefCountPtr<FRHITexture> YUY2Texture, const FIntPoint& OutputDimensions, const FMatrix& ColorTransform, const FVector& YUVOffset, bool SrgbToLinear);
+	RENDERCORE_API void SetParameters(FRHIBatchedShaderParameters& BatchedParameters, TRefCountPtr<FRHITexture> YUY2Texture, const FIntPoint& OutputDimensions, const FMatrix& ColorTransform, const FVector& YUVOffset, bool SrgbToLinear, MediaShaders::EToneMapMethod ToneMapMethod = MediaShaders::EToneMapMethod::None);
 };
 
 /**
@@ -578,7 +587,7 @@ public:
 		: FGlobalShader(Initializer)
 	{ }
 
-	RENDERCORE_API void SetParameters(FRHIBatchedShaderParameters& BatchedParameters, FShaderResourceViewRHIRef SRV, const FIntPoint& OutputDimensions, UE::Color::EEncoding Encoding, const FMatrix44f& CSTransform);
+	RENDERCORE_API void SetParameters(FRHIBatchedShaderParameters& BatchedParameters, FShaderResourceViewRHIRef SRV, const FIntPoint& OutputDimensions, UE::Color::EEncoding Encoding, const FMatrix44f& CSTransform, MediaShaders::EToneMapMethod ToneMapMethod = MediaShaders::EToneMapMethod::None);
 };
 
 
@@ -747,5 +756,5 @@ public:
 		: FGlobalShader(Initializer)
 	{ }
 
-	RENDERCORE_API void SetParameters(FRHIBatchedShaderParameters& BatchedParameters, TRefCountPtr<FRHITexture> VYUTexture, const FIntPoint& OutputDimensions, const FMatrix44f& ColorTransform, UE::Color::EEncoding Encoding, const FMatrix44f& CSTransform);
+	RENDERCORE_API void SetParameters(FRHIBatchedShaderParameters& BatchedParameters, TRefCountPtr<FRHITexture> VYUTexture, const FIntPoint& OutputDimensions, const FMatrix44f& ColorTransform, UE::Color::EEncoding Encoding, const FMatrix44f& CSTransform, MediaShaders::EToneMapMethod ToneMapMethod = MediaShaders::EToneMapMethod::None);
 };
