[shaders] don't attempt to output debug info if preprocessing has failed

diff --git a/Engine/Source/Runtime/RenderCore/Private/ShaderCore.cpp b/Engine/Source/Runtime/RenderCore/Private/ShaderCore.cpp
--- a/Engine/Source/Runtime/RenderCore/Private/ShaderCore.cpp
+++ b/Engine/Source/Runtime/RenderCore/Private/ShaderCore.cpp
@@ -4091,6 +4091,7 @@ void FShaderCompileJob::OnComplete()
 	// running compilation to generate (these will always be output by the compile step, and always only be generated for the single
 	// job which executed compilation)
 	if (Input.DumpDebugInfoEnabled()
+		&& PreprocessOutput.bSucceeded // skip outputting debug data if preprocessing failed, nothing useful we can do here
 		// if we only want debug info for jobs which actually compiled, check the CompileTime
 		// (jobs deserialized from the cache/wait list/ddc will have a compiletime of 0.0)
 		&& (CVarDumpDebugInfoForCacheHits.GetValueOnAnyThread() || Output.CompileTime > 0.0f))
