[shaders] when preprocessing fails output a simple .usf file containing commented defines, base64 env and directcompile cmdline (the former may help to determine why preprocessing is failing, while the latter two can be used to inspect the FShaderCompilerEnvironment via running SCW in directcompile mode)

diff --git a/Engine/Source/Developer/ShaderCompilerCommon/Private/ShaderCompilerCommon.cpp b/Engine/Source/Developer/ShaderCompilerCommon/Private/ShaderCompilerCommon.cpp
--- a/Engine/Source/Developer/ShaderCompilerCommon/Private/ShaderCompilerCommon.cpp
+++ b/Engine/Source/Developer/ShaderCompilerCommon/Private/ShaderCompilerCommon.cpp
@@ -2138,6 +2138,20 @@ namespace UE::ShaderCompilerCommon
 		}
 	}
 
+
+	void FBaseShaderFormat::OutputDebugDataMinimal(const FShaderCompilerInput& Input) const
+	{
+		const FStringView FailedSourceStr = TEXTVIEW(R"(
+// Preprocessing failed for shader; defines used can be seen above.
+// DirectCompile cmdline as given at the bottom of the file can be used to inspect the contents of the compilation environment.)");
+
+		FDebugShaderDataOptions Options;
+		Options.bSourceOnly = true;
+		// Invoking this with bSourceOnly=true will dump only a .usf file; given that we don't have preprocessed source
+		// we use the FailedSourceStr to indicate what value such a file can provide.
+		DumpDebugShaderData(Input, FailedSourceStr, Options, nullptr);
+	}
+
 	void DumpExtendedDebugShaderData(
 		const FShaderCompilerInput& Input,
 		const FShaderPreprocessOutput& PreprocessOutput,
diff --git a/Engine/Source/Developer/ShaderCompilerCommon/Public/ShaderCompilerCommon.h b/Engine/Source/Developer/ShaderCompilerCommon/Public/ShaderCompilerCommon.h
--- a/Engine/Source/Developer/ShaderCompilerCommon/Public/ShaderCompilerCommon.h
+++ b/Engine/Source/Developer/ShaderCompilerCommon/Public/ShaderCompilerCommon.h
@@ -170,6 +170,8 @@ namespace UE::ShaderCompilerCommon
 			const FShaderCompilerInput& Input,
 			const FShaderPreprocessOutput& PreprocessOutput,
 			const FShaderCompilerOutput& Output) const override;
+
+		virtual SHADERCOMPILERCOMMON_API void OutputDebugDataMinimal(const FShaderCompilerInput& Input) const override;
 	};
 }
 
diff --git a/Engine/Source/Developer/TargetPlatform/Public/Interfaces/IShaderFormat.h b/Engine/Source/Developer/TargetPlatform/Public/Interfaces/IShaderFormat.h
--- a/Engine/Source/Developer/TargetPlatform/Public/Interfaces/IShaderFormat.h
+++ b/Engine/Source/Developer/TargetPlatform/Public/Interfaces/IShaderFormat.h
@@ -210,6 +210,9 @@ public:
 		const FShaderCompilerOutput& Output,
 		const FShaderCompilerOutput& SecondaryOutput) const {};
 
+	/* Implement to output a minimal set of debug info based solely on the compile input. This is invoked when preprocessing fails. */
+	virtual void OutputDebugDataMinimal(const FShaderCompilerInput& Input) const {};
+
 public:
 	/** Virtual destructor. */
 	virtual ~IShaderFormat() { }
diff --git a/Engine/Source/Runtime/RenderCore/Private/ShaderCore.cpp b/Engine/Source/Runtime/RenderCore/Private/ShaderCore.cpp
--- a/Engine/Source/Runtime/RenderCore/Private/ShaderCore.cpp
+++ b/Engine/Source/Runtime/RenderCore/Private/ShaderCore.cpp
@@ -4090,20 +4090,26 @@ void FShaderCompileJob::OnComplete()
 	// note that depending on shaderformat implementation this may not necessarily be _all_ debug artifacts, since some require
 	// running compilation to generate (these will always be output by the compile step, and always only be generated for the single
 	// job which executed compilation)
-	if (Input.DumpDebugInfoEnabled()
-		&& PreprocessOutput.bSucceeded // skip outputting debug data if preprocessing failed, nothing useful we can do here
-		// if we only want debug info for jobs which actually compiled, check the CompileTime
-		// (jobs deserialized from the cache/wait list/ddc will have a compiletime of 0.0)
-		&& (CVarDumpDebugInfoForCacheHits.GetValueOnAnyThread() || Output.CompileTime > 0.0f))
+	if (Input.DumpDebugInfoEnabled())
 	{
 		const IShaderFormat* ShaderFormat = GetTargetPlatformManagerRef().FindShaderFormat(Input.ShaderFormat);
-		if (SecondaryPreprocessOutput.IsValid() && SecondaryOutput.IsValid())
+		if (!PreprocessOutput.bSucceeded)
 		{
-			ShaderFormat->OutputDebugData(Input, PreprocessOutput, *SecondaryPreprocessOutput, Output, *SecondaryOutput);
+			// output minimal debug info for preprocessing failures
+			ShaderFormat->OutputDebugDataMinimal(Input);
 		}
-		else
+		else if (CVarDumpDebugInfoForCacheHits.GetValueOnAnyThread() || Output.CompileTime > 0.0f)
+			// if we only want debug info for jobs which actually compiled, check the CompileTime
+			// (jobs deserialized from the cache/wait list/ddc will have a compiletime of 0.0)
 		{
-			ShaderFormat->OutputDebugData(Input, PreprocessOutput, Output);
+			if (SecondaryPreprocessOutput.IsValid() && SecondaryOutput.IsValid())
+			{
+				ShaderFormat->OutputDebugData(Input, PreprocessOutput, *SecondaryPreprocessOutput, Output, *SecondaryOutput);
+			}
+			else
+			{
+				ShaderFormat->OutputDebugData(Input, PreprocessOutput, Output);
+			}
 		}
 	}
 }
