Fix UYVY color conversion.

24937677 introduced a regression, the UYVY to RGBA conversion used to be done by the UYVYConvertPS shader, which interpolates the chroma values using a bilinear sampler.
The change moved the UYVY conversion code to use YUVv216ConvertPS, which uses a point sampler by default.

This fix introduces an option (UseBilinearSampler) that dictates whether to use a bilinear sampler for the chroma values of the converter.

diff --git a/Engine/Shaders/Private/MediaShaders.usf b/Engine/Shaders/Private/MediaShaders.usf
--- a/Engine/Shaders/Private/MediaShaders.usf
+++ b/Engine/Shaders/Private/MediaShaders.usf
@@ -300,6 +300,18 @@ void YUVv216ConvertPS(
 	CbY0CrY1 = (YUVv216ConvertUB.SwapChroma != 0u) ? CbY0CrY1.zyxw : CbY0CrY1;						// Cb/Cr swap
 	float3 YUV = (uint(InUV.x * YUVv216ConvertUB.OutputDimX) & 1u) ? CbY0CrY1.wxz : CbY0CrY1.yxz;	// Decode to YCbCr for odd/even pixels
 
+	// Needed for UYVY to RGB
+	BRANCH
+	if (YUVv216ConvertUB.UseBilinearSampler != 0u)
+	{
+		float4 InSampled = YUVv216ConvertUB.YUVTexture.Sample(YUVv216ConvertUB.SamplerB, InUV);
+		InSampled = (YUVv216ConvertUB.IsARGBFmt != 0u) ? InSampled.zyxw : InSampled;										// PF_ARGB8 vs. PF_ABGR16 ordering
+		float4 CbY0CrY1Sampled = (YUVv216ConvertUB.IsCbY0CrY1 != 0u) ? InSampled : InSampled.yxwz;							// CbY0CrY1 vs. Y0CbY1Cr ordering
+		CbY0CrY1Sampled = (YUVv216ConvertUB.SwapChroma != 0u) ? CbY0CrY1Sampled.zyxw : CbY0CrY1Sampled;						// Cb/Cr swap
+		float3 YUVSampled = (uint(InUV.x * YUVv216ConvertUB.OutputDimX) & 1u) ? CbY0CrY1Sampled.wxz : CbY0CrY1Sampled.yxz;	// Decode to YCbCr for odd/even pixels
+		YUV.yz = YUVSampled.yz;
+	}
+
 	OutColor = float4(ApplyColorProcessing(YUV, YUVv216ConvertUB.ColorTransform, YUVv216ConvertUB.CSTransform, YUVv216ConvertUB.EOTF, YUVv216ConvertUB.ToneMapMethod), 1.0f); 
 }
 
diff --git a/Engine/Source/Runtime/MediaAssets/Private/Misc/MediaTextureResource.cpp b/Engine/Source/Runtime/MediaAssets/Private/Misc/MediaTextureResource.cpp
--- a/Engine/Source/Runtime/MediaAssets/Private/Misc/MediaTextureResource.cpp
+++ b/Engine/Source/Runtime/MediaAssets/Private/Misc/MediaTextureResource.cpp
@@ -1167,7 +1167,8 @@ void FMediaTextureResource::GetColorSpaceConversionMatrixForSample(const TShared
 												SampleFormat!= EMediaTextureSampleFormat::CharYUY2 && SampleFormat != EMediaTextureSampleFormat::CharYVYU,	// Y or Cb first
 												InputTexture->GetFormat() == PF_B8G8R8A8,																	// ARGB vs. ABGR (memory order)
 												SampleFormat == EMediaTextureSampleFormat::CharYVYU,														// Cb / Cr swap
-												Sample->GetToneMapMethod());
+												Sample->GetToneMapMethod(),
+												SampleFormat == EMediaTextureSampleFormat::CharUYVY);  // Whether to use bilinear sampling for the chroma values.
 				}
 				break;
 
diff --git a/Engine/Source/Runtime/RenderCore/Private/MediaShaders.cpp b/Engine/Source/Runtime/RenderCore/Private/MediaShaders.cpp
--- a/Engine/Source/Runtime/RenderCore/Private/MediaShaders.cpp
+++ b/Engine/Source/Runtime/RenderCore/Private/MediaShaders.cpp
@@ -641,6 +641,8 @@ SHADER_PARAMETER(uint32, ToneMapMethod)
 SHADER_PARAMETER(uint32, IsCbY0CrY1)
 SHADER_PARAMETER(uint32, IsARGBFmt)
 SHADER_PARAMETER(uint32, SwapChroma)
+SHADER_PARAMETER(uint32, UseBilinearSampler)
+SHADER_PARAMETER_SAMPLER(SamplerState, SamplerB)
 SHADER_PARAMETER(float, OutputDimX)
 SHADER_PARAMETER(float, OutputDimY)
 SHADER_PARAMETER_TEXTURE(Texture2D<float4>, YUVTexture)
@@ -650,7 +652,7 @@ IMPLEMENT_GLOBAL_SHADER_PARAMETER_STRUCT(FYUVv216ConvertUB, "YUVv216ConvertUB");
 IMPLEMENT_SHADER_TYPE(, FYUVv216ConvertPS, TEXT("/Engine/Private/MediaShaders.usf"), TEXT("YUVv216ConvertPS"), SF_Pixel);
 
 
-void FYUVv216ConvertPS::SetParameters(FRHIBatchedShaderParameters& BatchedParameters, TRefCountPtr<FRHITexture> YUVTexture, const FIntPoint& OutputDimensions, const FMatrix44f& ColorTransform, UE::Color::EEncoding Encoding, const FMatrix44f& CSTransform, bool bIsCbY0CrY1, bool bIsARGBFmt, bool bSwapChroma, MediaShaders::EToneMapMethod ToneMapMethod)
+void FYUVv216ConvertPS::SetParameters(FRHIBatchedShaderParameters& BatchedParameters, TRefCountPtr<FRHITexture> YUVTexture, const FIntPoint& OutputDimensions, const FMatrix44f& ColorTransform, UE::Color::EEncoding Encoding, const FMatrix44f& CSTransform, bool bIsCbY0CrY1, bool bIsARGBFmt, bool bSwapChroma, MediaShaders::EToneMapMethod ToneMapMethod, bool bUseBilinearSamplingForChroma)
 {
 	FYUVv216ConvertUB UB;
 	{
@@ -661,6 +663,8 @@ void FYUVv216ConvertPS::SetParameters(FRHIBatchedShaderParameters& BatchedParame
 		UB.IsCbY0CrY1 = bIsCbY0CrY1;
 		UB.IsARGBFmt = bIsARGBFmt;
 		UB.SwapChroma = bSwapChroma;
+		UB.UseBilinearSampler = bUseBilinearSamplingForChroma;
+		UB.SamplerB = TStaticSamplerState<SF_Bilinear>::GetRHI();
 		UB.OutputDimX = (float)OutputDimensions.X;
 		UB.OutputDimY = (float)OutputDimensions.Y;
 		UB.YUVTexture = YUVTexture;
diff --git a/Engine/Source/Runtime/RenderCore/Public/MediaShaders.h b/Engine/Source/Runtime/RenderCore/Public/MediaShaders.h
--- a/Engine/Source/Runtime/RenderCore/Public/MediaShaders.h
+++ b/Engine/Source/Runtime/RenderCore/Public/MediaShaders.h
@@ -501,7 +501,7 @@ public:
 		: FGlobalShader(Initializer)
 	{ }
 
-	RENDERCORE_API void SetParameters(FRHIBatchedShaderParameters& BatchedParameters, TRefCountPtr<FRHITexture> YUVTexture, const FIntPoint& OutputDimensions, const FMatrix44f& ColorTransform, UE::Color::EEncoding Encoding, const FMatrix44f& CSTransform, bool bIsCbY0CrY1, bool bIsARGBFmt, bool bSwapChroma, MediaShaders::EToneMapMethod ToneMapMethod = MediaShaders::EToneMapMethod::None);
+	RENDERCORE_API void SetParameters(FRHIBatchedShaderParameters& BatchedParameters, TRefCountPtr<FRHITexture> YUVTexture, const FIntPoint& OutputDimensions, const FMatrix44f& ColorTransform, UE::Color::EEncoding Encoding, const FMatrix44f& CSTransform, bool bIsCbY0CrY1, bool bIsARGBFmt, bool bSwapChroma, MediaShaders::EToneMapMethod ToneMapMethod = MediaShaders::EToneMapMethod::None, bool bUseBilinearSamplingForChroma = false);
 };
 
 
