[shaders] fix asan issue; eliminate lambdas for code serialization/reserve in serialization context subclasses in favour of virtual functions for clarity, add an explicit reset function to change internal state to point to different buffers. The code in the job cache path (processing duplicate completed jobs with the same load context) which constructed a temporary object and assigned to reset the internal pointers was creating a new lambda capturing the wrong "this" and assigning it to the original object.

diff --git a/Engine/Source/Runtime/Engine/Private/ShaderCompiler/ShaderCompilerJobCache.cpp b/Engine/Source/Runtime/Engine/Private/ShaderCompiler/ShaderCompilerJobCache.cpp
--- a/Engine/Source/Runtime/Engine/Private/ShaderCompiler/ShaderCompilerJobCache.cpp
+++ b/Engine/Source/Runtime/Engine/Private/ShaderCompiler/ShaderCompilerJobCache.cpp
@@ -1954,7 +1954,7 @@ FShaderCacheLoadContext FShaderJobCache::FindOrAdd(const FShaderCompilerInputHas
 							JobLock.WriteUnlock();
 
 							// Re-initialize load context pointing to the stored code array in the cache and reuse for each duplicate job needing population
-							LoadContext = FShaderCacheLoadContext(StoredOutput->JobOutput, StoredOutput->JobCode);
+							LoadContext.Reset(StoredOutput->JobOutput, StoredOutput->JobCode);
 
 							// Call ProcessFinishedJob on main job and duplicates
 							for (FShaderCommonCompileJob* FinishedJob : FinishedJobs)
diff --git a/Engine/Source/Runtime/RenderCore/Internal/ShaderSerialization.h b/Engine/Source/Runtime/RenderCore/Internal/ShaderSerialization.h
--- a/Engine/Source/Runtime/RenderCore/Internal/ShaderSerialization.h
+++ b/Engine/Source/Runtime/RenderCore/Internal/ShaderSerialization.h
@@ -29,7 +29,18 @@ struct FShaderSerializeContext
 	/* Constructor which accepts an FArchive reference, used for simple serialization cases */
 	FShaderSerializeContext(FArchive& InArchive) : Ar(&InArchive) {}
 
-	/* Archive pointer which should be used for serializing the object, possibly excluding shader code if the SerializeCodeFunc is set */
+	virtual ~FShaderSerializeContext() = default;
+
+	/* If this is overridden to return true, SerializeCode will be called to serialize shader code separately from the rest of the serialized object. */
+	virtual bool EnableCustomCodeSerialize() { return false; }
+
+	/* Optional function which must be implemented if EnableCustomCodeSerialize returns true; use to serialize shader code separately from the main object */
+	virtual void SerializeCode(FShaderCodeResource& Resource, int32 Index) {};
+
+	/* Optional function that can be used to reserve space for the given number of code objects by derived classes. */
+	virtual void ReserveCode(int32 Count) {};
+
+	/* Archive pointer which should be used for serializing the object, possibly excluding shader code if the EnableCustomCodeSerialize returns true */
 	FArchive* Ar = nullptr;
 
 	/* Flag indicating whether this serialization is a cooked load; used to change serialization behaviour for cooked data vs cached data */
@@ -38,11 +49,6 @@ struct FShaderSerializeContext
 	/* FName of the asset which triggered the serialization, this is used only for diagnostic messages */
 	FName SerializingAsset = NAME_None;
 
-	/* Optional function which can be used to serialize shader code separately from the main object; if this is set serialization of bytecode 
-	 * will be excluded from the main object (i.e. not using the Ar member archive reference) and is the responsibility of this function.
-	 */
-	TFunction<void(FShaderCodeResource&, int32)> SerializeCodeFunc{};
-
 	/* Optional function which can be used to reserve space for a number of code entries as a minor optimization. */
 	TFunction<void(int32)> ReserveCodeFunc{};
 
@@ -65,6 +71,8 @@ struct FShaderCacheSerializeContext : public FShaderSerializeContext
 		ShaderCode = OwnedShaderCode;
 	}
 
+	virtual ~FShaderCacheSerializeContext() = default;
+
 	/* Buffer which stores the main object data for a cache entry - i.e. a shadermap or job structure */
 	FSharedBuffer ShaderObjectData;
 
@@ -78,6 +86,9 @@ struct FShaderCacheSerializeContext : public FShaderSerializeContext
 	 */
 	TArray<FCompositeBuffer> OwnedShaderCode;
 
+	/* Subclasses of this type must implement custom code serialization function */
+	virtual bool EnableCustomCodeSerialize() override { return true; }
+
 	/* Get the total serialized size of data for this context; note that this will return 0 if called prior to the FSharedBuffers
 	 * being set (this is done in the derived implementations, see below).
 	 */
@@ -110,16 +121,22 @@ struct FShaderCacheSerializeContext : public FShaderSerializeContext
 /* Implementation of FShaderCacheSerializeContext used for saving data to caches. */
 struct FShaderCacheSaveContext : public FShaderCacheSerializeContext
 {
-	/* Default constructor; sets up base class TFunctions to save bytecode independently of object structure, as
-	 * well as base class FArchive pointing to the owned memory writer.
-	 */
+	/* Default constructor; sets up base class FArchive pointing to the owned memory writer. */
 	RENDERCORE_API FShaderCacheSaveContext();
 
+	virtual ~FShaderCacheSaveContext() = default;
+
 	/* Converts the raw serialized object data into the ShaderObjectData FSharedBuffer.
 	 * Note that this is called by BuildCacheRecord as well, calls subsequent to the first will have no effect.
 	 */
 	RENDERCORE_API void Finalize();
 
+	/* Overridden code serialize implementation. */
+	RENDERCORE_API virtual void SerializeCode(FShaderCodeResource& Resource, int32 Index) override;
+
+	/* Overridden code reserve implementation. */
+	RENDERCORE_API virtual void ReserveCode(int32 Count) override;
+
 #if WITH_EDITOR
 	/* Helper function to generate a DDC record from the data serialized using this context. Serialization must
 	 * be executed prior to calling this function.
@@ -141,13 +158,22 @@ struct FShaderCacheSaveContext : public FShaderCacheSerializeContext
 struct FShaderCacheLoadContext : public FShaderCacheSerializeContext
 {
 	/*	Default constructor, use when array of code buffers will be allocated via ReadFromRecord */
-	RENDERCORE_API FShaderCacheLoadContext();
+	FShaderCacheLoadContext() = default;
+
 	/* Constructor which references buffers (and an array of code buffers) owned elsewhere, does not allocate the OwnedShaderCode array */
 	RENDERCORE_API FShaderCacheLoadContext(FSharedBuffer ShaderObjectData, TArrayView<FCompositeBuffer> CodeBuffers);
 
+	virtual ~FShaderCacheLoadContext() = default;
+
+	/* Resets internal state to point to the given buffers and recreates the owned memoryreader */
+	RENDERCORE_API void Reset(FSharedBuffer ShaderObjectData, TArrayView<FCompositeBuffer> CodeBuffers);
+
 	/* Call to reset reader to start position so the same load context can be used to populate multiple objects. */
 	RENDERCORE_API void Reuse();
 
+	/* Overridden code serialize implementation. */
+	RENDERCORE_API virtual void SerializeCode(FShaderCodeResource& Resource, int32 Index) override;
+
 #if WITH_EDITOR
 	/* Helper function to populate the internal state (FSharedBuffers defined on FShaderCacheSerializeContext) from a DDC record. */
 	RENDERCORE_API void ReadFromRecord(const UE::DerivedData::FCacheRecord& CacheRecord, bool bIsPersistent = false);
diff --git a/Engine/Source/Runtime/RenderCore/Private/ShaderCore.cpp b/Engine/Source/Runtime/RenderCore/Private/ShaderCore.cpp
--- a/Engine/Source/Runtime/RenderCore/Private/ShaderCore.cpp
+++ b/Engine/Source/Runtime/RenderCore/Private/ShaderCore.cpp
@@ -4049,7 +4049,8 @@ void FShaderCompileJob::SerializeOutput(FShaderCacheSerializeContext& Ctx, int32
 		CodeResource = Output.GetFinalizedCodeResource();
 	}
 
-	Ctx.SerializeCodeFunc(CodeResource, CodeIndex);
+	check(Ctx.EnableCustomCodeSerialize());
+	Ctx.SerializeCode(CodeResource, CodeIndex);
 
 	// we intentionally re-set the internal ShaderCode even when saving; GetCodeResource moves the code array into the
 	// FShaderCodeResource's internal array and this moves it back (preventing an unnecessary temporary copy of the code)
diff --git a/Engine/Source/Runtime/RenderCore/Private/ShaderResource.cpp b/Engine/Source/Runtime/RenderCore/Private/ShaderResource.cpp
--- a/Engine/Source/Runtime/RenderCore/Private/ShaderResource.cpp
+++ b/Engine/Source/Runtime/RenderCore/Private/ShaderResource.cpp
@@ -363,7 +363,7 @@ void FShaderMapResourceCode::Serialize(FShaderSerializeContext& Ctx)
 	FArchive& Ar = Ctx.GetMainArchive();
 	Ar << ResourceHash;
 	Ar << ShaderHashes;
-	if (!Ctx.SerializeCodeFunc)
+	if (!Ctx.EnableCustomCodeSerialize())
 	{
 		Ar << ShaderCodeResources;
 	}
@@ -381,7 +381,7 @@ void FShaderMapResourceCode::Serialize(FShaderSerializeContext& Ctx)
 
 		for (int32 CodeIndex = 0; CodeIndex < ShaderCodeResources.Num(); ++CodeIndex)
 		{
-			Ctx.SerializeCodeFunc(ShaderCodeResources[CodeIndex], CodeIndex);
+			Ctx.SerializeCode(ShaderCodeResources[CodeIndex], CodeIndex);
 		}
 	}
 	check(ShaderCodeResources.Num() == ShaderHashes.Num());
diff --git a/Engine/Source/Runtime/RenderCore/Private/ShaderSerialization.cpp b/Engine/Source/Runtime/RenderCore/Private/ShaderSerialization.cpp
--- a/Engine/Source/Runtime/RenderCore/Private/ShaderSerialization.cpp
+++ b/Engine/Source/Runtime/RenderCore/Private/ShaderSerialization.cpp
@@ -12,18 +12,6 @@
 
 FShaderCacheSaveContext::FShaderCacheSaveContext()
 {
-	SerializeCodeFunc = [this](FShaderCodeResource& Res, int32 Index)
-		{
-			OwnedShaderCode.Add(Res.GetCacheBuffer());
-			// reset the array view any time an entry is added; we do this instead of calling Resize in the reserve delegate
-			// and setting it there since not all code paths (i.e. single job cache records) call reserve
-			ShaderCode = OwnedShaderCode;
-		};
-
-	ReserveCodeFunc = [this](int32 Num)
-		{
-			OwnedShaderCode.Reserve(Num);
-		};
 	Reset();
 }
 
@@ -33,6 +21,19 @@ const UE::DerivedData::FValueId ShaderCodeDataValue = UE::DerivedData::FValueId:
 const FAnsiStringView CodeCountMetaField = ANSITEXTVIEW("CodeCount");
 #endif
 
+void FShaderCacheSaveContext::SerializeCode(FShaderCodeResource& Resource, int32 Index)
+{
+	OwnedShaderCode.Add(Resource.GetCacheBuffer());
+	// reset the array view any time an entry is added; we do this instead of calling Resize in the reserve delegate
+	// and setting it there since not all code paths (i.e. single job cache records) call reserve
+	ShaderCode = OwnedShaderCode;
+}
+
+void FShaderCacheSaveContext::ReserveCode(int32 Count)
+{
+	OwnedShaderCode.Reserve(Count);
+}
+
 void FShaderCacheSaveContext::Reset()
 {
 	ShaderObjectData.Reset();
@@ -76,15 +77,13 @@ UE::DerivedData::FCacheRecord FShaderCacheSaveContext::BuildCacheRecord(const UE
 }
 #endif
 
-FShaderCacheLoadContext::FShaderCacheLoadContext()
+
+FShaderCacheLoadContext::FShaderCacheLoadContext(FSharedBuffer InShaderObjectData, TArrayView<FCompositeBuffer> InCodeBuffers)
 {
-	SerializeCodeFunc = [this](FShaderCodeResource& Res, int32 Index)
-		{
-			Res.PopulateFromComposite(ShaderCode[Index]);
-		};
+	Reset(InShaderObjectData, InCodeBuffers);
 }
 
-FShaderCacheLoadContext::FShaderCacheLoadContext(FSharedBuffer InShaderObjectData, TArrayView<FCompositeBuffer> InCodeBuffers) : FShaderCacheLoadContext()
+void FShaderCacheLoadContext::Reset(FSharedBuffer InShaderObjectData, TArrayView<FCompositeBuffer> InCodeBuffers)
 {
 	ShaderObjectData = InShaderObjectData;
 	ShaderCode = InCodeBuffers;
@@ -93,6 +92,11 @@ FShaderCacheLoadContext::FShaderCacheLoadContext(FSharedBuffer InShaderObjectDat
 	Ar = Reader.Get();
 }
 
+void FShaderCacheLoadContext::SerializeCode(FShaderCodeResource& Resource, int32 Index)
+{
+	Resource.PopulateFromComposite(ShaderCode[Index]);
+}
+
 void FShaderCacheLoadContext::Reuse()
 {
 	Reader->Seek(0u);
