Fix PreloadShaderMap when DynamicShaderPreloading is enable. Go through the preloadshader code path when we want to load the entire shadermap.This make sure we only increase the NumRef once and the reloadentry will be release when either we release the RHI shaders, destroy the material instance or we call ReleasePreloadShaderIfNecessary.

diff --git a/Engine/Source/Runtime/Engine/Private/Materials/Material.cpp b/Engine/Source/Runtime/Engine/Private/Materials/Material.cpp
--- a/Engine/Source/Runtime/Engine/Private/Materials/Material.cpp
+++ b/Engine/Source/Runtime/Engine/Private/Materials/Material.cpp
@@ -4465,21 +4465,24 @@ void UMaterial::PostLoad()
 	}
 #endif // WITH_EDITOR
 
-	if (IsDeferredDecal())
+	
+
+	if(IsPSOShaderPreloadingEnabled())
+	{ 
+		// When dynamic preload shaders is enabled, we need to prelaod some material domains since there is no
+		// code logic within the PSO precaching system.
+		if (IsUIMaterial() || IsDeferredDecal() || (MaterialDomain == MD_LightFunction) || IsPostProcessMaterial())
+		{
+			FGraphEventArray Unused;
+			PreloadMaterialShaderMap(GetMaterialResource(GMaxRHIFeatureLevel), Unused);
+		}
+	}
+	else if (IsDeferredDecal())
 	{
 		FPSOPrecacheParams PSOPrecacheParams;
 		UMaterialInterface::PrecachePSOs(&FLocalVertexFactory::StaticType, PSOPrecacheParams);
 	}
 
-
-	// When dynamic preload shaders is enabled, we need to prelaod some material domains since there is no
-	// code logic within the PSO precaching system.
-	if (IsUIMaterial() || IsDeferredDecal() || (MaterialDomain == MD_LightFunction) || IsPostProcessMaterial())
-	{
-		FGraphEventArray Unused;
-		PreloadMaterialShaderMap(GetMaterialResource(GMaxRHIFeatureLevel), Unused);
-	}
-
 }
 
 #if WITH_EDITORONLY_DATA
diff --git a/Engine/Source/Runtime/Engine/Private/Materials/MaterialInstance.cpp b/Engine/Source/Runtime/Engine/Private/Materials/MaterialInstance.cpp
--- a/Engine/Source/Runtime/Engine/Private/Materials/MaterialInstance.cpp
+++ b/Engine/Source/Runtime/Engine/Private/Materials/MaterialInstance.cpp
@@ -3463,16 +3463,20 @@ void UMaterialInstance::PostLoad()
 		LightingGuidFixupMap.Add(GetLightingGuid(), this);
 	}
 
-	if (IsDeferredDecal())
+	if (IsPSOShaderPreloadingEnabled())
 	{
-		FPSOPrecacheParams PSOPrecacheParams;
-		UMaterialInterface::PrecachePSOs(&FLocalVertexFactory::StaticType, PSOPrecacheParams);
+		// When dynamic preload shaders is enabled, we need to prelaod some material domains since there is no
+		// code logic within the PSO precaching system.
+		if (IsUIMaterial() || IsDeferredDecal() || IsPostProcessMaterial())
+		{
+			FGraphEventArray Unused;
+			PreloadMaterialShaderMap(GetMaterialResource(GMaxRHIFeatureLevel), Unused);
+		}
 	}
-
-	if (IsUIMaterial()  || IsDeferredDecal() || IsPostProcessMaterial())
+	else if (IsDeferredDecal())
 	{
-		FGraphEventArray Unused;
-		PreloadMaterialShaderMap(GetMaterialResource(GMaxRHIFeatureLevel), Unused);
+		FPSOPrecacheParams PSOPrecacheParams;
+		UMaterialInterface::PrecachePSOs(&FLocalVertexFactory::StaticType, PSOPrecacheParams);
 	}
 	//DumpDebugInfo(*GLog);
 }
diff --git a/Engine/Source/Runtime/RHI/Public/RHIShaderLibrary.h b/Engine/Source/Runtime/RHI/Public/RHIShaderLibrary.h
--- a/Engine/Source/Runtime/RHI/Public/RHIShaderLibrary.h
+++ b/Engine/Source/Runtime/RHI/Public/RHIShaderLibrary.h
@@ -24,6 +24,7 @@ public:
 	virtual int32 GetNumShaders() const = 0;
 	virtual int32 GetNumShadersForShaderMap(int32 ShaderMapIndex) const = 0;
 	virtual int32 GetShaderIndex(int32 ShaderMapIndex, int32 i) const = 0;
+	virtual void GetAllShaderIndices(int32 ShaderMapIndex, TArray<int32>& ShaderIndices) {}
 	virtual uint32 GetSizeBytes() const = 0;
 	virtual FSHAHash GetShaderHash(int32 ShaderMapIndex, int32 ShaderIndex) = 0;
 	virtual int32 FindShaderMapIndex(const FSHAHash& Hash) = 0;
diff --git a/Engine/Source/Runtime/RenderCore/Private/ShaderCodeLibrary.cpp b/Engine/Source/Runtime/RenderCore/Private/ShaderCodeLibrary.cpp
--- a/Engine/Source/Runtime/RenderCore/Private/ShaderCodeLibrary.cpp
+++ b/Engine/Source/Runtime/RenderCore/Private/ShaderCodeLibrary.cpp
@@ -1415,6 +1415,17 @@ void FShaderMapResource_SharedCode::PreloadShader(int32 ShaderIndex, FGraphEvent
 	LibraryInstance->PreloadShader(LibraryShaderIndex, OutCompletionEvents);
 }
 
+void FShaderMapResource_SharedCode::PreloadShaderMap(FGraphEventArray& OutCompletionEvents)
+{
+	TArray<int32> ShaderIndices;
+	LibraryInstance->Library->GetAllShaderIndices(ShaderMapIndex, ShaderIndices);
+
+	for (int32 ShaderIndex : ShaderIndices)
+	{
+		LibraryInstance->PreloadShader(ShaderIndex, OutCompletionEvents);
+	}
+}
+
 void FShaderMapResource_SharedCode::ReleaseRHI()
 {
 	if (LibraryInstance && ensureMsgf(LibraryInstance->Library, TEXT("LibraryInstance->Library pointer is expected to be valid as long as library's FShaderMapResource are alive.")))
@@ -1489,11 +1500,6 @@ FString FShaderMapResource_SharedCode::GetFriendlyName() const
 	return LibraryInstance->Library->GetName();
 }
 
-void FShaderMapResource_SharedCode::PreloadShaderMap(FGraphEventArray& OutCompletionEvents)
-{
-	bEntireShaderMapPreloaded = LibraryInstance->Library->PreloadShaderMap(ShaderMapIndex, OutCompletionEvents);
-}
-
 int32 FShaderMapResource_SharedCode::GetGroupIndexForShader(int32 ShaderIndex) const
 {
 	const int32 LibraryShaderIndex = LibraryInstance->Library->GetShaderIndex(ShaderMapIndex, ShaderIndex);
diff --git a/Engine/Source/Runtime/RenderCore/Public/ShaderCodeArchive.h b/Engine/Source/Runtime/RenderCore/Public/ShaderCodeArchive.h
--- a/Engine/Source/Runtime/RenderCore/Public/ShaderCodeArchive.h
+++ b/Engine/Source/Runtime/RenderCore/Public/ShaderCodeArchive.h
@@ -383,6 +383,15 @@ public:
 		return SerializedShaders.GetShaderIndices()[ShaderMapEntry.ShaderIndicesOffset + i];
 	}
 
+	virtual void GetAllShaderIndices(int32 ShaderMapIndex, TArray<int32>& ShaderIndices)
+	{
+		const FShaderMapEntry& ShaderMapEntry = SerializedShaders.GetShaderMapEntries()[ShaderMapIndex];
+		for (uint32 i = 0u; i < ShaderMapEntry.NumShaders; ++i)
+		{
+			ShaderIndices.AddUnique(ShaderIndices[ShaderMapEntry.ShaderIndicesOffset + i]);
+		}
+	}
+
 	virtual int32 FindShaderMapIndex(const FSHAHash& Hash) override
 	{
 		return SerializedShaders.FindShaderMap(Hash);
@@ -643,6 +652,15 @@ public:
 		return Header.ShaderIndices[ShaderMapEntry.ShaderIndicesOffset + i];
 	}
 
+	virtual void GetAllShaderIndices(int32 ShaderMapIndex, TArray<int32>& ShaderIndices)
+	{
+		const FIoStoreShaderMapEntry& ShaderMapEntry = Header.ShaderMapEntries[ShaderMapIndex];
+		for (uint32 i = 0u; i < ShaderMapEntry.NumShaders; ++i)
+		{
+			ShaderIndices.AddUnique(Header.ShaderIndices[ShaderMapEntry.ShaderIndicesOffset + i]);
+		}
+	}
+
 	virtual int32 FindShaderMapIndex(const FSHAHash& Hash) override;
 	virtual int32 FindShaderIndex(const FSHAHash& Hash) override;
 	virtual FSHAHash GetShaderHash(int32 ShaderMapIndex, int32 ShaderIndex) override
