Fixed display gamma issues with backbuffer format renderer setting.

When Float RGBA is selected, we update DisplayGamma and (view family) SceneCaptureSource to output linear colors matching the regular R10G10B10A2 appearance.

diff --git a/Engine/Source/Runtime/Engine/Classes/Engine/RendererSettings.h b/Engine/Source/Runtime/Engine/Classes/Engine/RendererSettings.h
--- a/Engine/Source/Runtime/Engine/Classes/Engine/RendererSettings.h
+++ b/Engine/Source/Runtime/Engine/Classes/Engine/RendererSettings.h
@@ -851,7 +851,7 @@ class URendererSettings : public UDeveloperSettings
 
 	UPROPERTY(config, EditAnywhere, AdvancedDisplay, Category = DefaultSettings, meta = (
 		ConsoleVariable = "r.DefaultBackBufferPixelFormat",DisplayName = "Frame Buffer Pixel Format",
-		ToolTip = "Pixel format used for back buffer, when not specified",
+		ToolTip = "Pixel format used for back buffer, when not specified. This setting is currently incompatible with variable DisplayGamma and HDR mode.",
 		ConfigRestartRequired=true))
 	TEnumAsByte<EDefaultBackBufferPixelFormat::Type> DefaultBackBufferPixelFormat;
 	
diff --git a/Engine/Source/Runtime/Engine/Private/RendererSettings.cpp b/Engine/Source/Runtime/Engine/Private/RendererSettings.cpp
--- a/Engine/Source/Runtime/Engine/Private/RendererSettings.cpp
+++ b/Engine/Source/Runtime/Engine/Private/RendererSettings.cpp
@@ -4,6 +4,7 @@
 #include "GPUSkinVertexFactory.h"
 #include "ColorManagement/ColorSpace.h"
 #include "HAL/PlatformFile.h"
+#include "HDRHelper.h"
 #include "SceneManagement.h"
 #include "Misc/App.h"
 #include "Misc/Paths.h"
@@ -481,6 +482,11 @@ bool URendererSettings::CanEditChange(const FProperty* InProperty) const
 		return MobileShadingPath.GetValue() == 0;
 	}
 
+	if ((InProperty->GetFName() == GET_MEMBER_NAME_CHECKED(URendererSettings, DefaultBackBufferPixelFormat)))
+	{
+		return !IsHDREnabled();
+	}
+
 	PRAGMA_DISABLE_DEPRECATION_WARNINGS
 	// only allow changing ExtendDefaultLuminanceRange if it was disabled.
 	// we don't want new projects disabling this setting but still allow existing projects to enable it.
diff --git a/Engine/Source/Runtime/Engine/Private/SceneView.cpp b/Engine/Source/Runtime/Engine/Private/SceneView.cpp
--- a/Engine/Source/Runtime/Engine/Private/SceneView.cpp
+++ b/Engine/Source/Runtime/Engine/Private/SceneView.cpp
@@ -28,6 +28,7 @@
 #include "Camera/CameraTypes.h"
 #include "UObject/Interface.h"
 #include "TextureResource.h"
+#include "HDRHelper.h"
 
 DEFINE_LOG_CATEGORY(LogBufferVisualization);
 DEFINE_LOG_CATEGORY(LogNaniteVisualization);
@@ -2989,6 +2990,13 @@ PRAGMA_ENABLE_DEPRECATION_WARNINGS
 	bOverrideVirtualTextureThrottle = false;
 	VirtualTextureFeedbackFactor = GVirtualTextureFeedbackFactor;
 
+	if (IsDefaultBackBufferLinearSDR())
+	{
+		// We preserve linear color values when users select the Float RGBA frame buffer pixel format.
+		// Note: This could be removed if our tonemap pass correctly handled DisplayGamma 1.0 (instead of asymmetrically inverting gamma 2.2 before applying sRGB).
+		SceneCaptureSource = SCS_FinalToneCurveHDR;
+	}
+
 #if !WITH_EDITOR
 	check(!EngineShowFlags.StationaryLightOverlap);
 #else
diff --git a/Engine/Source/Runtime/Engine/Private/UnrealEngine.cpp b/Engine/Source/Runtime/Engine/Private/UnrealEngine.cpp
--- a/Engine/Source/Runtime/Engine/Private/UnrealEngine.cpp
+++ b/Engine/Source/Runtime/Engine/Private/UnrealEngine.cpp
@@ -2149,6 +2149,11 @@ void UEngine::Init(IEngineLoop* InEngineLoop)
 	SelectedMaterialColor = DefaultSelectedMaterialColor;
 	SelectionOutlineColor = DefaultSelectedMaterialColor;
 
+	if (IsDefaultBackBufferLinearSDR())
+	{
+		DisplayGamma = 1.0f;
+	}
+
 #if WITH_EDITOR
 	FModuleManager::Get().LoadModule("WorldPartitionEditor");
 #endif
diff --git a/Engine/Source/Runtime/RenderCore/Private/RenderCore.cpp b/Engine/Source/Runtime/RenderCore/Private/RenderCore.cpp
--- a/Engine/Source/Runtime/RenderCore/Private/RenderCore.cpp
+++ b/Engine/Source/Runtime/RenderCore/Private/RenderCore.cpp
@@ -462,6 +462,18 @@ RENDERCORE_API bool IsHDRAllowed()
 	return (CVarAllowHDR.GetValueOnAnyThread() != 0);
 }
 
+RENDERCORE_API bool IsDefaultBackBufferLinearSDR()
+{
+	// Note: Renderer settings console variable requires project restart, so we can statically cache this check.
+	static const bool bIsDefaultBackBufferLinear = []() -> bool
+		{
+			const TConsoleVariableData<int32>* CVarDefaultBackBufferPixelFormat = IConsoleManager::Get().FindTConsoleVariableDataInt(TEXT("r.DefaultBackBufferPixelFormat"));
+			return CVarDefaultBackBufferPixelFormat ? CVarDefaultBackBufferPixelFormat->GetValueOnAnyThread() == 3 /* EDefaultBackBufferPixelFormat::DBBPF_FloatRGBA */ : false;
+		}();
+	
+	return bIsDefaultBackBufferLinear && !IsHDREnabled();
+}
+
 RENDERCORE_API EDisplayOutputFormat HDRGetDefaultDisplayOutputFormat()
 {
 	return static_cast<EDisplayOutputFormat>(FMath::Clamp(CVarDisplayOutputDevice.GetValueOnAnyThread(), 0, static_cast<int32>(EDisplayOutputFormat::MAX) - 1));
diff --git a/Engine/Source/Runtime/RenderCore/Public/HDRHelper.h b/Engine/Source/Runtime/RenderCore/Public/HDRHelper.h
--- a/Engine/Source/Runtime/RenderCore/Public/HDRHelper.h
+++ b/Engine/Source/Runtime/RenderCore/Public/HDRHelper.h
@@ -7,6 +7,8 @@
 RENDERCORE_API bool IsHDREnabled();
 RENDERCORE_API bool IsHDRAllowed();
 
+RENDERCORE_API bool IsDefaultBackBufferLinearSDR();
+
 RENDERCORE_API void HDRGetMetaData(EDisplayOutputFormat& OutDisplayOutputFormat, EDisplayColorGamut& OutDisplayColorGamut, bool& OutbHDRSupported,
 								   const FVector2D& WindowTopLeft, const FVector2D& WindowBottomRight, void* OSWindow);
 
