Added explicit Pipe wait inside of RDG to avoid speculative race condition between task completion and clearing the active task in the pipe.

diff --git a/Engine/Source/Runtime/RenderCore/Private/RenderGraphBuilder.cpp b/Engine/Source/Runtime/RenderCore/Private/RenderGraphBuilder.cpp
--- a/Engine/Source/Runtime/RenderCore/Private/RenderGraphBuilder.cpp
+++ b/Engine/Source/Runtime/RenderCore/Private/RenderGraphBuilder.cpp
@@ -1648,9 +1648,9 @@ void FRDGBuilder::Compile()
 
 void FRDGBuilder::LaunchAsyncSetupQueueTask()
 {
-	if (AsyncSetupQueue.LastTask.IsCompleted())
+	if (!AsyncSetupQueue.Pipe.HasWork())
 	{
-		AsyncSetupQueue.LastTask = AsyncSetupQueue.Pipe.Launch(UE_SOURCE_LOCATION, [this]() mutable
+		AsyncSetupQueue.Pipe.Launch(UE_SOURCE_LOCATION, [this]() mutable
 		{
 			ProcessAsyncSetupQueue();
 		});
@@ -1766,7 +1766,7 @@ void FRDGBuilder::Execute()
 
 		if (ParallelSetup.bEnabled)
 		{
-			AsyncSetupQueue.LastTask.Wait();
+			AsyncSetupQueue.Pipe.WaitUntilEmpty();
 			ProcessAsyncSetupQueue();
 		}
 
diff --git a/Engine/Source/Runtime/RenderCore/Public/RenderGraphBuilder.h b/Engine/Source/Runtime/RenderCore/Public/RenderGraphBuilder.h
--- a/Engine/Source/Runtime/RenderCore/Public/RenderGraphBuilder.h
+++ b/Engine/Source/Runtime/RenderCore/Public/RenderGraphBuilder.h
@@ -939,7 +939,6 @@ private:
 
 		UE::FMutex Mutex;
 		TArray<FAsyncSetupOp, FRDGArrayAllocator> Ops;
-		UE::Tasks::FTask LastTask;
 		UE::Tasks::FPipe Pipe{ TEXT("FRDGBuilder::AsyncSetupQueue") };
 
 	} AsyncSetupQueue;
