Fixed incorrect transition of unused texture subresources. Use the prologue pass instead of the first pass since the first pass might be on the async compute pipe.

diff --git a/Engine/Source/Runtime/RenderCore/Private/RenderGraphBuilder.cpp b/Engine/Source/Runtime/RenderCore/Private/RenderGraphBuilder.cpp
--- a/Engine/Source/Runtime/RenderCore/Private/RenderGraphBuilder.cpp
+++ b/Engine/Source/Runtime/RenderCore/Private/RenderGraphBuilder.cpp
@@ -3938,7 +3938,7 @@ void FRDGBuilder::AddLastTextureTransition(FRDGTexture* Texture)
 			if (!SubresourceStateBefore)
 			{
 				SubresourceStateBefore = AllocSubresource();
-				SubresourceStateBefore->SetPass(GetPassPipeline(Texture->FirstPass), Texture->FirstPass);
+				SubresourceStateBefore->SetPass(ERHIPipeline::Graphics, GetProloguePassHandle());
 			}
 
 			State = SubresourceStateBefore;
