Separate parallel execution bit flags to fix race condition. bCulled races with bParallelExecute, since bParallelExecute is assigned from a task while bCulled is read from another.

diff --git a/Engine/Source/Runtime/RenderCore/Public/RenderGraphPass.h b/Engine/Source/Runtime/RenderCore/Public/RenderGraphPass.h
--- a/Engine/Source/Runtime/RenderCore/Public/RenderGraphPass.h
+++ b/Engine/Source/Runtime/RenderCore/Public/RenderGraphPass.h
@@ -411,30 +411,30 @@ protected:
 		struct
 		{
 			/** Whether the render pass begin / end should be skipped. */
-			uint32 bSkipRenderPassBegin : 1;
-			uint32 bSkipRenderPassEnd : 1;
+			uint16 bSkipRenderPassBegin : 1;
+			uint16 bSkipRenderPassEnd : 1;
 
 			/** (AsyncCompute only) Whether this is the first / last async compute pass in an async interval. */
-			uint32 bAsyncComputeBegin : 1;
-			uint32 bAsyncComputeEnd : 1;
+			uint16 bAsyncComputeBegin : 1;
+			uint16 bAsyncComputeEnd : 1;
 
 			/** (Graphics only) Whether this is a graphics fork / join pass. */
-			uint32 bGraphicsFork : 1;
-			uint32 bGraphicsJoin : 1;
+			uint16 bGraphicsFork : 1;
+			uint16 bGraphicsJoin : 1;
 
 			/** Whether the pass only writes to resources in its render pass. */
-			uint32 bRenderPassOnlyWrites : 1;
+			uint16 bRenderPassOnlyWrites : 1;
 
 			/** Whether this pass is a sentinel (prologue / epilogue) pass. */
-			uint32 bSentinel : 1;
+			uint16 bSentinel : 1;
 
 			/** If set, dispatches to the RHI thread after executing this pass. */
-			uint32 bDispatchAfterExecute : 1;
+			uint16 bDispatchAfterExecute : 1;
 
 			/** If set, this is a dispatch pass. */
-			uint32 bDispatchPass : 1;
+			uint16 bDispatchPass : 1;
 		};
-		uint32 PackedBits1 = 0;
+		uint16 PackedBits = 0;
 	};
 
 	union
@@ -442,26 +442,34 @@ protected:
 		// Task-specific bits which are written in a task in parallel with reads from the other set.
 		struct
 		{
-			/** If set, marks the begin / end of a span of passes executed in parallel in a task. */
-			uint32 bParallelExecuteBegin : 1;
-			uint32 bParallelExecuteEnd : 1;
-
-			/** If set, marks that a pass is executing in parallel. */
-			uint32 bParallelExecute : 1;
-
 			/** Whether this pass does not contain parameters. */
-			uint32 bEmptyParameters : 1;
+			uint8 bEmptyParameters : 1;
 
 			/** Whether this pass has external UAVs that are not tracked by RDG. */
-			uint32 bHasExternalOutputs : 1;
+			uint8 bHasExternalOutputs : 1;
 
 			/** Whether this pass has been culled. */
-			uint32 bCulled : 1;
+			uint8 bCulled : 1;
 
 			/** Whether this pass is used for external access transitions. */
-			uint32 bExternalAccessPass : 1;
+			uint8 bExternalAccessPass : 1;
+		};
+		uint8 PacketBits_AsyncSetupQueue = 0;
+	};
+
+	union
+	{
+		// Task-specific bits which are written in a task in parallel with reads from the other set.
+		struct
+		{
+			/** If set, marks the begin / end of a span of passes executed in parallel in a task. */
+			uint8 bParallelExecuteBegin : 1;
+			uint8 bParallelExecuteEnd : 1;
+
+			/** If set, marks that a pass is executing in parallel. */
+			uint8 bParallelExecute : 1;
 		};
-		uint32 PacketBits2 = 0;
+		uint8 PacketBits_ParallelExecute = 0;
 	};
 
 	/** Handle of the latest cross-pipeline producer. */
