Turn off RDG async tasks due to bad interactions with the old GPU profiler system.

RDG async tasks reference GPU profiler queries. These queries are deleted in EndFrameRenderThread, but async tasks can run after that, resulting in use-after-free crashes.

#rnx

diff --git a/Engine/Source/Runtime/RenderCore/Private/RenderGraphPrivate.cpp b/Engine/Source/Runtime/RenderCore/Private/RenderGraphPrivate.cpp
--- a/Engine/Source/Runtime/RenderCore/Private/RenderGraphPrivate.cpp
+++ b/Engine/Source/Runtime/RenderCore/Private/RenderGraphPrivate.cpp
@@ -342,13 +342,13 @@ FAutoConsoleVariableRef CVarRDGParallelSetup(
 	TEXT(" 1: pass setup is done asynchronously (default);"),
 	ECVF_RenderThreadSafe);
 
-int32 GRDGParallelExecute = 2;
+int32 GRDGParallelExecute = 1;
 FAutoConsoleVariableRef CVarRDGParallelExecute(
 	TEXT("r.RDG.ParallelExecute"), GRDGParallelExecute,
 	TEXT("Whether to enable parallel execution of passes when supported.")
-	TEXT(" 0: off;")
-	TEXT(" 1: parallel with all tasks awaited)")
-	TEXT(" 2: parallel with async tasks) (default)"),
+	TEXT(" 0: off")
+	TEXT(" 1: parallel with all tasks awaited (default)")
+	TEXT(" 2: parallel with async tasks"),
 	FConsoleVariableDelegate::CreateLambda([](IConsoleVariable* Variable)
 	{
 		if (Variable->GetInt())
