Use scene renderer shader platform in RDG builder rather than the max one. This makes the builder behave correctly in mobile preview modes.

diff --git a/Engine/Source/Runtime/RenderCore/Private/RenderGraphBuilder.cpp b/Engine/Source/Runtime/RenderCore/Private/RenderGraphBuilder.cpp
--- a/Engine/Source/Runtime/RenderCore/Private/RenderGraphBuilder.cpp
+++ b/Engine/Source/Runtime/RenderCore/Private/RenderGraphBuilder.cpp
@@ -456,7 +456,7 @@ bool FRDGBuilder::IsImmediateMode()
 	return ::IsImmediateMode();
 }
 
-ERDGPassFlags FRDGBuilder::OverridePassFlags(const TCHAR* PassName, ERDGPassFlags PassFlags)
+ERDGPassFlags FRDGBuilder::OverridePassFlags(const TCHAR* PassName, ERDGPassFlags PassFlags) const
 {
 	const bool bDebugAllowedForPass =
 #if RDG_ENABLE_DEBUG
@@ -465,7 +465,7 @@ ERDGPassFlags FRDGBuilder::OverridePassFlags(const TCHAR* PassName, ERDGPassFlag
 		true;
 #endif
 
-	if (IsAsyncComputeSupported())
+	if (bSupportsAsyncCompute)
 	{
 		if (EnumHasAnyFlags(PassFlags, ERDGPassFlags::Compute) && GRDGAsyncCompute == RDG_ASYNC_COMPUTE_FORCE_ENABLED)
 		{
@@ -569,8 +569,8 @@ bool FRDGBuilder::IsTransientInternal(FRDGViewableResource* Resource, bool bFast
 	return true;
 }
 
-FRDGBuilder::FRDGBuilder(FRHICommandListImmediate& InRHICmdList, FRDGEventName InName, ERDGBuilderFlags InFlags)
-	: FRDGScopeState(InRHICmdList, IsImmediateMode(), ::IsParallelExecuteEnabled() && EnumHasAnyFlags(InFlags, ERDGBuilderFlags::ParallelExecute))
+FRDGBuilder::FRDGBuilder(FRHICommandListImmediate& InRHICmdList, FRDGEventName InName, ERDGBuilderFlags InFlags, EShaderPlatform ShaderPlatform)
+	: FRDGScopeState(InRHICmdList, IsImmediateMode(), ::IsParallelExecuteEnabled(ShaderPlatform) && EnumHasAnyFlags(InFlags, ERDGBuilderFlags::ParallelExecute))
 	, RootAllocatorScope(Allocators.Root)
 	, Blackboard(Allocators.Root)
 	, BuilderName(InName)
@@ -589,8 +589,11 @@ FRDGBuilder::FRDGBuilder(FRHICommandListImmediate& InRHICmdList, FRDGEventName I
 
 	ProloguePass = SetupEmptyPass(Passes.Allocate<FRDGSentinelPass>(Allocators.Root, RDG_EVENT_NAME("Graph Prologue (Graphics)")));
 
+	bSupportsAsyncCompute = IsAsyncComputeSupported(ShaderPlatform);
+	bSupportsRenderPassMerge = IsRenderPassMergeEnabled(ShaderPlatform);
+
 	const bool bParallelExecuteFlag = EnumHasAnyFlags(InFlags, ERDGBuilderFlags::ParallelExecute);
-	const bool bParallelExecuteAllowedAwait = ::IsParallelExecuteEnabled();
+	const bool bParallelExecuteAllowedAwait = ::IsParallelExecuteEnabled(ShaderPlatform);
 	const bool bParallelExecuteAllowedAsync = bParallelExecuteAllowedAwait && GRDGParallelExecute > 1;
 
 	if (bParallelExecuteFlag)
@@ -605,8 +608,9 @@ FRDGBuilder::FRDGBuilder(FRHICommandListImmediate& InRHICmdList, FRDGEventName I
 		}
 	}
 
-	ParallelSetup.bEnabled   = ::IsParallelSetupEnabled()   && EnumHasAnyFlags(InFlags, ERDGBuilderFlags::ParallelSetup);
-	bParallelCompileEnabled  = ::IsParallelSetupEnabled()   && EnumHasAnyFlags(InFlags, ERDGBuilderFlags::ParallelCompile);
+	const bool bParallelSetupEnabledForPlatform = ::IsParallelSetupEnabled(ShaderPlatform);
+	ParallelSetup.bEnabled   = bParallelSetupEnabledForPlatform && EnumHasAnyFlags(InFlags, ERDGBuilderFlags::ParallelSetup);
+	bParallelCompileEnabled  = ::IsParallelSetupEnabled(ShaderPlatform) && EnumHasAnyFlags(InFlags, ERDGBuilderFlags::ParallelCompile);
 
 	if (TransientResourceAllocator)
 	{
@@ -868,7 +872,7 @@ END_SHADER_PARAMETER_STRUCT()
 
 void FRDGBuilder::UseExternalAccessMode(FRDGViewableResource* Resource, ERHIAccess ReadOnlyAccess, ERHIPipeline Pipelines)
 {
-	if (!IsAsyncComputeSupported())
+	if (!bSupportsAsyncCompute)
 	{
 		Pipelines = ERHIPipeline::Graphics;
 	}
@@ -1395,7 +1399,7 @@ void FRDGBuilder::Compile()
 	}
 
 	// Traverses passes on the graphics pipe and merges raster passes with the same render targets into a single RHI render pass.
-	if (IsRenderPassMergeEnabled() && RasterPassCount > 0)
+	if (bSupportsRenderPassMerge && RasterPassCount > 0)
 	{
 		SCOPED_NAMED_EVENT(MergeRenderPasses, FColor::Emerald);
 
diff --git a/Engine/Source/Runtime/RenderCore/Private/RenderGraphPrivate.cpp b/Engine/Source/Runtime/RenderCore/Private/RenderGraphPrivate.cpp
--- a/Engine/Source/Runtime/RenderCore/Private/RenderGraphPrivate.cpp
+++ b/Engine/Source/Runtime/RenderCore/Private/RenderGraphPrivate.cpp
@@ -262,7 +262,7 @@ FAutoConsoleVariableSink CVarRDGAsyncComputeSink(FConsoleCommandDelegate::Create
 		GRDGAsyncCompute = 0;
 	}
 
-	if (!IsAsyncComputeSupported())
+	if (!IsAsyncComputeSupported(GMaxRHIShaderPlatform))
 	{
 		GRDGAsyncCompute = 0;
 	}
@@ -621,7 +621,7 @@ void ShutdownRenderGraph()
 	FRDGBuilder::WaitForAsyncDeleteTask();
 }
 
-bool IsParallelExecuteEnabled()
+bool IsParallelExecuteEnabled(EShaderPlatform ShaderPlatform)
 {
 	return GRDGParallelExecute > 0
 		&& !GRHICommandList.Bypass()
@@ -629,9 +629,9 @@ bool IsParallelExecuteEnabled()
 		&& !GRDGDebug
 		&& !GRDGDebugFlushGPU
 		&& !GRDGTransitionLog
-		&& !IsMobilePlatform(GMaxRHIShaderPlatform)
-		&& !IsOpenGLPlatform(GMaxRHIShaderPlatform)
-		&& !IsVulkanMobileSM5Platform(GMaxRHIShaderPlatform)
+		&& !IsMobilePlatform(ShaderPlatform)
+		&& !IsOpenGLPlatform(ShaderPlatform)
+		&& !IsVulkanMobileSM5Platform(ShaderPlatform)
 		&& GRHISupportsMultithreadedShaderCreation
 #if WITH_DUMPGPU
 		&& !UE::RenderCore::DumpGPU::IsDumpingFrame()
@@ -641,16 +641,16 @@ bool IsParallelExecuteEnabled()
 		;
 }
 
-bool IsParallelSetupEnabled()
+bool IsParallelSetupEnabled(EShaderPlatform ShaderPlatform)
 {
 	return GRDGParallelSetup > 0
 		&& !GRHICommandList.Bypass()
 		&& !IsImmediateMode()
 		&& !GRDGDebug
 		&& !GRDGTransitionLog
-		&& !IsMobilePlatform(GMaxRHIShaderPlatform)
-		&& !IsOpenGLPlatform(GMaxRHIShaderPlatform)
-		&& !IsVulkanMobileSM5Platform(GMaxRHIShaderPlatform)
+		&& !IsMobilePlatform(ShaderPlatform)
+		&& !IsOpenGLPlatform(ShaderPlatform)
+		&& !IsVulkanMobileSM5Platform(ShaderPlatform)
 		&& GRHISupportsMultithreadedShaderCreation
 #if WITH_DUMPGPU
 		&& !UE::RenderCore::DumpGPU::IsDumpingFrame()
diff --git a/Engine/Source/Runtime/RenderCore/Private/RenderGraphPrivate.h b/Engine/Source/Runtime/RenderCore/Private/RenderGraphPrivate.h
--- a/Engine/Source/Runtime/RenderCore/Private/RenderGraphPrivate.h
+++ b/Engine/Source/Runtime/RenderCore/Private/RenderGraphPrivate.h
@@ -225,20 +225,20 @@ FORCEINLINE bool IsImmediateMode()
 	return GRDGImmediateMode != 0;
 }
 
-FORCEINLINE bool IsRenderPassMergeEnabled()
+FORCEINLINE bool IsRenderPassMergeEnabled(EShaderPlatform ShaderPlatform)
 {
-	return GRDGMergeRenderPasses != 0 && !IsImmediateMode() && FDataDrivenShaderPlatformInfo::IsValid(GMaxRHIShaderPlatform) && RHIHasTiledGPU(GMaxRHIShaderPlatform) && !GRDGDebugFlushGPU;
+	return GRDGMergeRenderPasses != 0 && !IsImmediateMode() && FDataDrivenShaderPlatformInfo::IsValid(ShaderPlatform) && RHIHasTiledGPU(ShaderPlatform) && !GRDGDebugFlushGPU;
 }
 
-FORCEINLINE bool IsAsyncComputeSupported()
+FORCEINLINE bool IsAsyncComputeSupported(EShaderPlatform ShaderPlatform)
 {
 	// Render pass merging and async compute are mutually exclusive since they are mutually dependent on each other (i.e. we can't
 	// merge render passes that involve async compute fences, since transitions would be relocated to the ends of the merged region).
-	return GRDGAsyncCompute > 0 && !IsImmediateMode() && !IsRenderPassMergeEnabled() && GSupportsEfficientAsyncCompute && GRHISupportsSeparateDepthStencilCopyAccess && !GTriggerGPUProfile;
+	return GRDGAsyncCompute > 0 && !IsImmediateMode() && !IsRenderPassMergeEnabled(ShaderPlatform) && GSupportsEfficientAsyncCompute && GRHISupportsSeparateDepthStencilCopyAccess && !GTriggerGPUProfile;
 }
 
-extern bool IsParallelExecuteEnabled();
-extern bool IsParallelSetupEnabled();
+extern bool IsParallelExecuteEnabled(EShaderPlatform ShaderPlatform);
+extern bool IsParallelSetupEnabled(EShaderPlatform ShaderPlatform);
 extern bool IsExtendedLifetimeResource(FRDGViewableResource* Resource);
 
 inline bool IsAsyncComputeTransientAliasingEnabled()
diff --git a/Engine/Source/Runtime/RenderCore/Public/RenderGraphBuilder.h b/Engine/Source/Runtime/RenderCore/Public/RenderGraphBuilder.h
--- a/Engine/Source/Runtime/RenderCore/Public/RenderGraphBuilder.h
+++ b/Engine/Source/Runtime/RenderCore/Public/RenderGraphBuilder.h
@@ -78,7 +78,7 @@ class FRDGBuilder
 	FRDGAllocatorScope RootAllocatorScope;
 
 public:
-	RENDERCORE_API FRDGBuilder(FRHICommandListImmediate& RHICmdList, FRDGEventName Name = {}, ERDGBuilderFlags Flags = ERDGBuilderFlags::None);
+	RENDERCORE_API FRDGBuilder(FRHICommandListImmediate& RHICmdList, FRDGEventName Name = {}, ERDGBuilderFlags Flags = ERDGBuilderFlags::None, EShaderPlatform ShaderPlatform = GMaxRHIShaderPlatform);
 	FRDGBuilder(const FRDGBuilder&) = delete;
 	RENDERCORE_API ~FRDGBuilder();
 
@@ -480,13 +480,16 @@ private:
 	FRDGPass* ProloguePass = nullptr;
 	FRDGPass* EpiloguePass = nullptr;
 
+	bool bSupportsAsyncCompute = false;
+	bool bSupportsRenderPassMerge = false;
+
 	uint32 AsyncComputePassCount = 0;
 	uint32 RasterPassCount = 0;
 
 	/** Tracks dispatch passes that need to launch tasks. */
 	TArray<FRDGDispatchPass*, FRDGArrayAllocator> DispatchPasses;
 
-	static RENDERCORE_API ERDGPassFlags OverridePassFlags(const TCHAR* PassName, ERDGPassFlags Flags);
+	RENDERCORE_API ERDGPassFlags OverridePassFlags(const TCHAR* PassName, ERDGPassFlags Flags) const;
 
 	FORCEINLINE FRDGPass* GetProloguePass() const
 	{
diff --git a/Engine/Source/Runtime/Renderer/Private/MeshEdgesRendering.cpp b/Engine/Source/Runtime/Renderer/Private/MeshEdgesRendering.cpp
--- a/Engine/Source/Runtime/Renderer/Private/MeshEdgesRendering.cpp
+++ b/Engine/Source/Runtime/Renderer/Private/MeshEdgesRendering.cpp
@@ -278,7 +278,7 @@ void RenderMeshEdges_RenderThread(
 
 	FUniformExpressionCacheAsyncUpdateScope AsyncUpdateScope;
 
-	FRDGBuilder GraphBuilder(RHICmdList, RDG_EVENT_NAME("MeshEdges"), ERDGBuilderFlags::Parallel);
+	FRDGBuilder GraphBuilder(RHICmdList, RDG_EVENT_NAME("MeshEdges"), ERDGBuilderFlags::Parallel, SceneRenderer->ShaderPlatform);
 	{
 		RDG_EVENT_SCOPE(GraphBuilder, "RenderMeshEdges_RenderThread");
 
diff --git a/Engine/Source/Runtime/Renderer/Private/PlanarReflectionRendering.cpp b/Engine/Source/Runtime/Renderer/Private/PlanarReflectionRendering.cpp
--- a/Engine/Source/Runtime/Renderer/Private/PlanarReflectionRendering.cpp
+++ b/Engine/Source/Runtime/Renderer/Private/PlanarReflectionRendering.cpp
@@ -348,7 +348,7 @@ static void UpdatePlanarReflectionContents_RenderThread(
 	FUniformExpressionCacheAsyncUpdateScope AsyncUpdateScope;
 
 	const ERHIFeatureLevel::Type FeatureLevel = SceneRenderer->FeatureLevel;
-	FRDGBuilder GraphBuilder(RHICmdList, RDG_EVENT_NAME("PlanarReflection"), ERDGBuilderFlags::Parallel);
+	FRDGBuilder GraphBuilder(RHICmdList, RDG_EVENT_NAME("PlanarReflection"), ERDGBuilderFlags::Parallel, SceneRenderer->ShaderPlatform);
 
 	// Make sure we render to the same set of GPUs as the main scene renderer.
 	if (MainSceneRenderer->ViewFamily.RenderTarget != nullptr)
diff --git a/Engine/Source/Runtime/Renderer/Private/ReflectionEnvironmentCapture.cpp b/Engine/Source/Runtime/Renderer/Private/ReflectionEnvironmentCapture.cpp
--- a/Engine/Source/Runtime/Renderer/Private/ReflectionEnvironmentCapture.cpp
+++ b/Engine/Source/Runtime/Renderer/Private/ReflectionEnvironmentCapture.cpp
@@ -518,7 +518,7 @@ void CaptureSceneToScratchCubemap(
 
 	FUniformExpressionCacheAsyncUpdateScope AsyncUpdateScope;
 
-	FRDGBuilder GraphBuilder(RHICmdList, RDG_EVENT_NAME("CubeMapCapture"), ERDGBuilderFlags::Parallel);
+	FRDGBuilder GraphBuilder(RHICmdList, RDG_EVENT_NAME("CubeMapCapture"), ERDGBuilderFlags::Parallel, SceneRenderer->ShaderPlatform);
 
 	{
 		RDG_EVENT_SCOPE(GraphBuilder, "CubeMapCapture");
diff --git a/Engine/Source/Runtime/Renderer/Private/SceneCaptureRendering.cpp b/Engine/Source/Runtime/Renderer/Private/SceneCaptureRendering.cpp
--- a/Engine/Source/Runtime/Renderer/Private/SceneCaptureRendering.cpp
+++ b/Engine/Source/Runtime/Renderer/Private/SceneCaptureRendering.cpp
@@ -411,10 +411,10 @@ static void UpdateSceneCaptureContentDeferred_RenderThread(
 
 #if WANTS_DRAW_MESH_EVENTS
 	SCOPED_DRAW_EVENTF(RHICmdList, SceneCapture, TEXT("SceneCapture %s"), EventName);
-	FRDGBuilder GraphBuilder(RHICmdList, RDG_EVENT_NAME("SceneCapture %s", *EventName), ERDGBuilderFlags::Parallel);
+	FRDGBuilder GraphBuilder(RHICmdList, RDG_EVENT_NAME("SceneCapture %s", *EventName), ERDGBuilderFlags::Parallel, SceneRenderer->ShaderPlatform);
 #else
 	SCOPED_DRAW_EVENT(RHICmdList, UpdateSceneCaptureContent_RenderThread);
-	FRDGBuilder GraphBuilder(RHICmdList, RDG_EVENT_NAME("SceneCapture"), ERDGBuilderFlags::Parallel);
+	FRDGBuilder GraphBuilder(RHICmdList, RDG_EVENT_NAME("SceneCapture"), ERDGBuilderFlags::Parallel, SceneRenderer->ShaderPlatform);
 #endif
 
 	{
