Delete GPU profiler frame after RDG async execution tasks complete, and re-enabled async RDG.

diff --git a/Engine/Source/Runtime/RHI/Private/DynamicRHI.cpp b/Engine/Source/Runtime/RHI/Private/DynamicRHI.cpp
--- a/Engine/Source/Runtime/RHI/Private/DynamicRHI.cpp
+++ b/Engine/Source/Runtime/RHI/Private/DynamicRHI.cpp
@@ -778,6 +778,7 @@ FDefaultRHIRenderQueryPool::~FDefaultRHIRenderQueryPool()
 FRHIPooledRenderQuery FDefaultRHIRenderQueryPool::AllocateQuery()
 {
 	check(IsInParallelRenderingThread());
+	UE::TScopeLock Lock(QueriesMutex);
 	if (Queries.Num() > 0)
 	{
 		return FRHIPooledRenderQuery(this, Queries.Pop());
@@ -798,7 +799,8 @@ void FDefaultRHIRenderQueryPool::ReleaseQuery(TRefCountPtr<FRHIRenderQuery>&& Qu
 {
 	check(IsInParallelRenderingThread());
 	checkf(Query.IsValid(), TEXT("Only release valid queries"));
-
+	
+	UE::TScopeLock Lock(QueriesMutex);
 	Queries.Push(MoveTemp(Query));
 	check(!Query.IsValid());
 }
diff --git a/Engine/Source/Runtime/RHI/Public/DynamicRHI.h b/Engine/Source/Runtime/RHI/Public/DynamicRHI.h
--- a/Engine/Source/Runtime/RHI/Public/DynamicRHI.h
+++ b/Engine/Source/Runtime/RHI/Public/DynamicRHI.h
@@ -109,6 +109,7 @@ private:
 	const ERenderQueryType QueryType;
 	uint32 AllocatedQueries = 0;
 	TArray<TRefCountPtr<FRHIRenderQuery>> Queries;
+	UE::FMutex QueriesMutex;
 };
 
 struct FRHICalcTextureSizeResult
diff --git a/Engine/Source/Runtime/RenderCore/Private/ProfilingDebugging/RealtimeGPUProfiler.cpp b/Engine/Source/Runtime/RenderCore/Private/ProfilingDebugging/RealtimeGPUProfiler.cpp
--- a/Engine/Source/Runtime/RenderCore/Private/ProfilingDebugging/RealtimeGPUProfiler.cpp
+++ b/Engine/Source/Runtime/RenderCore/Private/ProfilingDebugging/RealtimeGPUProfiler.cpp
@@ -6,6 +6,7 @@
 #include "RenderCore.h"
 #include "RenderingThread.h"
 #include "GPUProfiler.h"
+#include "RenderGraphBuilder.h"
 #include "Misc/ScopeRWLock.h"
 
 #if RHI_NEW_GPU_PROFILER
@@ -840,6 +841,8 @@ FRealtimeGPUProfiler::FRealtimeGPUProfiler()
 
 void FRealtimeGPUProfiler::Cleanup()
 {
+	CleanUpTask.Wait();
+
 	ActiveFrame.Reset();
 	PendingFrames.Empty();
 
@@ -938,6 +941,14 @@ void FRealtimeGPUProfiler::EndFrame(FRHICommandListImmediate& RHICmdList)
 #endif
 		))
 		{
+			// Launch an inline async task that will free the frame once RDG async task execution is complete to keep it alive.
+			CleanUpTask = UE::Tasks::Launch(UE_SOURCE_LOCATION, [Frame = MoveTemp(*FramePtr)] () mutable
+			{
+				FOptionalTaskTagScope TaskTag(ETaskTag::EParallelRenderingThread);
+				Frame = {};
+
+			}, MakeArrayView({ CleanUpTask, FRDGBuilder::GetAsyncExecuteTask() }), UE::Tasks::ETaskPriority::High, UE::Tasks::EExtendedTaskPriority::Inline);
+
 			PendingFrames.Pop();
 		}
 	}
diff --git a/Engine/Source/Runtime/RenderCore/Private/RenderGraphPrivate.cpp b/Engine/Source/Runtime/RenderCore/Private/RenderGraphPrivate.cpp
--- a/Engine/Source/Runtime/RenderCore/Private/RenderGraphPrivate.cpp
+++ b/Engine/Source/Runtime/RenderCore/Private/RenderGraphPrivate.cpp
@@ -342,13 +342,13 @@ FAutoConsoleVariableRef CVarRDGParallelSetup(
 	TEXT(" 1: pass setup is done asynchronously (default);"),
 	ECVF_RenderThreadSafe);
 
-int32 GRDGParallelExecute = 1;
+int32 GRDGParallelExecute = 2;
 FAutoConsoleVariableRef CVarRDGParallelExecute(
 	TEXT("r.RDG.ParallelExecute"), GRDGParallelExecute,
 	TEXT("Whether to enable parallel execution of passes when supported.")
 	TEXT(" 0: off")
-	TEXT(" 1: parallel with all tasks awaited (default)")
-	TEXT(" 2: parallel with async tasks"),
+	TEXT(" 1: parallel with all tasks awaited")
+	TEXT(" 2: parallel with async tasks (default)"),
 	FConsoleVariableDelegate::CreateLambda([](IConsoleVariable* Variable)
 	{
 		if (Variable->GetInt())
diff --git a/Engine/Source/Runtime/RenderCore/Public/ProfilingDebugging/RealtimeGPUProfiler.h b/Engine/Source/Runtime/RenderCore/Public/ProfilingDebugging/RealtimeGPUProfiler.h
--- a/Engine/Source/Runtime/RenderCore/Public/ProfilingDebugging/RealtimeGPUProfiler.h
+++ b/Engine/Source/Runtime/RenderCore/Public/ProfilingDebugging/RealtimeGPUProfiler.h
@@ -329,6 +329,8 @@ private:
 	/** Deinitialize of the object*/
 	void Cleanup();
 
+	UE::Tasks::FTask CleanUpTask;
+
 	TUniquePtr<FRealtimeGPUProfilerFrame> ActiveFrame;
 	TQueue<TUniquePtr<FRealtimeGPUProfilerFrame>> PendingFrames;
 
