Removed deprecated Texture functionality from UnifiedBuffer.

#rnx
#tests Not needed, no code uses this.

diff --git a/Engine/Shaders/Private/ByteBuffer.usf b/Engine/Shaders/Private/ByteBuffer.usf
--- a/Engine/Shaders/Private/ByteBuffer.usf
+++ b/Engine/Shaders/Private/ByteBuffer.usf
@@ -7,7 +7,6 @@
 #define RESOURCE_TYPE_STRUCTURED_BUFFER			(1)
 #define RESOURCE_TYPE_UINT_BUFFER				(2)
 #define RESOURCE_TYPE_UINT4_ALIGNED_BUFFER		(3)
-#define RESOURCE_TYPE_FLOAT4_TEXTURE			(4)
 
 #define STRUCTURED_ELEMENT_SIZE_UINT1 (0)
 #define STRUCTURED_ELEMENT_SIZE_UINT2 (1)
@@ -76,11 +75,6 @@ struct FUint8
 	RWStructuredBuffer<uint4>	DstStructuredBuffer4x;
 	RWStructuredBuffer<FUint8>	DstStructuredBuffer8x;
 	
-#elif RESOURCE_TYPE == RESOURCE_TYPE_FLOAT4_TEXTURE
-	StructuredBuffer<uint>			ScatterStructuredBuffer;
-	Texture2D<float4>				SrcTexture;
-	StructuredBuffer<float4>		UploadStructuredBuffer4x;
-	RWTexture2D<float4>				DstTexture;
 #elif RESOURCE_TYPE == RESOURCE_TYPE_UINT_BUFFER || RESOURCE_TYPE == RESOURCE_TYPE_UINT4_ALIGNED_BUFFER
 	ByteAddressBuffer				ScatterByteAddressBuffer;
 	ByteAddressBuffer				SrcByteAddressBuffer;
@@ -125,15 +119,6 @@ void MemsetBufferCS( uint3 GroupID : SV_GroupID, uint GroupIndex : SV_GroupIndex
 	{
 		DstByteAddressBuffer.Store( DstIndex * 4, Value );
 	}
-#elif RESOURCE_TYPE == RESOURCE_TYPE_FLOAT4_TEXTURE
-	uint2 IndexTexture;
-	IndexTexture.y = ThreadId / (Float4sPerLine);
-	IndexTexture.x = ThreadId % (Float4sPerLine);
-
-	if (ThreadId < Size)
-	{
-		DstTexture[IndexTexture.xy] = asfloat(Value);
-	}
 #else
 	#error "Not implemented"
 #endif
@@ -154,16 +139,6 @@ void MemcpyCS( uint ThreadId : SV_DispatchThreadID )
 	{
 		STRUCTURED_BUFFER_ACCESS(DstStructuredBuffer, DstOffset + ThreadId) = STRUCTURED_BUFFER_ACCESS(SrcStructuredBuffer, SrcOffset + ThreadId);
 	}
-#elif RESOURCE_TYPE == RESOURCE_TYPE_FLOAT4_TEXTURE
-	uint2 IndexTexture;
-	IndexTexture.y = ThreadId / (Float4sPerLine);
-	IndexTexture.x = ThreadId % (Float4sPerLine);
-
-	if(ThreadId < Size)
-	{
-		float4 SrcValue = SrcTexture.Load(float3(IndexTexture.x, IndexTexture.y, 0));
-		DstTexture[IndexTexture.xy] = SrcValue;
-	}
 #elif RESOURCE_TYPE == RESOURCE_TYPE_UINT_BUFFER || RESOURCE_TYPE == RESOURCE_TYPE_UINT4_ALIGNED_BUFFER
 	// Size and offsets are in dwords
 	uint SrcIndex = SrcOffset + ThreadId * 4;
@@ -217,15 +192,6 @@ void ScatterCopyCS( uint DispatchThreadId : SV_DispatchThreadID )
 		uint SrcIndex = ThreadId;
 		uint4 SrcData = UploadByteAddressBuffer.Load4( SrcIndex * 16 );
 		DstByteAddressBuffer.Store4( DstIndex * 16, SrcData );
-#elif RESOURCE_TYPE == RESOURCE_TYPE_FLOAT4_TEXTURE
-		uint DstIndex = ScatterStructuredBuffer[ScatterIndex] * Size + ScatterOffset;
-		uint SrcIndex = ThreadId;
-		uint2 IndexTexture;
-		IndexTexture.y = DstIndex / (Float4sPerLine);
-		IndexTexture.x = DstIndex % (Float4sPerLine);
-
-		float4 srvResourceVal = asfloat(UploadStructuredBuffer4x[SrcIndex]);
-		DstTexture[IndexTexture.xy] = srvResourceVal;
 #elif RESOURCE_TYPE == RESOURCE_TYPE_UINT_BUFFER
 		uint DstIndex = ScatterByteAddressBuffer.Load( ScatterIndex * 4 ) * Size + ScatterOffset;
 		uint SrcIndex = ThreadId;
diff --git a/Engine/Source/Runtime/RenderCore/Private/UnifiedBuffer.cpp b/Engine/Source/Runtime/RenderCore/Private/UnifiedBuffer.cpp
--- a/Engine/Source/Runtime/RenderCore/Private/UnifiedBuffer.cpp
+++ b/Engine/Source/Runtime/RenderCore/Private/UnifiedBuffer.cpp
@@ -23,7 +23,6 @@ enum class EByteBufferResourceType
 	StructuredBuffer,
 	Uint_Buffer,
 	Uint4Aligned_Buffer,
-	Float4_Texture,
 	Count
 };
 
@@ -87,7 +86,6 @@ class FByteBufferShader : public FGlobalShader
 		SHADER_PARAMETER_UAV(RWBuffer<float4>, DstBuffer)
 		SHADER_PARAMETER_UAV(RWStructuredBuffer<float4>, DstStructuredBuffer4x)
 		SHADER_PARAMETER_UAV(RWByteAddressBuffer, DstByteAddressBuffer)
-		SHADER_PARAMETER_UAV(RWTexture2D<float4>, DstTexture)
 	END_SHADER_PARAMETER_STRUCT()
 };
 
@@ -130,7 +128,6 @@ class FMemcpyCS : public FByteBufferShader
 		SHADER_PARAMETER_SRV(Buffer<float4>, SrcBuffer)
 		SHADER_PARAMETER_SRV(StructuredBuffer<float4>, SrcStructuredBuffer4x)
 		SHADER_PARAMETER_SRV(ByteAddressBuffer, SrcByteAddressBuffer)
-		SHADER_PARAMETER_SRV(Texture2D<float4>, SrcTexture)
 	END_SHADER_PARAMETER_STRUCT()
 };
 IMPLEMENT_GLOBAL_SHADER( FMemcpyCS, "/Engine/Private/ByteBuffer.usf", "MemcpyCS", SF_Compute );
@@ -166,8 +163,7 @@ enum class EResourceType
 {
 	BUFFER,
 	STRUCTURED_BUFFER,
-	BYTEBUFFER,
-	TEXTURE // NOTE: Deprecated
+	BYTEBUFFER
 };
 
 template<typename ResourceType>
@@ -185,29 +181,12 @@ struct ResourceTypeTraits<FRWBufferStructured>
 	static const EResourceType Type = EResourceType::STRUCTURED_BUFFER;
 };
 
-// NOTE: Deprecated
-template<>
-struct ResourceTypeTraits<FTextureRWBuffer>
-{
-	static const EResourceType Type = EResourceType::TEXTURE;
-};
-
 template<>
 struct ResourceTypeTraits<FRWByteAddressBuffer>
 {
 	static const EResourceType Type = EResourceType::BYTEBUFFER;
 };
 
-static uint32 CalculateFloat4sPerLine()
-{
-	// This is a remnant from when primitive scene data used textures to scatter upload, and should be removed once
-	// texture uploads have been fully deprecated
-	PRAGMA_DISABLE_DEPRECATION_WARNINGS
-	uint16 PrimitivesPerTextureLine = (uint16)FMath::Min((int32)MAX_uint16, (int32)GMaxTextureDimensions) / FScatterUploadBuffer::PrimitiveDataStrideInFloat4s;
-	return PrimitivesPerTextureLine * FScatterUploadBuffer::PrimitiveDataStrideInFloat4s;
-	PRAGMA_ENABLE_DEPRECATION_WARNINGS
-}
-
 ///////////////////////////////////////////////////////////////////////////////////////////////////
 
 class FRDGByteBufferShader : public FGlobalShader
@@ -240,7 +219,6 @@ class FRDGByteBufferShader : public FGlobalShader
 		SHADER_PARAMETER_RDG_BUFFER_UAV(RWStructuredBuffer<uint2>, DstStructuredBuffer2x)
 		SHADER_PARAMETER_RDG_BUFFER_UAV(RWStructuredBuffer<uint>, DstStructuredBuffer1x)
 		SHADER_PARAMETER_RDG_BUFFER_UAV(RWByteAddressBuffer, DstByteAddressBuffer)
-		SHADER_PARAMETER_RDG_TEXTURE_UAV(RWTexture2D<float4>, DstTexture)
 	END_SHADER_PARAMETER_STRUCT()
 };
 
@@ -264,7 +242,6 @@ class FRDGMemcpyCS : public FRDGByteBufferShader
 		SHADER_PARAMETER_RDG_BUFFER_SRV(StructuredBuffer<uint2>, SrcStructuredBuffer2x)
 		SHADER_PARAMETER_RDG_BUFFER_SRV(StructuredBuffer<uint>, SrcStructuredBuffer1x)
 		SHADER_PARAMETER_RDG_BUFFER_SRV(ByteAddressBuffer, SrcByteAddressBuffer)
-		SHADER_PARAMETER_RDG_TEXTURE_SRV(Texture2D<float4>, SrcTexture)
 	END_SHADER_PARAMETER_STRUCT()
 };
 IMPLEMENT_GLOBAL_SHADER(FRDGMemcpyCS, "/Engine/Private/ByteBuffer.usf", "MemcpyCS", SF_Compute);
@@ -311,8 +288,6 @@ EResourceType GetResourceType(FRDGViewableResource* Resource)
 	check(Resource);
 	switch (Resource->Type)
 	{
-	case ERDGViewableResourceType::Texture:
-		return EResourceType::TEXTURE;
 	case ERDGViewableResourceType::Buffer:
 		return GetBufferType(GetAsBuffer(Resource));
 	}
@@ -392,11 +367,6 @@ void MemsetResource(FRDGBuilder& GraphBuilder, FRDGUnorderedAccessView* UAV, con
 		};
 		break;
 	}
-	case EResourceType::TEXTURE:
-		ResourceTypeEnum = EByteBufferResourceType::Float4_Texture;
-		PassParameters->DstTexture = GetAs<FRDGTextureUAV>(UAV);
-		PassParameters->Float4sPerLine = CalculateFloat4sPerLine();
-		break;
 	default:
 		checkNoEntry();
 	}
@@ -421,13 +391,6 @@ void MemsetResource(FRDGBuilder& GraphBuilder, FRDGBufferUAV* UAV, const FMemset
 	PRAGMA_ENABLE_DEPRECATION_WARNINGS
 }
 
-void MemsetResource(FRDGBuilder& GraphBuilder, FRDGTextureUAV* UAV, const FMemsetResourceParams& Params)
-{
-	PRAGMA_DISABLE_DEPRECATION_WARNINGS
-	MemsetResource(GraphBuilder, static_cast<FRDGUnorderedAccessView*>(UAV), Params);
-	PRAGMA_ENABLE_DEPRECATION_WARNINGS
-}
-
 void MemcpyResource(FRDGBuilder& GraphBuilder, FRDGUnorderedAccessView* UAV, FRDGShaderResourceView* SRV, const FMemcpyResourceParams& Params)
 {
 	check(UAV && SRV);
@@ -498,12 +461,6 @@ void MemcpyResource(FRDGBuilder& GraphBuilder, FRDGUnorderedAccessView* UAV, FRD
 			PassParameters->SrcBuffer                      = GetAs<FRDGBufferSRV>(SRV);
 			PassParameters->Common.DstBuffer               = GetAs<FRDGBufferUAV>(UAV);
 			break;
-		case EResourceType::TEXTURE:
-			ResourceTypeEnum = EByteBufferResourceType::Float4_Texture;
-			PassParameters->SrcTexture                     = GetAs<FRDGTextureSRV>(SRV);
-			PassParameters->Common.DstTexture              = GetAs<FRDGTextureUAV>(UAV);
-			PassParameters->Common.Float4sPerLine          = CalculateFloat4sPerLine();
-			break;
 		default:
 			checkNoEntry();
 		}
@@ -531,13 +488,6 @@ void MemcpyResource(FRDGBuilder& GraphBuilder, FRDGBufferUAV* UAV, FRDGBufferSRV
 	PRAGMA_ENABLE_DEPRECATION_WARNINGS
 }
 
-void MemcpyResource(FRDGBuilder& GraphBuilder, FRDGTextureUAV* UAV, FRDGTextureSRV* SRV, const FMemcpyResourceParams& Params)
-{
-	PRAGMA_DISABLE_DEPRECATION_WARNINGS
-	MemcpyResource(GraphBuilder, static_cast<FRDGUnorderedAccessView*>(UAV), static_cast<FRDGShaderResourceView*>(SRV), Params);
-	PRAGMA_ENABLE_DEPRECATION_WARNINGS
-}
-
 FRDGBuffer* ResizeBufferIfNeeded(FRDGBuilder& GraphBuilder, TRefCountPtr<FRDGPooledBuffer>& ExternalBuffer, const FRDGBufferDesc& BufferDesc, const TCHAR* Name)
 {
 	FRDGBuffer* InternalBufferNew = nullptr;
@@ -787,16 +737,6 @@ void ScatterCopyResource(FRDGBuilder& GraphBuilder, FRDGViewableResource* DstRes
 		Parameters.ScatterStructuredBuffer = ScatterBufferSRV;
 		Parameters.Common.DstBuffer = GraphBuilder.CreateUAV(GetAsBuffer(DstResource), ERDGUnorderedAccessViewFlags::SkipBarrier);
 	}
-	else if (DstResourceType == EResourceType::TEXTURE)
-	{
-		ResourceTypeEnum = EByteBufferResourceType::Float4_Texture;
-
-		Parameters.UploadStructuredBuffer4x = UploadBufferSRV;
-		Parameters.ScatterStructuredBuffer = ScatterBufferSRV;
-		Parameters.Common.DstTexture = GraphBuilder.CreateUAV(GetAsTexture(DstResource), ERDGUnorderedAccessViewFlags::SkipBarrier);
-
-		Parameters.Common.Float4sPerLine = CalculateFloat4sPerLine();
-	}
 
 	FRDGScatterCopyCS::FPermutationDomain PermutationVector;
 	PermutationVector.Set<FRDGScatterCopyCS::ResourceTypeDim>((int)ResourceTypeEnum);
@@ -824,14 +764,6 @@ void ScatterCopyResource(FRDGBuilder& GraphBuilder, FRDGBuffer* DstResource, FRD
 	ScatterCopyResource(GraphBuilder, static_cast<FRDGViewableResource*>(DstResource), ScatterBufferSRV, UploadBufferSRV, Params);
 	PRAGMA_ENABLE_DEPRECATION_WARNINGS
 }
-void ScatterCopyResource(FRDGBuilder& GraphBuilder, FRDGTexture* DstResource, FRDGBufferSRV* ScatterBufferSRV, FRDGBufferSRV* UploadBufferSRV, const FScatterCopyParams &Params)
-{
-	PRAGMA_DISABLE_DEPRECATION_WARNINGS
-	ScatterCopyResource(GraphBuilder, static_cast<FRDGViewableResource*>(DstResource), ScatterBufferSRV, UploadBufferSRV, Params);
-	PRAGMA_ENABLE_DEPRECATION_WARNINGS
-}
-
-
 
 void FRDGScatterUploadBuffer::Init(FRDGBuilder& GraphBuilder, uint32 NumElements, uint32 InNumBytesPerElement, bool bInFloat4Buffer, const TCHAR* Name)
 {
@@ -1069,16 +1001,6 @@ void FRDGAsyncScatterUploadBuffer::End(FRDGBuilder& GraphBuilder, FRDGScatterUpl
 		Parameters.ScatterStructuredBuffer = ScatterBufferSRV;
 		Parameters.Common.DstBuffer = GraphBuilder.CreateUAV(GetAsBuffer(DstResource), ERDGUnorderedAccessViewFlags::SkipBarrier);
 	}
-	else if (DstResourceType == EResourceType::TEXTURE)
-	{
-		ResourceTypeEnum = EByteBufferResourceType::Float4_Texture;
-
-		Parameters.UploadStructuredBuffer4x = UploadBufferSRV;
-		Parameters.ScatterStructuredBuffer = ScatterBufferSRV;
-		Parameters.Common.DstTexture = GraphBuilder.CreateUAV(GetAsTexture(DstResource), ERDGUnorderedAccessViewFlags::SkipBarrier);
-
-		Parameters.Common.Float4sPerLine = CalculateFloat4sPerLine();
-	}
 
 	FRDGScatterCopyCS::FPermutationDomain PermutationVector;
 	PermutationVector.Set<FRDGScatterCopyCS::ResourceTypeDim>((int)ResourceTypeEnum);
@@ -1168,13 +1090,6 @@ void MemsetResource(FRHICommandList& RHICmdList, const ResourceType& DstBuffer,
 
 		Parameters.DstStructuredBuffer4x = DstBuffer.UAV;
 	}
-	else if (ResourceTypeTraits<ResourceType>::Type == EResourceType::TEXTURE)
-	{
-		ResourceTypeEnum = EByteBufferResourceType::Float4_Texture;
-
-		Parameters.DstTexture = DstBuffer.UAV;
-		Parameters.Float4sPerLine = CalculateFloat4sPerLine();
-	}
 
 	FMemcpyCS::FPermutationDomain PermutationVector;
 	PermutationVector.Set<FMemcpyCS::ResourceTypeDim >((int)ResourceTypeEnum);
@@ -1234,14 +1149,6 @@ void MemcpyResource(FRHICommandList& RHICmdList, const ResourceType& DstBuffer,
 			Parameters.SrcBuffer = SrcBuffer.SRV;
 			Parameters.Common.DstBuffer = DstBuffer.UAV;
 		}
-		else if (ResourceTypeTraits<ResourceType>::Type == EResourceType::TEXTURE)
-		{
-			ResourceTypeEnum = EByteBufferResourceType::Float4_Texture;
-
-			Parameters.SrcTexture = SrcBuffer.SRV;
-			Parameters.Common.DstTexture = DstBuffer.UAV;
-			Parameters.Common.Float4sPerLine = CalculateFloat4sPerLine();
-		}
 		else
 		{
 			check(false);
@@ -1262,42 +1169,6 @@ void MemcpyResource(FRHICommandList& RHICmdList, const ResourceType& DstBuffer,
 		RHICmdList.EndUAVOverlap(DstBuffer.UAV);
 }
 
-// NOTE: Deprecated
-RENDERCORE_API bool ResizeResourceIfNeeded(FRHICommandList& RHICmdList, FTextureRWBuffer& Texture, uint32 NumBytes, const TCHAR* DebugName)
-{
-	check((NumBytes & 15) == 0);
-	uint32 Float4sPerLine = CalculateFloat4sPerLine();
-	uint32 BytesPerLine = Float4sPerLine * 16;
-
-	EPixelFormat BufferFormat = PF_A32B32G32R32F;
-	uint32 BytesPerElement = GPixelFormats[BufferFormat].BlockBytes;
-
-	uint32 NumLines = (NumBytes + BytesPerLine - 1) / BytesPerLine;
-
-	if (Texture.NumBytes == 0)
-	{
-		Texture.Initialize2D(DebugName, BytesPerElement, Float4sPerLine, NumLines, PF_A32B32G32R32F, TexCreate_RenderTargetable | TexCreate_UAV);
-		return true;
-	}
-	else if ((NumLines * Float4sPerLine * BytesPerElement) != Texture.NumBytes)
-	{
-		FTextureRWBuffer NewTexture;
-		NewTexture.Initialize2D(DebugName, BytesPerElement, Float4sPerLine, NumLines, PF_A32B32G32R32F, TexCreate_RenderTargetable | TexCreate_UAV);
-
-		FMemcpyResourceParams Params;
-		Params.Count = NumBytes / BytesPerElement;
-		Params.SrcOffset = 0;
-		Params.DstOffset = 0;
-		PRAGMA_DISABLE_DEPRECATION_WARNINGS
-		MemcpyResource(RHICmdList, NewTexture, Texture, Params);
-		PRAGMA_ENABLE_DEPRECATION_WARNINGS
-		Texture = NewTexture;
-		return true;
-	}
-
-	return false;
-}
-
 template<>
 RENDERCORE_API bool ResizeResourceIfNeeded<FRWBufferStructured>(FRHICommandList& RHICmdList, FRWBufferStructured& Buffer, uint32 NumBytes, const TCHAR* DebugName)
 {
@@ -1403,8 +1274,7 @@ RENDERCORE_API bool ResizeResourceIfNeeded(FRHICommandList& RHICmdList, FRWBuffe
 	return false;
 }
 
-template<>
-RENDERCORE_API bool ResizeResourceSOAIfNeeded<FRWBufferStructured>(FRHICommandList& RHICmdList, FRWBufferStructured& Buffer, const FResizeResourceSOAParams& Params, const TCHAR* DebugName)
+RENDERCORE_API bool ResizeResourceSOAIfNeeded(FRHICommandList& RHICmdList, FRWBufferStructured& Buffer, const FResizeResourceSOAParams& Params, const TCHAR* DebugName)
 {
 	const uint32 BytesPerElement = 16;
 
@@ -1772,16 +1642,6 @@ void FScatterUploadBuffer::ResourceUploadTo(FRHICommandList& RHICmdList, const R
 		Parameters.ScatterStructuredBuffer = ScatterBuffer.SRV;
 		Parameters.Common.DstBuffer = DstBuffer.UAV;
 	}
-	else if (ResourceTypeTraits<ResourceType>::Type == EResourceType::TEXTURE)
-	{
-		ResourceTypeEnum = EByteBufferResourceType::Float4_Texture;
-
-		Parameters.UploadStructuredBuffer4x = UploadBuffer.SRV;
-		Parameters.ScatterStructuredBuffer = ScatterBuffer.SRV;
-		Parameters.Common.DstTexture = DstBuffer.UAV;
-
-		Parameters.Common.Float4sPerLine = CalculateFloat4sPerLine();
-	}
 
 	FByteBufferShader::FPermutationDomain PermutationVector;
 	PermutationVector.Set<FByteBufferShader::ResourceTypeDim>((int)ResourceTypeEnum);
@@ -1815,7 +1675,6 @@ void FScatterUploadBuffer::ResourceUploadTo(FRHICommandList& RHICmdList, const R
 template RENDERCORE_API void MemsetResource<FRWBufferStructured>(FRHICommandList& RHICmdList, const FRWBufferStructured& DstBuffer, const FMemsetResourceParams& Params);
 template RENDERCORE_API void MemsetResource<FRWByteAddressBuffer>(FRHICommandList& RHICmdList, const FRWByteAddressBuffer& DstBuffer, const FMemsetResourceParams& Params);
 
-template RENDERCORE_API void MemcpyResource<FTextureRWBuffer>(FRHICommandList& RHICmdList, const FTextureRWBuffer& DstBuffer, const FTextureRWBuffer& SrcBuffer, const FMemcpyResourceParams& Params, bool bAlreadyInUAVOverlap);
 template RENDERCORE_API void MemcpyResource<FRWBuffer>(FRHICommandList& RHICmdList, const FRWBuffer& DstBuffer, const FRWBuffer& SrcBuffer, const FMemcpyResourceParams& Params, bool bAlreadyInUAVOverlap);
 template RENDERCORE_API void MemcpyResource<FRWBufferStructured>(FRHICommandList& RHICmdList, const FRWBufferStructured& DstBuffer, const FRWBufferStructured& SrcBuffer, const FMemcpyResourceParams& Params, bool bAlreadyInUAVOverlap);
 template RENDERCORE_API void MemcpyResource<FRWByteAddressBuffer>(FRHICommandList& RHICmdList, const FRWByteAddressBuffer& DstBuffer, const FRWByteAddressBuffer& SrcBuffer, const FMemcpyResourceParams& Params, bool bAlreadyInUAVOverlap);
@@ -1823,6 +1682,3 @@ template RENDERCORE_API void MemcpyResource<FRWByteAddressBuffer>(FRHICommandLis
 template RENDERCORE_API void FScatterUploadBuffer::ResourceUploadTo<FRWBuffer>(FRHICommandList& RHICmdList, const FRWBuffer& DstBuffer, bool bFlush);
 template RENDERCORE_API void FScatterUploadBuffer::ResourceUploadTo<FRWBufferStructured>(FRHICommandList& RHICmdList, const FRWBufferStructured& DstBuffer, bool bFlush);
 template RENDERCORE_API void FScatterUploadBuffer::ResourceUploadTo<FRWByteAddressBuffer>(FRHICommandList& RHICmdList, const FRWByteAddressBuffer& DstBuffer, bool bFlush);
-
-// NOTE: Deprecated
-template RENDERCORE_API void FScatterUploadBuffer::ResourceUploadTo<FTextureRWBuffer>(FRHICommandList& RHICmdList, const FTextureRWBuffer& DstBuffer, bool bFlush);
\ No newline at end of file
diff --git a/Engine/Source/Runtime/RenderCore/Public/UnifiedBuffer.h b/Engine/Source/Runtime/RenderCore/Public/UnifiedBuffer.h
--- a/Engine/Source/Runtime/RenderCore/Public/UnifiedBuffer.h
+++ b/Engine/Source/Runtime/RenderCore/Public/UnifiedBuffer.h
@@ -47,29 +47,16 @@ struct FResizeResourceSOAParams
 
 template<typename ResourceType>
 extern RENDERCORE_API void MemsetResource(FRHICommandList& RHICmdList, const ResourceType& DstBuffer, const FMemsetResourceParams& Params);
+
 template<typename ResourceType>
 extern RENDERCORE_API void MemcpyResource(FRHICommandList& RHICmdList, const ResourceType& DstBuffer, const ResourceType& SrcBuffer, const FMemcpyResourceParams& Params, bool bAlreadyInUAVOverlap = false);
 
-UE_DEPRECATED(5.3, "MemsetResource will no longer support FTextureRWBuffer in a future release. Either use a buffer target instead, or clear the contents manually.")
-inline void MemsetResource(FRHICommandList& RHICmdList, const FTextureRWBuffer& DstBuffer, const FMemsetResourceParams& Params)
-{
-	MemsetResource<FTextureRWBuffer>(RHICmdList, DstBuffer, Params);
-}
+RENDERCORE_API bool ResizeResourceSOAIfNeeded(FRHICommandList& RHICmdList, FRWBufferStructured& Texture, const FResizeResourceSOAParams& Params, const TCHAR* DebugName);
 
-UE_DEPRECATED(5.3, "MemcpyResource will no longer support FTextureRWBuffer in a future release. Either use a buffer target instead, or copy the contents manually.")
-inline void MemcpyResource(FRHICommandList& RHICmdList, const FTextureRWBuffer& DstBuffer, const FTextureRWBuffer& SrcBuffer, const FMemcpyResourceParams& Params, bool bAlreadyInUAVOverlap = false)
-{
-	MemcpyResource<FTextureRWBuffer>(RHICmdList, DstBuffer, SrcBuffer, Params, bAlreadyInUAVOverlap);
-}
-
-template<typename ResourceType>
-extern RENDERCORE_API bool ResizeResourceSOAIfNeeded(FRHICommandList& RHICmdList, ResourceType& Texture, const FResizeResourceSOAParams& Params, const TCHAR* DebugName);
 template<typename ResourceType>
 extern RENDERCORE_API bool ResizeResourceIfNeeded(FRHICommandList& RHICmdList, ResourceType& Buffer, uint32 NumBytes, const TCHAR* DebugName);
-RENDERCORE_API bool ResizeResourceIfNeeded(FRHICommandList& RHICmdList, FRWBuffer& Buffer, EPixelFormat Format, uint32 NumElements, const TCHAR* DebugName);
 
-UE_DEPRECATED(5.3, "ResizeResourceIfNeeded will no longer support FTextureRWBuffer in a future release. Either use a buffer target instead, or create a new texture and copy the contents manually.")
-RENDERCORE_API bool ResizeResourceIfNeeded(FRHICommandList& RHICmdList, FTextureRWBuffer& Texture, uint32 NumBytes, const TCHAR* DebugName);
+RENDERCORE_API bool ResizeResourceIfNeeded(FRHICommandList& RHICmdList, FRWBuffer& Buffer, EPixelFormat Format, uint32 NumElements, const TCHAR* DebugName);
 
 /**
  * This version will resize/allocate the buffer at once and add a RDG pass to perform the copy on the RDG time-line if there was previous data).
@@ -82,14 +69,6 @@ RENDERCORE_API bool ResizeResourceIfNeeded(FRDGBuilder& GraphBuilder, FRWBuffer&
 class FScatterUploadBuffer
 {
 public:
-	// This is an old dependency upon the primitive data that is no longer used
-	
-	enum
-	UE_DEPRECATED(5.3, "Use PRIMITIVE_SCENE_DATA_STRIDE in SceneDefinitions.h instead.")
-	{
-		PrimitiveDataStrideInFloat4s = 43
-	};
-
 	FByteAddressBuffer ScatterBuffer;
 	FByteAddressBuffer UploadBuffer;
 
@@ -115,12 +94,6 @@ public:
 	template<typename ResourceType>
 	void ResourceUploadTo(FRHICommandList& RHICmdList, const ResourceType& DstBuffer, bool bFlush = false);
 	
-	UE_DEPRECATED(5.3, "Scattered uploads to FTextureRWBuffer will no longer supported in a future release. Use a supported destination buffer type instead.")
-	inline void ResourceUploadTo(FRHICommandList& RHICmdList, const FTextureRWBuffer& DstBuffer, bool bFlush = false)
-	{
-		ResourceUploadTo<FTextureRWBuffer>(RHICmdList, DstBuffer, bFlush);
-	}
-
 	void Add( uint32 Index, const void* Data, uint32 Num = 1 )
 	{
 		void* Dst = Add_GetRef( Index, Num );
@@ -222,15 +195,6 @@ extern RENDERCORE_API void MemcpyResource(FRDGBuilder& GraphBuilder, FRDGBuffer*
 extern RENDERCORE_API void MemsetResource(FRDGBuilder& GraphBuilder, FRDGBufferUAV* DstResource, const FMemsetResourceParams& Params);
 extern RENDERCORE_API void MemcpyResource(FRDGBuilder& GraphBuilder, FRDGBufferUAV* DstResource, FRDGBufferSRV* SrcResource, const FMemcpyResourceParams& Params);
 
-UE_DEPRECATED(5.3, "Use the overloads that take FRDGBuffer/FRDGBufferUAV instead.")
-extern RENDERCORE_API void MemsetResource(FRDGBuilder& GraphBuilder, FRDGUnorderedAccessView* DstResource, const FMemsetResourceParams& Params);
-UE_DEPRECATED(5.3, "MemsetResource will no longer support FRDGTextureUAV in a future release. Either use a buffer target instead, or clear the contents manually.")
-extern RENDERCORE_API void MemsetResource(FRDGBuilder& GraphBuilder, FRDGTextureUAV* DstResource, const FMemsetResourceParams& Params);
-UE_DEPRECATED(5.3, "Use the overloads that take FRDGBuffer/FRDGBufferUAV instead.")
-extern RENDERCORE_API void MemcpyResource(FRDGBuilder& GraphBuilder, FRDGUnorderedAccessView* DstResource, FRDGShaderResourceView* SrcResource, const FMemcpyResourceParams& Params);
-UE_DEPRECATED(5.3, "MemcpyResource will no longer support FRDGTextureUAV in a future release. Either use a buffer target instead, or copy the contents manually.")
-extern RENDERCORE_API void MemcpyResource(FRDGBuilder& GraphBuilder, FRDGTextureUAV* DstResource, FRDGTextureSRV* SrcResource, const FMemcpyResourceParams& Params);
-
 struct FScatterCopyParams
 {
 	uint32 NumScatters = 0u;
@@ -240,12 +204,6 @@ struct FScatterCopyParams
 
 void RENDERCORE_API ScatterCopyResource(FRDGBuilder& GraphBuilder, FRDGBuffer* DstResource, FRDGBufferSRV* ScatterBufferSRV, FRDGBufferSRV* UploadBufferSRV, const FScatterCopyParams &Params);
 
-UE_DEPRECATED(5.3, "Use the overload that takes FRDGBuffer instead.")
-void RENDERCORE_API ScatterCopyResource(FRDGBuilder& GraphBuilder, FRDGViewableResource* DstResource, FRDGBufferSRV* ScatterBufferSRV, FRDGBufferSRV* UploadBufferSRV, const FScatterCopyParams &Params);
-UE_DEPRECATED(5.3, "ScatterCopyResource will no longer support FRDGTextureUAV in a future release. Either use a buffer target instead, or copy the contents manually.")
-void RENDERCORE_API ScatterCopyResource(FRDGBuilder& GraphBuilder, FRDGTexture* DstResource, FRDGBufferSRV* ScatterBufferSRV, FRDGBufferSRV* UploadBufferSRV, const FScatterCopyParams &Params);
-
-
 extern RENDERCORE_API FRDGBuffer* ResizeBufferIfNeeded(FRDGBuilder& GraphBuilder, TRefCountPtr<FRDGPooledBuffer>& ExternalBuffer, const FRDGBufferDesc& BufferDesc, const TCHAR* Name);
 extern RENDERCORE_API FRDGBuffer* ResizeBufferIfNeeded(FRDGBuilder& GraphBuilder, TRefCountPtr<FRDGPooledBuffer>& ExternalBuffer, EPixelFormat Format, uint32 NumElements, const TCHAR* Name);
 extern RENDERCORE_API FRDGBuffer* ResizeStructuredBufferIfNeeded(FRDGBuilder& GraphBuilder, TRefCountPtr<FRDGPooledBuffer>& ExternalBuffer, uint32 NumBytes, const TCHAR* Name);
@@ -401,15 +359,6 @@ class FRDGScatterUploadBuffer
 	: public FRDGScatterUploadBase
 {
 public:
-	
-	enum
-	UE_DEPRECATED(5.3, "Use PRIMITIVE_SCENE_DATA_STRIDE in SceneDefinitions.h instead.")
-	{
-		PRAGMA_DISABLE_DEPRECATION_WARNINGS
-		PrimitiveDataStrideInFloat4s = FScatterUploadBuffer::PrimitiveDataStrideInFloat4s
-		PRAGMA_ENABLE_DEPRECATION_WARNINGS
-	};
-
 	/**
 	 * Init with presized num scatters, expecting each to be set at a later point. Requires the user to keep track of the offsets to use.
 	 */
@@ -427,12 +376,6 @@ public:
 		ResourceUploadToInternal(GraphBuilder, DstResource);
 	}
 
-	UE_DEPRECATED(5.3, "Scattered uploads to FTextureRWBuffer will no longer supported in a future release. Use a supported destination buffer type instead.")
-	inline void ResourceUploadTo(FRDGBuilder& GraphBuilder, FRDGTexture* DstResource)
-	{
-		ResourceUploadToInternal(GraphBuilder, DstResource);
-	}
-
 	RENDERCORE_API void Release();
 
 	RENDERCORE_API uint32 GetNumBytes() const;
