Refactor UnifiedBuffer.cpp to share logic between ScatterCopyResource(...) and FRDGAsyncScatterUploadBuffer::End(...)
- Split GetScatterUploadComputeConfig(...) into GetScatterUploadConfig(...) and GetScatterUploadDispatchConfig(...) so that only the latter requires NumScatters as input.
- Moved logic from ScatterCopyResource(...) to a new helper function GetScatterCopyParamsAndPermutation(...) which is then also used by FRDGAsyncScatterUploadBuffer::End(...).

diff --git a/Engine/Source/Runtime/RenderCore/Private/UnifiedBuffer.cpp b/Engine/Source/Runtime/RenderCore/Private/UnifiedBuffer.cpp
--- a/Engine/Source/Runtime/RenderCore/Private/UnifiedBuffer.cpp
+++ b/Engine/Source/Runtime/RenderCore/Private/UnifiedBuffer.cpp
@@ -641,22 +641,15 @@ void FRDGScatterUploadBuffer::InitPreSized(FRDGBuilder& GraphBuilder, uint32 Num
 	NumScatters = NumElements;
 }
 
-struct FScatterUploadComputeConfig
+struct FScatterUploadConfig
 {
-	uint32 ThreadGroupSize;
 	uint32 NumBytesPerThread;
 	uint32 NumThreadsPerScatter;
-	uint32 NumThreads;
-	uint32 NumDispatches;
-	uint32 NumLoops;
 };
 
-FScatterUploadComputeConfig GetScatterUploadComputeConfig(uint32 NumScatters, uint32 NumBytesPerElement, int32 NumElementsPerScatter = -1)
+static FScatterUploadConfig GetScatterUploadConfig(uint32 NumBytesPerElement, int32 NumElementsPerScatter = -1)
 {
-	constexpr uint32 ThreadGroupSize = 64u;
-
-	FScatterUploadComputeConfig Config;
-	Config.ThreadGroupSize = ThreadGroupSize;
+	FScatterUploadConfig Config;
 	if (NumElementsPerScatter != INDEX_NONE)
 	{
 		Config.NumBytesPerThread = NumBytesPerElement;
@@ -667,28 +660,50 @@ FScatterUploadComputeConfig GetScatterUploadComputeConfig(uint32 NumScatters, ui
 		Config.NumBytesPerThread = (NumBytesPerElement & 15) == 0 ? 16 : 4;
 		Config.NumThreadsPerScatter = NumBytesPerElement / Config.NumBytesPerThread;
 	}
-	Config.NumThreads = NumScatters * Config.NumThreadsPerScatter;
-	Config.NumDispatches = FMath::DivideAndRoundUp(Config.NumThreads, Config.ThreadGroupSize);
-	Config.NumLoops = FMath::DivideAndRoundUp(Config.NumDispatches, (uint32)GMaxComputeDispatchDimension);
 	return Config;
 }
 
+struct FScatterUploadDispatchConfig
+{
+	uint32 ThreadGroupSize;
+	uint32 NumThreads;
+	uint32 NumDispatches;
+	uint32 NumLoops;
+};
 
-void ScatterCopyResource(FRDGBuilder& GraphBuilder, FRDGViewableResource* DstResource, FRDGBufferSRV* ScatterBufferSRV, FRDGBufferSRV* UploadBufferSRV, const FScatterCopyParams &Params)
+static FScatterUploadDispatchConfig GetScatterUploadDispatchConfig(const FScatterUploadConfig& Config, uint32 NumScatters)
+{
+	constexpr uint32 ThreadGroupSize = 64u;
+
+	FScatterUploadDispatchConfig DispatchConfig;
+	DispatchConfig.ThreadGroupSize = ThreadGroupSize;
+	DispatchConfig.NumThreads = NumScatters * Config.NumThreadsPerScatter;
+	DispatchConfig.NumDispatches = FMath::DivideAndRoundUp(DispatchConfig.NumThreads, DispatchConfig.ThreadGroupSize);
+	DispatchConfig.NumLoops = FMath::DivideAndRoundUp(DispatchConfig.NumDispatches, (uint32)GMaxComputeDispatchDimension);
+	return DispatchConfig;
+}
+
+static void GetScatterCopyParamsAndPermutation(
+	FRDGBuilder& GraphBuilder,
+	FRDGViewableResource* DstResource,
+	FRDGBufferSRV* ScatterBufferSRV,
+	FRDGBufferSRV* UploadBufferSRV,
+	const FScatterUploadConfig& Config,
+	uint32 NumScatters,
+	FRDGScatterCopyCS::FParameters& OutParameters,
+	FRDGScatterCopyCS::FPermutationDomain& OutPermutationVector)
 {
-	const FScatterUploadComputeConfig ComputeConfig = GetScatterUploadComputeConfig(Params.NumScatters, Params.NumBytesPerElement, Params.NumElementsPerScatter);
 	const EResourceType DstResourceType = GetResourceType(DstResource);
 
 	EByteBufferResourceType ResourceTypeEnum = EByteBufferResourceType::Count;
 	EByteBufferStructuredSize ByteBufferStructuredSize = EByteBufferStructuredSize::Uint4;
 
-	FRDGScatterCopyCS::FParameters Parameters;
-	Parameters.Common.Size = ComputeConfig.NumThreadsPerScatter;
-	Parameters.NumScatters = Params.NumScatters;
+	OutParameters.Common.Size = Config.NumThreadsPerScatter;
+	OutParameters.NumScatters = NumScatters;
 
 	if (DstResourceType == EResourceType::BYTEBUFFER)
 	{
-		if (ComputeConfig.NumBytesPerThread == 16)
+		if (Config.NumBytesPerThread == 16)
 		{
 			ResourceTypeEnum = EByteBufferResourceType::Uint4Aligned_Buffer;
 		}
@@ -696,58 +711,68 @@ void ScatterCopyResource(FRDGBuilder& GraphBuilder, FRDGViewableResource* DstRes
 		{
 			ResourceTypeEnum = EByteBufferResourceType::Uint_Buffer;
 		}
-		Parameters.UploadByteAddressBuffer = UploadBufferSRV;
-		Parameters.ScatterByteAddressBuffer = ScatterBufferSRV;
-		Parameters.Common.DstByteAddressBuffer = GraphBuilder.CreateUAV(GetAsBuffer(DstResource), ERDGUnorderedAccessViewFlags::SkipBarrier);
+		OutParameters.UploadByteAddressBuffer = UploadBufferSRV;
+		OutParameters.ScatterByteAddressBuffer = ScatterBufferSRV;
+		OutParameters.Common.DstByteAddressBuffer = GraphBuilder.CreateUAV(GetAsBuffer(DstResource), ERDGUnorderedAccessViewFlags::SkipBarrier);
 	}
 	else if (DstResourceType == EResourceType::STRUCTURED_BUFFER)
 	{
 		ResourceTypeEnum = EByteBufferResourceType::StructuredBuffer;
 
-		FRDGBufferUAV * RDGBufferUAV = GraphBuilder.CreateUAV(GetAsBuffer(DstResource), ERDGUnorderedAccessViewFlags::SkipBarrier);
+		FRDGBufferUAV* RDGBufferUAV = GraphBuilder.CreateUAV(GetAsBuffer(DstResource), ERDGUnorderedAccessViewFlags::SkipBarrier);
 		ByteBufferStructuredSize = GetStructuredBufferElementSize(RDGBufferUAV);
-		switch(ByteBufferStructuredSize)
+		switch (ByteBufferStructuredSize)
 		{
 		case EByteBufferStructuredSize::Uint1:
-			Parameters.UploadStructuredBuffer1x = UploadBufferSRV;
-			Parameters.Common.DstStructuredBuffer1x = RDGBufferUAV;
+			OutParameters.UploadStructuredBuffer1x = UploadBufferSRV;
+			OutParameters.Common.DstStructuredBuffer1x = RDGBufferUAV;
 			break;
 		case EByteBufferStructuredSize::Uint2:
-			Parameters.UploadStructuredBuffer2x = UploadBufferSRV;
-			Parameters.Common.DstStructuredBuffer2x = RDGBufferUAV;
+			OutParameters.UploadStructuredBuffer2x = UploadBufferSRV;
+			OutParameters.Common.DstStructuredBuffer2x = RDGBufferUAV;
 			break;
 		case EByteBufferStructuredSize::Uint4:
-			Parameters.UploadStructuredBuffer4x = UploadBufferSRV;
-			Parameters.Common.DstStructuredBuffer4x = RDGBufferUAV;
+			OutParameters.UploadStructuredBuffer4x = UploadBufferSRV;
+			OutParameters.Common.DstStructuredBuffer4x = RDGBufferUAV;
 			break;
 		case EByteBufferStructuredSize::Uint8:
-			Parameters.UploadStructuredBuffer8x = UploadBufferSRV;
-			Parameters.Common.DstStructuredBuffer8x = RDGBufferUAV;
+			OutParameters.UploadStructuredBuffer8x = UploadBufferSRV;
+			OutParameters.Common.DstStructuredBuffer8x = RDGBufferUAV;
 			break;
 		default:
 			break;
 		};
-		Parameters.ScatterStructuredBuffer = ScatterBufferSRV;
+		OutParameters.ScatterStructuredBuffer = ScatterBufferSRV;
 	}
 	else if (DstResourceType == EResourceType::BUFFER)
 	{
 		ResourceTypeEnum = EByteBufferResourceType::Float4_Buffer;
 
-		Parameters.UploadStructuredBuffer4x = UploadBufferSRV;
-		Parameters.ScatterStructuredBuffer = ScatterBufferSRV;
-		Parameters.Common.DstBuffer = GraphBuilder.CreateUAV(GetAsBuffer(DstResource), ERDGUnorderedAccessViewFlags::SkipBarrier);
+		OutParameters.UploadStructuredBuffer4x = UploadBufferSRV;
+		OutParameters.ScatterStructuredBuffer = ScatterBufferSRV;
+		OutParameters.Common.DstBuffer = GraphBuilder.CreateUAV(GetAsBuffer(DstResource), ERDGUnorderedAccessViewFlags::SkipBarrier);
 	}
 
+	OutPermutationVector.Set<FRDGScatterCopyCS::ResourceTypeDim>((int)ResourceTypeEnum);
+	OutPermutationVector.Set<FRDGScatterCopyCS::StructuredElementSizeDim >((int)ByteBufferStructuredSize);
+}
+
+void ScatterCopyResource(FRDGBuilder& GraphBuilder, FRDGViewableResource* DstResource, FRDGBufferSRV* ScatterBufferSRV, FRDGBufferSRV* UploadBufferSRV, const FScatterCopyParams &Params)
+{
+	const FScatterUploadConfig Config = GetScatterUploadConfig(Params.NumBytesPerElement, Params.NumElementsPerScatter);
+	const FScatterUploadDispatchConfig DispatchConfig = GetScatterUploadDispatchConfig(Config, Params.NumScatters);
+
+	FRDGScatterCopyCS::FParameters Parameters;
 	FRDGScatterCopyCS::FPermutationDomain PermutationVector;
-	PermutationVector.Set<FRDGScatterCopyCS::ResourceTypeDim>((int)ResourceTypeEnum);
-	PermutationVector.Set<FRDGScatterCopyCS::StructuredElementSizeDim >((int)ByteBufferStructuredSize);
+	GetScatterCopyParamsAndPermutation(GraphBuilder, DstResource, ScatterBufferSRV, UploadBufferSRV, Config, Params.NumScatters, Parameters, PermutationVector);
+
 	TShaderMapRef<FRDGScatterCopyCS> ComputeShader(GetGlobalShaderMap(GMaxRHIFeatureLevel), PermutationVector);
 
-	for (uint32 LoopIdx = 0; LoopIdx < ComputeConfig.NumLoops; ++LoopIdx)
+	for (uint32 LoopIdx = 0; LoopIdx < DispatchConfig.NumLoops; ++LoopIdx)
 	{
-		Parameters.Common.SrcOffset = LoopIdx * (uint32)GMaxComputeDispatchDimension * ComputeConfig.ThreadGroupSize;
+		Parameters.Common.SrcOffset = LoopIdx * (uint32)GMaxComputeDispatchDimension * DispatchConfig.ThreadGroupSize;
 
-		uint32 LoopNumDispatch = FMath::Min(ComputeConfig.NumDispatches - LoopIdx * (uint32)GMaxComputeDispatchDimension, (uint32)GMaxComputeDispatchDimension);
+		uint32 LoopNumDispatch = FMath::Min(DispatchConfig.NumDispatches - LoopIdx * (uint32)GMaxComputeDispatchDimension, (uint32)GMaxComputeDispatchDimension);
 
 		FComputeShaderUtils::AddPass(
 			GraphBuilder,
@@ -935,87 +960,29 @@ void FRDGAsyncScatterUploadBuffer::End(FRDGBuilder& GraphBuilder, FRDGScatterUpl
 
 	const uint32 NumScatters = Uploader->bNumScattersPreSized ? Uploader->NumScatters : Uploader->MaxScatters;
 
-	const FScatterUploadComputeConfig ComputeConfig = GetScatterUploadComputeConfig(NumScatters, Uploader->NumBytesPerElement);
+	const FScatterUploadConfig Config = GetScatterUploadConfig(Uploader->NumBytesPerElement);
+	const FScatterUploadDispatchConfig DispatchConfig = GetScatterUploadDispatchConfig(Config, NumScatters);
 
-	FRDGScatterCopyCS::FParameters Parameters;
-	Parameters.Common.Size = ComputeConfig.NumThreadsPerScatter;
-	Parameters.NumScatters = NumScatters;
+	FRDGViewableResource* DstResource = Uploader->DstResource;
 
 	FRDGBufferSRV* ScatterBufferSRV = GraphBuilder.CreateSRV(GraphBuilder.RegisterExternalBuffer(ScatterBuffer));
 	FRDGBufferSRV* UploadBufferSRV = GraphBuilder.CreateSRV(GraphBuilder.RegisterExternalBuffer(UploadBuffer));
 
-	FRDGViewableResource* DstResource = Uploader->DstResource;
-	const EResourceType DstResourceType = GetResourceType(DstResource);
-
-	EByteBufferResourceType ResourceTypeEnum = EByteBufferResourceType::Count;
-	EByteBufferStructuredSize ByteBufferStructuredSize = EByteBufferStructuredSize::Uint4;
-
-	if (DstResourceType == EResourceType::BYTEBUFFER)
-	{
-		if (ComputeConfig.NumBytesPerThread == 16)
-		{
-			ResourceTypeEnum = EByteBufferResourceType::Uint4Aligned_Buffer;
-		}
-		else
-		{
-			ResourceTypeEnum = EByteBufferResourceType::Uint_Buffer;
-		}
-		Parameters.UploadByteAddressBuffer = UploadBufferSRV;
-		Parameters.ScatterByteAddressBuffer = ScatterBufferSRV;
-		Parameters.Common.DstByteAddressBuffer = GraphBuilder.CreateUAV(GetAsBuffer(DstResource), ERDGUnorderedAccessViewFlags::SkipBarrier);
-	}
-	else if (DstResourceType == EResourceType::STRUCTURED_BUFFER)
-	{
-		ResourceTypeEnum = EByteBufferResourceType::StructuredBuffer;
-
-		FRDGBufferUAV * RDGBufferUAV = GraphBuilder.CreateUAV(GetAsBuffer(DstResource), ERDGUnorderedAccessViewFlags::SkipBarrier);
-		ByteBufferStructuredSize = GetStructuredBufferElementSize(RDGBufferUAV);
-		switch(ByteBufferStructuredSize)
-		{
-		case EByteBufferStructuredSize::Uint1:
-			Parameters.UploadStructuredBuffer1x = UploadBufferSRV;
-			Parameters.Common.DstStructuredBuffer1x = RDGBufferUAV;
-			break;
-		case EByteBufferStructuredSize::Uint2:
-			Parameters.UploadStructuredBuffer2x = UploadBufferSRV;
-			Parameters.Common.DstStructuredBuffer2x = RDGBufferUAV;
-			break;
-		case EByteBufferStructuredSize::Uint4:
-			Parameters.UploadStructuredBuffer4x = UploadBufferSRV;
-			Parameters.Common.DstStructuredBuffer4x = RDGBufferUAV;
-			break;
-		case EByteBufferStructuredSize::Uint8:
-			Parameters.UploadStructuredBuffer8x = UploadBufferSRV;
-			Parameters.Common.DstStructuredBuffer8x = RDGBufferUAV;
-			break;
-		default:
-			break;
-		};
-		Parameters.ScatterStructuredBuffer = ScatterBufferSRV;
-	}
-	else if (DstResourceType == EResourceType::BUFFER)
-	{
-		ResourceTypeEnum = EByteBufferResourceType::Float4_Buffer;
-
-		Parameters.UploadStructuredBuffer4x = UploadBufferSRV;
-		Parameters.ScatterStructuredBuffer = ScatterBufferSRV;
-		Parameters.Common.DstBuffer = GraphBuilder.CreateUAV(GetAsBuffer(DstResource), ERDGUnorderedAccessViewFlags::SkipBarrier);
-	}
-
+	FRDGScatterCopyCS::FParameters Parameters;
 	FRDGScatterCopyCS::FPermutationDomain PermutationVector;
-	PermutationVector.Set<FRDGScatterCopyCS::ResourceTypeDim>((int)ResourceTypeEnum);
-	PermutationVector.Set<FRDGScatterCopyCS::StructuredElementSizeDim >(static_cast<int32>(ByteBufferStructuredSize));
+	GetScatterCopyParamsAndPermutation(GraphBuilder, DstResource, ScatterBufferSRV, UploadBufferSRV, Config, NumScatters, Parameters, PermutationVector);
+
 	TShaderMapRef<FRDGScatterCopyCS> ComputeShader(GetGlobalShaderMap(GMaxRHIFeatureLevel), PermutationVector);
 
-	for (uint32 LoopIdx = 0; LoopIdx < ComputeConfig.NumLoops; ++LoopIdx)
+	for (uint32 LoopIdx = 0; LoopIdx < DispatchConfig.NumLoops; ++LoopIdx)
 	{
-		Parameters.Common.SrcOffset = LoopIdx * (uint32)GMaxComputeDispatchDimension * ComputeConfig.ThreadGroupSize;
+		Parameters.Common.SrcOffset = LoopIdx * (uint32)GMaxComputeDispatchDimension * DispatchConfig.ThreadGroupSize;
 
 		FRDGScatterCopyCS::FParameters* PassParameters = GraphBuilder.AllocParameters(&Parameters);
 
 		if (Uploader->bNumScattersPreSized)
 		{
-			const uint32 LoopNumDispatch = FMath::Min(ComputeConfig.NumDispatches - LoopIdx * (uint32)GMaxComputeDispatchDimension, (uint32)GMaxComputeDispatchDimension);
+			const uint32 LoopNumDispatch = FMath::Min(DispatchConfig.NumDispatches - LoopIdx * (uint32)GMaxComputeDispatchDimension, (uint32)GMaxComputeDispatchDimension);
 
 			FComputeShaderUtils::AddPass(
 				GraphBuilder,
@@ -1031,16 +998,16 @@ void FRDGAsyncScatterUploadBuffer::End(FRDGBuilder& GraphBuilder, FRDGScatterUpl
 				RDG_EVENT_NAME("ScatterUpload[%d] (Resource: %s, Offset: %u)", LoopIdx, DstResource->Name, Parameters.Common.SrcOffset),
 				ComputeShader,
 				PassParameters,
-				[PassParameters, LoopIdx, NumBytesPerElement = Uploader->NumBytesPerElement, Uploader]
+				[PassParameters, LoopIdx, Config, Uploader]
 			{
 				const uint32 NumScatters = Uploader->GetFinalNumScatters();
-				const FScatterUploadComputeConfig ComputeConfig = GetScatterUploadComputeConfig(NumScatters, NumBytesPerElement);
+				const FScatterUploadDispatchConfig DispatchConfig = GetScatterUploadDispatchConfig(Config, NumScatters);
 
-				if (LoopIdx < ComputeConfig.NumLoops)
+				if (LoopIdx < DispatchConfig.NumLoops)
 				{
 					PassParameters->NumScatters = NumScatters;
 
-					return FIntVector(FMath::Min(ComputeConfig.NumDispatches - LoopIdx * (uint32)GMaxComputeDispatchDimension, (uint32)GMaxComputeDispatchDimension), 1, 1);
+					return FIntVector(FMath::Min(DispatchConfig.NumDispatches - LoopIdx * (uint32)GMaxComputeDispatchDimension, (uint32)GMaxComputeDispatchDimension), 1, 1);
 				}
 				else
 				{
