Refactored UnifiedBuffer to use GroupCountWrapped.

diff --git a/Engine/Shaders/Private/ByteBuffer.usf b/Engine/Shaders/Private/ByteBuffer.usf
--- a/Engine/Shaders/Private/ByteBuffer.usf
+++ b/Engine/Shaders/Private/ByteBuffer.usf
@@ -15,7 +15,6 @@
 
 uint Value;
 uint Size;
-uint Float4sPerLine;
 uint NumScatters;
 uint SrcOffset;
 uint DstOffset;
@@ -125,8 +124,10 @@ void MemsetBufferCS( uint3 GroupID : SV_GroupID, uint GroupIndex : SV_GroupIndex
 }
 
 [numthreads(64, 1, 1)]
-void MemcpyCS( uint ThreadId : SV_DispatchThreadID ) 
+void MemcpyCS( uint3 GroupID : SV_GroupID, uint GroupIndex : SV_GroupIndex )
 {
+	const uint ThreadId = GetUnWrappedDispatchThreadId(GroupID, GroupIndex, 64);
+
 #if RESOURCE_TYPE == RESOURCE_TYPE_FLOAT4_BUFFER
 	// Size is in float4s
 	if (ThreadId < Size)
@@ -170,9 +171,9 @@ void MemcpyCS( uint ThreadId : SV_DispatchThreadID )
 }
 
 [numthreads(64, 1, 1)]
-void ScatterCopyCS( uint DispatchThreadId : SV_DispatchThreadID ) 
+void ScatterCopyCS( uint3 GroupID : SV_GroupID, uint GroupIndex : SV_GroupIndex )
 {
-	uint ThreadId = DispatchThreadId + SrcOffset;
+	const uint ThreadId = GetUnWrappedDispatchThreadId(GroupID, GroupIndex, 64) + SrcOffset;
 
 	uint ScatterIndex = ThreadId / Size;
 	uint ScatterOffset = ThreadId - ScatterIndex * Size;
diff --git a/Engine/Source/Runtime/RenderCore/Private/UnifiedBuffer.cpp b/Engine/Source/Runtime/RenderCore/Private/UnifiedBuffer.cpp
--- a/Engine/Source/Runtime/RenderCore/Private/UnifiedBuffer.cpp
+++ b/Engine/Source/Runtime/RenderCore/Private/UnifiedBuffer.cpp
@@ -82,7 +82,6 @@ class FByteBufferShader : public FGlobalShader
 		SHADER_PARAMETER(uint32, Size)
 		SHADER_PARAMETER(uint32, SrcOffset)
 		SHADER_PARAMETER(uint32, DstOffset)
-		SHADER_PARAMETER(uint32, Float4sPerLine)
 		SHADER_PARAMETER_UAV(RWBuffer<float4>, DstBuffer)
 		SHADER_PARAMETER_UAV(RWStructuredBuffer<float4>, DstStructuredBuffer4x)
 		SHADER_PARAMETER_UAV(RWByteAddressBuffer, DstByteAddressBuffer)
@@ -212,7 +211,6 @@ class FRDGByteBufferShader : public FGlobalShader
 		SHADER_PARAMETER(uint32, Size)
 		SHADER_PARAMETER(uint32, SrcOffset)
 		SHADER_PARAMETER(uint32, DstOffset)
-		SHADER_PARAMETER(uint32, Float4sPerLine)
 		SHADER_PARAMETER_RDG_BUFFER_UAV(RWBuffer<float4>, DstBuffer)
 		SHADER_PARAMETER_RDG_BUFFER_UAV(RWStructuredBuffer<FUint8>, DstStructuredBuffer8x)
 		SHADER_PARAMETER_RDG_BUFFER_UAV(RWStructuredBuffer<uint4>, DstStructuredBuffer4x)
@@ -400,25 +398,17 @@ void MemcpyResource(FRDGBuilder& GraphBuilder, FRDGUnorderedAccessView* UAV, FRD
 	const EResourceType ResourceType = GetResourceType(DstResource);
 	checkf(ResourceType == GetResourceType(SrcResource), TEXT("Unable to MemcpyResource because the source and destination view types don't match."));
 
-	RDG_EVENT_SCOPE(GraphBuilder, "Memcpy %s -> %s", DstResource->Name, SrcResource->Name);
-
 	// each thread will copy 4 floats / uints
 	const uint32 Divisor = ResourceType == EResourceType::BYTEBUFFER ? 4 : 1;
 
-	uint32 NumElementsProcessed = 0;
-
-	while (NumElementsProcessed < Params.Count)
 	{
-		const uint32 NumWaves = FMath::Max(FMath::Min<uint32>(GRHIMaxDispatchThreadGroupsPerDimension.X, FMath::DivideAndRoundUp(Params.Count / Divisor, 64u)), 1u);
-		const uint32 NumElementsPerDispatch = FMath::Min(FMath::Max(NumWaves, 1u) * Divisor * 64, Params.Count - NumElementsProcessed);
-
 		EByteBufferResourceType ResourceTypeEnum = EByteBufferResourceType::Count;
 		EByteBufferStructuredSize ByteBufferStructuredSize = EByteBufferStructuredSize::Uint4;
 
 		auto* PassParameters = GraphBuilder.AllocParameters<FRDGMemcpyCS::FParameters>();
-		PassParameters->Common.Size = NumElementsPerDispatch;
-		PassParameters->Common.SrcOffset = (Params.SrcOffset + NumElementsProcessed);
-		PassParameters->Common.DstOffset = (Params.DstOffset + NumElementsProcessed);
+		PassParameters->Common.Size = Params.Count;
+		PassParameters->Common.SrcOffset = Params.SrcOffset;
+		PassParameters->Common.DstOffset = Params.DstOffset;
 
 		switch (ResourceType)
 		{
@@ -472,12 +462,10 @@ void MemcpyResource(FRDGBuilder& GraphBuilder, FRDGUnorderedAccessView* UAV, FRD
 
 		FComputeShaderUtils::AddPass(
 			GraphBuilder,
-			RDG_EVENT_NAME("Offset[%d] Count[%d]", NumElementsProcessed, NumElementsPerDispatch),
+			RDG_EVENT_NAME("Memcpy %s -> %s", DstResource->Name, SrcResource->Name),
 			ComputeShader,
 			PassParameters,
-			FIntVector(NumWaves, 1, 1));
-
-		NumElementsProcessed += NumElementsPerDispatch;
+			FComputeShaderUtils::GetGroupCountWrapped(Params.Count / Divisor, 64));
 	}
 }
 
@@ -667,8 +655,7 @@ struct FScatterUploadDispatchConfig
 {
 	uint32 ThreadGroupSize;
 	uint32 NumThreads;
-	uint32 NumDispatches;
-	uint32 NumLoops;
+	FIntVector GroupCountWrapped;
 };
 
 static FScatterUploadDispatchConfig GetScatterUploadDispatchConfig(const FScatterUploadConfig& Config, uint32 NumScatters)
@@ -678,8 +665,7 @@ static FScatterUploadDispatchConfig GetScatterUploadDispatchConfig(const FScatte
 	FScatterUploadDispatchConfig DispatchConfig;
 	DispatchConfig.ThreadGroupSize = ThreadGroupSize;
 	DispatchConfig.NumThreads = NumScatters * Config.NumThreadsPerScatter;
-	DispatchConfig.NumDispatches = FMath::DivideAndRoundUp(DispatchConfig.NumThreads, DispatchConfig.ThreadGroupSize);
-	DispatchConfig.NumLoops = FMath::DivideAndRoundUp(DispatchConfig.NumDispatches, (uint32)GMaxComputeDispatchDimension);
+	DispatchConfig.GroupCountWrapped = FComputeShaderUtils::GetGroupCountWrapped(DispatchConfig.NumThreads, DispatchConfig.ThreadGroupSize);
 	return DispatchConfig;
 }
 
@@ -699,6 +685,8 @@ static void GetScatterCopyParamsAndPermutation(
 	EByteBufferStructuredSize ByteBufferStructuredSize = EByteBufferStructuredSize::Uint4;
 
 	OutParameters.Common.Size = Config.NumThreadsPerScatter;
+	OutParameters.Common.SrcOffset = 0;
+	OutParameters.Common.DstOffset = 0;
 	OutParameters.NumScatters = NumScatters;
 
 	if (DstResourceType == EResourceType::BYTEBUFFER)
@@ -768,19 +756,12 @@ void ScatterCopyResource(FRDGBuilder& GraphBuilder, FRDGViewableResource* DstRes
 
 	TShaderMapRef<FRDGScatterCopyCS> ComputeShader(GetGlobalShaderMap(GMaxRHIFeatureLevel), PermutationVector);
 
-	for (uint32 LoopIdx = 0; LoopIdx < DispatchConfig.NumLoops; ++LoopIdx)
-	{
-		Parameters.Common.SrcOffset = LoopIdx * (uint32)GMaxComputeDispatchDimension * DispatchConfig.ThreadGroupSize;
-
-		uint32 LoopNumDispatch = FMath::Min(DispatchConfig.NumDispatches - LoopIdx * (uint32)GMaxComputeDispatchDimension, (uint32)GMaxComputeDispatchDimension);
-
-		FComputeShaderUtils::AddPass(
-			GraphBuilder,
-			RDG_EVENT_NAME("ScatterUpload[%d] (Resource: %s, Offset: %u, GroupSize: %u)", LoopIdx, DstResource->Name, Parameters.Common.SrcOffset, LoopNumDispatch),
-			ComputeShader,
-			GraphBuilder.AllocParameters(&Parameters),
-			FIntVector(LoopNumDispatch, 1, 1));
-	}
+	FComputeShaderUtils::AddPass(
+		GraphBuilder,
+		RDG_EVENT_NAME("ScatterUpload (Resource: %s)", DstResource->Name),
+		ComputeShader,
+		GraphBuilder.AllocParameters(&Parameters),
+		DispatchConfig.GroupCountWrapped);
 }
 
 void ScatterCopyResource(FRDGBuilder& GraphBuilder, FRDGBuffer* DstResource, FRDGBufferSRV* ScatterBufferSRV, FRDGBufferSRV* UploadBufferSRV, const FScatterCopyParams &Params)
@@ -974,45 +955,33 @@ void FRDGAsyncScatterUploadBuffer::End(FRDGBuilder& GraphBuilder, FRDGScatterUpl
 
 	TShaderMapRef<FRDGScatterCopyCS> ComputeShader(GetGlobalShaderMap(GMaxRHIFeatureLevel), PermutationVector);
 
-	for (uint32 LoopIdx = 0; LoopIdx < DispatchConfig.NumLoops; ++LoopIdx)
 	{
-		Parameters.Common.SrcOffset = LoopIdx * (uint32)GMaxComputeDispatchDimension * DispatchConfig.ThreadGroupSize;
-
 		FRDGScatterCopyCS::FParameters* PassParameters = GraphBuilder.AllocParameters(&Parameters);
 
 		if (Uploader->bNumScattersPreSized)
 		{
-			const uint32 LoopNumDispatch = FMath::Min(DispatchConfig.NumDispatches - LoopIdx * (uint32)GMaxComputeDispatchDimension, (uint32)GMaxComputeDispatchDimension);
-
 			FComputeShaderUtils::AddPass(
 				GraphBuilder,
-				RDG_EVENT_NAME("ScatterUpload[%d] (Resource: %s, Offset: %u, GroupSize: %u)", LoopIdx, DstResource->Name, Parameters.Common.SrcOffset, LoopNumDispatch),
+				RDG_EVENT_NAME("ScatterUpload (Resource: %s)", DstResource->Name),
 				ComputeShader,
 				PassParameters,
-				FIntVector(LoopNumDispatch, 1, 1));
+				DispatchConfig.GroupCountWrapped);
 		}
 		else
 		{
 			FComputeShaderUtils::AddPass(
 				GraphBuilder,
-				RDG_EVENT_NAME("ScatterUpload[%d] (Resource: %s, Offset: %u)", LoopIdx, DstResource->Name, Parameters.Common.SrcOffset),
+				RDG_EVENT_NAME("ScatterUpload (Resource: %s)", DstResource->Name),
 				ComputeShader,
 				PassParameters,
-				[PassParameters, LoopIdx, Config, Uploader]
+				[PassParameters, Config, Uploader]
 			{
 				const uint32 NumScatters = Uploader->GetFinalNumScatters();
 				const FScatterUploadDispatchConfig DispatchConfig = GetScatterUploadDispatchConfig(Config, NumScatters);
 
-				if (LoopIdx < DispatchConfig.NumLoops)
-				{
-					PassParameters->NumScatters = NumScatters;
+				PassParameters->NumScatters = NumScatters;
 
-					return FIntVector(FMath::Min(DispatchConfig.NumDispatches - LoopIdx * (uint32)GMaxComputeDispatchDimension, (uint32)GMaxComputeDispatchDimension), 1, 1);
-				}
-				else
-				{
-					return FIntVector::ZeroValue;
-				}
+				return DispatchConfig.GroupCountWrapped;
 			});
 		}
 	}
@@ -1081,19 +1050,13 @@ void MemcpyResource(FRHICommandList& RHICmdList, const ResourceType& DstBuffer,
 	if (!bAlreadyInUAVOverlap)	// TODO: Get rid of this check once BeginUAVOverlap/EndUAVOverlap supports nesting.
 		RHICmdList.BeginUAVOverlap(DstBuffer.UAV);
 
-	uint32 NumElementsProcessed = 0;
-
-	while (NumElementsProcessed < Params.Count)
 	{
-		const uint32 NumWaves = FMath::Max(FMath::Min<uint32>(GRHIMaxDispatchThreadGroupsPerDimension.X, FMath::DivideAndRoundUp(Params.Count / Divisor, 64u)), 1u);
-		const uint32 NumElementsPerDispatch = FMath::Min(FMath::Max(NumWaves, 1u) * Divisor * 64, Params.Count - NumElementsProcessed);
-
 		EByteBufferResourceType ResourceTypeEnum;
 
 		FMemcpyCS::FParameters Parameters;
-		Parameters.Common.Size = NumElementsPerDispatch;
-		Parameters.Common.SrcOffset = (Params.SrcOffset + NumElementsProcessed);
-		Parameters.Common.DstOffset = (Params.DstOffset + NumElementsProcessed);
+		Parameters.Common.Size = Params.Count;
+		Parameters.Common.SrcOffset = Params.SrcOffset;
+		Parameters.Common.DstOffset = Params.DstOffset;
 
 		if (ResourceTypeTraits<ResourceType>::Type == EResourceType::BYTEBUFFER)
 		{
@@ -1127,9 +1090,11 @@ void MemcpyResource(FRHICommandList& RHICmdList, const ResourceType& DstBuffer,
 
 		auto ComputeShader = GetGlobalShaderMap(GMaxRHIFeatureLevel)->GetShader<FMemcpyCS >(PermutationVector);
 
-		FComputeShaderUtils::Dispatch(RHICmdList, ComputeShader, Parameters, FIntVector(NumWaves, 1, 1));
-
-		NumElementsProcessed += NumElementsPerDispatch;
+		FComputeShaderUtils::Dispatch(
+			RHICmdList,
+			ComputeShader,
+			Parameters,
+			FComputeShaderUtils::GetGroupCountWrapped(Params.Count / Divisor, 64));
 	}
 
 	if(!bAlreadyInUAVOverlap)
@@ -1564,24 +1529,22 @@ void FScatterUploadBuffer::ResourceUploadTo(FRHICommandList& RHICmdList, const R
 		return;
 	}
 
-	constexpr uint32 ThreadGroupSize = 64u;
-	uint32 NumBytesPerThread = (NumBytesPerElement & 15) == 0 ? 16 : 4;
-	uint32 NumThreadsPerScatter = NumBytesPerElement / NumBytesPerThread;
-	uint32 NumThreads = NumScatters * NumThreadsPerScatter;
-	uint32 NumDispatches = FMath::DivideAndRoundUp(NumThreads, ThreadGroupSize);
-	uint32 NumLoops = FMath::DivideAndRoundUp(NumDispatches, (uint32)GMaxComputeDispatchDimension);
+	const FScatterUploadConfig Config = GetScatterUploadConfig(NumBytesPerElement);
+	const FScatterUploadDispatchConfig DispatchConfig = GetScatterUploadDispatchConfig(Config, NumScatters);
 
 	EByteBufferResourceType ResourceTypeEnum;
 
 	FScatterCopyCS::FParameters Parameters;
-	Parameters.Common.Size = NumThreadsPerScatter;
+	Parameters.Common.Size = Config.NumThreadsPerScatter;
+	Parameters.Common.SrcOffset = 0;
+	Parameters.Common.DstOffset = 0;
 	Parameters.NumScatters = NumScatters;
 
 	check(bFloat4Buffer || ResourceTypeTraits<ResourceType>::Type == EResourceType::BYTEBUFFER);
 
 	if (ResourceTypeTraits<ResourceType>::Type == EResourceType::BYTEBUFFER)
 	{
-		if (NumBytesPerThread == 16)
+		if (Config.NumBytesPerThread == 16)
 		{
 			ResourceTypeEnum = EByteBufferResourceType::Uint4Aligned_Buffer;
 		}
@@ -1617,15 +1580,8 @@ void FScatterUploadBuffer::ResourceUploadTo(FRHICommandList& RHICmdList, const R
 	auto ComputeShader = GetGlobalShaderMap(GMaxRHIFeatureLevel)->GetShader<FScatterCopyCS>(PermutationVector);
 
 	RHICmdList.BeginUAVOverlap(DstBuffer.UAV);
-
-	for (uint32 LoopIdx = 0; LoopIdx < NumLoops; ++LoopIdx)
-	{
-		Parameters.Common.SrcOffset = LoopIdx * (uint32)GMaxComputeDispatchDimension * ThreadGroupSize;
-
-		uint32 LoopNumDispatch = FMath::Min(NumDispatches - LoopIdx * (uint32)GMaxComputeDispatchDimension, (uint32)GMaxComputeDispatchDimension);
-
-		FComputeShaderUtils::Dispatch(RHICmdList, ComputeShader, Parameters, FIntVector(LoopNumDispatch, 1, 1));
-	}
+	
+	FComputeShaderUtils::Dispatch(RHICmdList, ComputeShader, Parameters, DispatchConfig.GroupCountWrapped);
 
 	RHICmdList.EndUAVOverlap(DstBuffer.UAV);
 
