Use  RHI_USE_RESOURCE_DEBUG_NAME define to control the usage of DX12/DX11 resources debug name in the build.
Disabled in Shipping and Test builds by default due the performance penalty cost.

diff --git a/Engine/Source/Runtime/D3D12RHI/Private/D3D12RHIDefinitions.h b/Engine/Source/Runtime/D3D12RHI/Private/D3D12RHIDefinitions.h
--- a/Engine/Source/Runtime/D3D12RHI/Private/D3D12RHIDefinitions.h
+++ b/Engine/Source/Runtime/D3D12RHI/Private/D3D12RHIDefinitions.h
@@ -35,7 +35,7 @@
 #define MAX_ROOT_CBVS	MAX_CBS
 
 #define EXECUTE_DEBUG_COMMAND_LISTS 0
-#define NAME_OBJECTS !(UE_BUILD_SHIPPING || UE_BUILD_TEST)	// Name objects in all builds except shipping
+#define NAME_OBJECTS RHI_USE_RESOURCE_DEBUG_NAME
 #define LOG_PSO_CREATES (0 && STATS)	// Logs Create Pipeline State timings (also requires STATS)
 #define TRACK_RESOURCE_ALLOCATIONS (PLATFORM_WINDOWS && !UE_BUILD_SHIPPING && !UE_BUILD_TEST)
 
diff --git a/Engine/Source/Runtime/RHI/Private/RHI.cpp b/Engine/Source/Runtime/RHI/Private/RHI.cpp
--- a/Engine/Source/Runtime/RHI/Private/RHI.cpp
+++ b/Engine/Source/Runtime/RHI/Private/RHI.cpp
@@ -2090,7 +2090,7 @@ FRHIUnorderedAccessView* FRHIBufferViewCache::GetOrCreateUAV(FRHICommandListBase
 	return View;
 }
 
-#if !(UE_BUILD_SHIPPING || UE_BUILD_TEST)
+#if RHI_USE_RESOURCE_DEBUG_NAME
 
 void FRHITextureViewCache::SetDebugName(FRHICommandListBase& RHICmdList, const TCHAR* DebugName)
 {
@@ -2114,7 +2114,7 @@ void FRHITransientTexture::BindDebugLabelName(FRHICommandListBase& RHICmdList)
 {
 	ViewCache.SetDebugName(RHICmdList, GetName());
 
-#if !(UE_BUILD_SHIPPING || UE_BUILD_TEST)
+#if RHI_USE_RESOURCE_DEBUG_NAME
 	RHICmdList.BindDebugLabelName(GetRHI(), GetName());
 #endif
 }
diff --git a/Engine/Source/Runtime/RHI/Public/RHIDefinitions.h b/Engine/Source/Runtime/RHI/Public/RHIDefinitions.h
--- a/Engine/Source/Runtime/RHI/Public/RHIDefinitions.h
+++ b/Engine/Source/Runtime/RHI/Public/RHIDefinitions.h
@@ -45,6 +45,16 @@ static_assert(sizeof(void*) <= SHADER_PARAMETER_POINTER_ALIGNMENT, "The alignmen
 #define USE_STATIC_SHADER_PLATFORM_INFO 0
 #endif
 
+
+#ifndef RHI_USE_RESOURCE_DEBUG_NAME
+#if (UE_BUILD_SHIPPING || UE_BUILD_TEST)
+	#define RHI_USE_RESOURCE_DEBUG_NAME	0
+#else
+	#define RHI_USE_RESOURCE_DEBUG_NAME	1
+#endif
+#endif
+
+
 #ifndef RHI_RAYTRACING
 #define RHI_RAYTRACING 0
 #endif
diff --git a/Engine/Source/Runtime/RHI/Public/RHIResources.h b/Engine/Source/Runtime/RHI/Public/RHIResources.h
--- a/Engine/Source/Runtime/RHI/Public/RHIResources.h
+++ b/Engine/Source/Runtime/RHI/Public/RHIResources.h
@@ -5384,7 +5384,7 @@ public:
 	RHI_API FRHIShaderResourceView* GetOrCreateSRV(FRHICommandListBase& RHICmdList, FRHITexture* Texture, const FRHITextureSRVCreateInfo& CreateInfo);
 
 	// Sets the debug name of the RHI view resources.
-#if !(UE_BUILD_SHIPPING || UE_BUILD_TEST)
+#if RHI_USE_RESOURCE_DEBUG_NAME
 	RHI_API void SetDebugName(FRHICommandListBase& RHICmdList, const TCHAR* DebugName);
 #else
 	void SetDebugName(FRHICommandListBase& RHICmdList, const TCHAR* DebugName) {}
@@ -5405,7 +5405,7 @@ public:
 	RHI_API FRHIShaderResourceView* GetOrCreateSRV(FRHICommandListBase& RHICmdList, FRHIBuffer* Buffer, const FRHIBufferSRVCreateInfo& CreateInfo);
 
 	// Sets the debug name of the RHI view resources.
-#if !(UE_BUILD_SHIPPING || UE_BUILD_TEST)
+#if RHI_USE_RESOURCE_DEBUG_NAME
 	RHI_API void SetDebugName(FRHICommandListBase& RHICmdList, const TCHAR* DebugName);
 #else
 	void SetDebugName(FRHICommandListBase& RHICmdList, const TCHAR* DebugName) {}
diff --git a/Engine/Source/Runtime/RenderCore/Private/RenderGraphResources.cpp b/Engine/Source/Runtime/RenderCore/Private/RenderGraphResources.cpp
--- a/Engine/Source/Runtime/RenderCore/Private/RenderGraphResources.cpp
+++ b/Engine/Source/Runtime/RenderCore/Private/RenderGraphResources.cpp
@@ -52,7 +52,7 @@ FRDGViewableResource::FRDGViewableResource(const TCHAR* InName, const ERDGViewab
 	, bSplitFirstTransition(bInSplitFirstTransition)
 	, bQueuedForUpload(0)
 	, bCollectForAllocate(1)
-#if !(UE_BUILD_SHIPPING || UE_BUILD_TEST)
+#if RHI_USE_RESOURCE_DEBUG_NAME
 	, bHeapAllocatedDebugName(0)
 #endif
 	, TransientExtractionHint(ETransientExtractionHint::None)
@@ -157,9 +157,7 @@ FRDGPooledBuffer::FRDGPooledBuffer(TRefCountPtr<FRHIBuffer> InBuffer, const FRDG
 
 void FRDGPooledBuffer::SetDebugLabelName(FRHICommandListBase& RHICmdList, const TCHAR* InName)
 {
-#if (UE_BUILD_SHIPPING || UE_BUILD_TEST)
-	Name = InName;
-#else
+#if RHI_USE_RESOURCE_DEBUG_NAME
 	// For performance, avoid updating name if it happens to be the same (true 80% of the time in testing)
 	bool bNameUpdated = false;
 	if (!InName || FCString::Strcmp(Name, InName))
@@ -179,6 +177,8 @@ void FRDGPooledBuffer::SetDebugLabelName(FRHICommandListBase& RHICmdList, const
 		NameUpdatedViewCacheNum = ViewCacheNum;
 		ViewCache.SetDebugName(RHICmdList, InName);
 	}
+#else
+	Name = InName;
 #endif
 }
 
diff --git a/Engine/Source/Runtime/RenderCore/Private/RenderTargetPool.cpp b/Engine/Source/Runtime/RenderCore/Private/RenderTargetPool.cpp
--- a/Engine/Source/Runtime/RenderCore/Private/RenderTargetPool.cpp
+++ b/Engine/Source/Runtime/RenderCore/Private/RenderTargetPool.cpp
@@ -216,7 +216,7 @@ bool FRenderTargetPool::FindFreeElement(FRHICommandListBase& RHICmdList, const F
 		{
 			// we can reuse the same, but the debug name might have changed
 			Current->Desc.DebugName = Name;
-#if !(UE_BUILD_SHIPPING || UE_BUILD_TEST)
+#if RHI_USE_RESOURCE_DEBUG_NAME
 			if (Current->GetRHI())
 			{
 				RHICmdList.BindDebugLabelName(Current->GetRHI(), Name);
@@ -605,15 +605,15 @@ bool FPooledRenderTarget::OnFrameStart()
 
 void FPooledRenderTarget::SetDebugLabelName(FRHICommandListBase& RHICmdList, const TCHAR* Name)
 {
-#if (UE_BUILD_SHIPPING || UE_BUILD_TEST)
-	Desc.DebugName = Name;
-#else
+#if RHI_USE_RESOURCE_DEBUG_NAME
 	// For performance, avoid updating name if it happens to be the same (true 90% of the time in testing)
 	if (FCString::Strcmp(Desc.DebugName, Name))
 	{
 		Desc.DebugName = Name;
 		RHICmdList.BindDebugLabelName(GetRHI(), Name);
 	}
+#else
+	Desc.DebugName = Name;
 #endif
 }
 
diff --git a/Engine/Source/Runtime/RenderCore/Public/RenderGraphResources.h b/Engine/Source/Runtime/RenderCore/Public/RenderGraphResources.h
--- a/Engine/Source/Runtime/RenderCore/Public/RenderGraphResources.h
+++ b/Engine/Source/Runtime/RenderCore/Public/RenderGraphResources.h
@@ -275,7 +275,7 @@ class FRDGViewableResource
 	: public FRDGResource
 {
 public:
-#if !(UE_BUILD_SHIPPING || UE_BUILD_TEST)
+#if RHI_USE_RESOURCE_DEBUG_NAME
 	virtual ~FRDGViewableResource()
 	{
 		if (bHeapAllocatedDebugName)
@@ -320,7 +320,7 @@ public:
 #endif
 	}
 
-#if !(UE_BUILD_SHIPPING || UE_BUILD_TEST)
+#if RHI_USE_RESOURCE_DEBUG_NAME
 	inline void SetDebugNameIsHeapAllocated()
 	{
 		bHeapAllocatedDebugName = 1;
@@ -329,7 +329,7 @@ public:
 
 	inline bool IsDebugNameHeapAllocated() const
 	{
-#if !(UE_BUILD_SHIPPING || UE_BUILD_TEST)
+#if RHI_USE_RESOURCE_DEBUG_NAME
 		return bHeapAllocatedDebugName;
 #else
 		return false;
@@ -408,7 +408,7 @@ protected:
 	/** If false, the resource needs to be collected. */
 	uint8 bCollectForAllocate : 1;
 
-#if !(UE_BUILD_SHIPPING || UE_BUILD_TEST)
+#if RHI_USE_RESOURCE_DEBUG_NAME
 	/** If true, the debug name passed to the constructor was heap allocated, and needs to be freed in the destructor. */
 	uint8 bHeapAllocatedDebugName : 1;
 #endif
@@ -1295,7 +1295,7 @@ private:
 	// May be UINT64_MAX for regular (non-reserved) buffers or when the entire resource is committed.
 	uint64 CommittedSizeInBytes = UINT64_MAX;
 
-#if !(UE_BUILD_SHIPPING || UE_BUILD_TEST)
+#if RHI_USE_RESOURCE_DEBUG_NAME
 	// Tracks the number of items in ViewCache when the debug Name was updated, so we know if items were added and we need to propagate the debug name to the new items
 	int32 NameUpdatedViewCacheNum = 0;
 #endif
diff --git a/Engine/Source/Runtime/Windows/D3D11RHI/Private/D3D11Buffer.cpp b/Engine/Source/Runtime/Windows/D3D11RHI/Private/D3D11Buffer.cpp
--- a/Engine/Source/Runtime/Windows/D3D11RHI/Private/D3D11Buffer.cpp
+++ b/Engine/Source/Runtime/Windows/D3D11RHI/Private/D3D11Buffer.cpp
@@ -306,7 +306,7 @@ void FD3D11DynamicRHI::RHIBindDebugLabelName(FRHICommandListBase& RHICmdList, FR
 	check(BufferRHI);
 	FName DebugName(Name);
 	BufferRHI->SetName(DebugName);
-#if UE_BUILD_DEBUG || UE_BUILD_DEVELOPMENT
+#if RHI_USE_RESOURCE_DEBUG_NAME
 	FD3D11Buffer* BufferD3D = ResourceCast(BufferRHI);
 	if (BufferD3D->Resource != nullptr)
 	{
diff --git a/Engine/Source/Runtime/Windows/D3D11RHI/Private/D3D11Texture.cpp b/Engine/Source/Runtime/Windows/D3D11RHI/Private/D3D11Texture.cpp
--- a/Engine/Source/Runtime/Windows/D3D11RHI/Private/D3D11Texture.cpp
+++ b/Engine/Source/Runtime/Windows/D3D11RHI/Private/D3D11Texture.cpp
@@ -1298,7 +1298,7 @@ void FD3D11DynamicRHI::RHIBindDebugLabelName(FRHICommandListBase& RHICmdList, FR
 	//todo: require names at texture creation time.
 	FName DebugName(Name);
 	Texture->SetName(DebugName);
-#if UE_BUILD_DEBUG || UE_BUILD_DEVELOPMENT
+#if RHI_USE_RESOURCE_DEBUG_NAME
 	ID3D11Resource* ResourceD3D = Texture->GetResource();
 	if (ResourceD3D != nullptr)
 	{
diff --git a/Engine/Source/Runtime/Windows/D3D11RHI/Private/D3D11UAV.cpp b/Engine/Source/Runtime/Windows/D3D11RHI/Private/D3D11UAV.cpp
--- a/Engine/Source/Runtime/Windows/D3D11RHI/Private/D3D11UAV.cpp
+++ b/Engine/Source/Runtime/Windows/D3D11RHI/Private/D3D11UAV.cpp
@@ -133,7 +133,7 @@ FUnorderedAccessViewRHIRef FD3D11DynamicRHI::RHICreateUnorderedAccessView(class
 
 void FD3D11DynamicRHI::RHIBindDebugLabelName(FRHICommandListBase& RHICmdList, FRHIUnorderedAccessView* UnorderedAccessViewRHI, const TCHAR* Name)
 {
-#if UE_BUILD_DEBUG || UE_BUILD_DEVELOPMENT
+#if RHI_USE_RESOURCE_DEBUG_NAME
 	FD3D11UnorderedAccessView* UAV = ResourceCast(UnorderedAccessViewRHI);
 	if (UAV->View != nullptr)
 	{
