Also bump FStablePipelineCacheSerializedHeader version number after adding serializiation of RTPSO ShaderBindingLayout

diff --git a/Engine/Source/Runtime/RenderCore/Private/PipelineCacheUtilities.cpp b/Engine/Source/Runtime/RenderCore/Private/PipelineCacheUtilities.cpp
--- a/Engine/Source/Runtime/RenderCore/Private/PipelineCacheUtilities.cpp
+++ b/Engine/Source/Runtime/RenderCore/Private/PipelineCacheUtilities.cpp
@@ -70,8 +70,9 @@ namespace Private
 			AddingPipelineCacheVersion = 5,
 			AddingDepthBounds = 6,
 			AddedWorkGraphFrequencies = 7,
+			AddRTPSOShaderBindingLayout = 8,
 
-			Current = AddedWorkGraphFrequencies
+			Current = AddRTPSOShaderBindingLayout
 		};
 
 		/** Magic to reject other files */
