Fix an assert in RealtimeGPUProfiler when both stats and CsvProfiler stats are disabled

#tests compiled editor

diff --git a/Engine/Source/Runtime/RenderCore/Private/ProfilingDebugging/RealtimeGPUProfiler.cpp b/Engine/Source/Runtime/RenderCore/Private/ProfilingDebugging/RealtimeGPUProfiler.cpp
--- a/Engine/Source/Runtime/RenderCore/Private/ProfilingDebugging/RealtimeGPUProfiler.cpp
+++ b/Engine/Source/Runtime/RenderCore/Private/ProfilingDebugging/RealtimeGPUProfiler.cpp
@@ -481,9 +481,10 @@ public:
 			// Multi-GPU support : Tracing profiler is MGPU-aware, but not CSV profiler or Unreal stats.
 			if (Event.GetGPUMask().Contains(0))
 			{
+#if STATS || CSV_PROFILER_STATS
 				// Check if we've seen this stat yet 
 				const bool bKnownStat = StatSeenSet.Add(Event.GetName());
-
+#endif
 				const int64 EventTimeUs = GPUStatsChildTimesIncluded ? IncExcTime.InclusiveTimeUs : IncExcTime.ExclusiveTimeUs;
 				TotalUs += IncExcTime.ExclusiveTimeUs;
 
