Fixed transient transition before acquire for subresources that are not accessed by RDG within the frame.

diff --git a/Engine/Source/Runtime/RenderCore/Private/RenderGraphBuilder.cpp b/Engine/Source/Runtime/RenderCore/Private/RenderGraphBuilder.cpp
--- a/Engine/Source/Runtime/RenderCore/Private/RenderGraphBuilder.cpp
+++ b/Engine/Source/Runtime/RenderCore/Private/RenderGraphBuilder.cpp
@@ -3942,7 +3942,7 @@ void FRDGBuilder::AddLastTextureTransition(FRDGTexture* Texture)
 			if (!SubresourceStateBefore)
 			{
 				SubresourceStateBefore = AllocSubresource();
-				SubresourceStateBefore->SetPass(ERHIPipeline::Graphics, GetProloguePassHandle());
+				SubresourceStateBefore->SetPass(ERHIPipeline::Graphics, Texture->DiscardPass.IsValid() ? Texture->FirstPass : GetProloguePassHandle());
 			}
 
 			State = SubresourceStateBefore;
