Clear metadata scope for RDG pooled buffers

Pooled buffer allocations should not be associated with specific asset as they are shared and reused.

diff --git a/Engine/Source/Runtime/RenderCore/Private/RenderGraphResourcePool.cpp b/Engine/Source/Runtime/RenderCore/Private/RenderGraphResourcePool.cpp
--- a/Engine/Source/Runtime/RenderCore/Private/RenderGraphResourcePool.cpp
+++ b/Engine/Source/Runtime/RenderCore/Private/RenderGraphResourcePool.cpp
@@ -179,6 +179,10 @@ FRDGPooledBuffer* FRDGBufferPool::CreateBuffer(FRHICommandListBase& RHICmdList,
 	TRACE_COUNTER_ADD(BufferPoolCount, 1);
 	TRACE_COUNTER_ADD(BufferPoolCreateCount, 1);
 	TRACE_COUNTER_ADD(BufferPoolSize, NumBytes);
+	
+	UE_TRACE_METADATA_CLEAR_SCOPE(); // Do not associate a pooled buffer with specific asset
+	LLM_TAGSET_SCOPE_CLEAR(ELLMTagSet::Assets);
+	LLM_TAGSET_SCOPE_CLEAR(ELLMTagSet::AssetClasses);
 
 	const ERHIAccess InitialAccess = RHIGetDefaultResourceState(Desc.Usage, false);
 	FRHIResourceCreateInfo CreateInfo(InDebugName);
