Set LLM category for RDG pooled buffers

#rnx

diff --git a/Engine/Source/Runtime/RenderCore/Private/RenderGraphResourcePool.cpp b/Engine/Source/Runtime/RenderCore/Private/RenderGraphResourcePool.cpp
--- a/Engine/Source/Runtime/RenderCore/Private/RenderGraphResourcePool.cpp
+++ b/Engine/Source/Runtime/RenderCore/Private/RenderGraphResourcePool.cpp
@@ -180,6 +180,7 @@ FRDGPooledBuffer* FRDGBufferPool::CreateBuffer(FRHICommandListBase& RHICmdList,
 	TRACE_COUNTER_ADD(BufferPoolCreateCount, 1);
 	TRACE_COUNTER_ADD(BufferPoolSize, NumBytes);
 	
+	LLM_SCOPE_BYNAME(TEXT("RHIMisc/BufferPool"));
 	UE_TRACE_METADATA_CLEAR_SCOPE(); // Do not associate a pooled buffer with specific asset
 	LLM_TAGSET_SCOPE_CLEAR(ELLMTagSet::Assets);
 	LLM_TAGSET_SCOPE_CLEAR(ELLMTagSet::AssetClasses);
