Modified RDG blackboard to support getting a mutable instance of the struct.

diff --git a/Engine/Source/Runtime/RenderCore/Public/RenderGraphBlackboard.h b/Engine/Source/Runtime/RenderCore/Public/RenderGraphBlackboard.h
--- a/Engine/Source/Runtime/RenderCore/Public/RenderGraphBlackboard.h
+++ b/Engine/Source/Runtime/RenderCore/Public/RenderGraphBlackboard.h
@@ -21,8 +21,8 @@
 
 /** The blackboard is a map of struct instances with a lifetime tied to a render graph allocator. It is designed
  *  to solve cases where explicit marshaling of immutable data is undesirable. Structures are created once and the mutable
- *  reference is returned. Only the immutable version is accessible from the blackboard. This constraint on mutability
- *  is to discourage relying entirely on the blackboard. The mutable version should be marshaled around if needed.
+ *  reference is returned.
+ *
  *  Good candidates for the blackboard would be data that is created once and immutably fetched across the entire
  *  renderer pipeline, where marshaling would create more maintenance burden than benefit. More constrained data
  *  structures should be marshaled through function calls instead.
@@ -75,16 +75,16 @@ public:
 		return static_cast<HelperStructType*>(Result)->Struct;
 	}
 
-	/** Gets an immutable instance of the struct. Returns null if not present in the blackboard. */
+	/** Gets a mutable instance of the struct. Returns null if not present in the blackboard. */
 	template <typename StructType>
-	const StructType* Get() const
+	StructType* GetMutable() const
 	{
 		using HelperStructType = TStruct<StructType>;
 
 		const int32 StructIndex = GetStructIndex<StructType>();
 		if (StructIndex < Blackboard.Num())
 		{
-			if (const HelperStructType* Element = static_cast<const HelperStructType*>(Blackboard[StructIndex]))
+			if (HelperStructType* Element = static_cast<HelperStructType*>(Blackboard[StructIndex]))
 			{
 				return &Element->Struct;
 			}
@@ -92,14 +92,38 @@ public:
 		return nullptr;
 	}
 
-	/** Gets an immutable instance of the struct. Asserts if not present in the blackboard. */
+	/** Gets an immutable instance of the struct. Returns null if not present in the blackboard. */
 	template <typename StructType>
-	const StructType& GetChecked() const
+	const StructType* Get() const
 	{
-		const StructType* Struct = Get<StructType>();
+		return GetMutable<StructType>();
+	}
+	
+	template <typename StructType, typename... ArgsType>
+	StructType& GetOrCreate(ArgsType&&... Args)
+	{
+		if (StructType* Struct = GetMutable<StructType>())
+		{
+			return *Struct;
+		}
+		return Create<StructType>(Forward<ArgsType&&>(Args)...);
+	}
+
+	/** Gets a mutable instance of the struct. Asserts if not present in the blackboard. */
+	template <typename StructType>
+	StructType& GetMutableChecked() const
+	{
+		StructType* Struct = Get<StructType>();
 		checkf(Struct, TEXT("RDGBlackboard Get failed to find instance of struct '%s' in the blackboard."), GetGeneratedTypeName<StructType>());
 		return *Struct;
 	}
+	
+	/** Gets an immutable instance of the struct. Asserts if not present in the blackboard. */
+	template <typename StructType>
+	const StructType& GetChecked() const
+	{
+		return GetMutableChecked();
+	}
 
 private:
 	FRDGBlackboard(FRDGAllocator& InAllocator)
