Fixed build break.

diff --git a/Engine/Source/Runtime/RenderCore/Public/RenderGraphBlackboard.h b/Engine/Source/Runtime/RenderCore/Public/RenderGraphBlackboard.h
--- a/Engine/Source/Runtime/RenderCore/Public/RenderGraphBlackboard.h
+++ b/Engine/Source/Runtime/RenderCore/Public/RenderGraphBlackboard.h
@@ -122,7 +122,7 @@ public:
 	template <typename StructType>
 	const StructType& GetChecked() const
 	{
-		return GetMutableChecked();
+		return GetMutableChecked<StructType>();
 	}
 
 private:
