Fixed build break.

diff --git a/Engine/Source/Runtime/RenderCore/Public/RenderGraphBlackboard.h b/Engine/Source/Runtime/RenderCore/Public/RenderGraphBlackboard.h
--- a/Engine/Source/Runtime/RenderCore/Public/RenderGraphBlackboard.h
+++ b/Engine/Source/Runtime/RenderCore/Public/RenderGraphBlackboard.h
@@ -113,7 +113,7 @@ public:
 	template <typename StructType>
 	StructType& GetMutableChecked() const
 	{
-		StructType* Struct = Get<StructType>();
+		StructType* Struct = GetMutable<StructType>();
 		checkf(Struct, TEXT("RDGBlackboard Get failed to find instance of struct '%s' in the blackboard."), GetGeneratedTypeName<StructType>());
 		return *Struct;
 	}
