Prevent Heterogeneous Volume Component from issuing SVT Streaming requests when Heterogeneous Volumes are not enabled for the active runtime
Added public function DoesRuntimeSupportHeterogeneousVolumes()

diff --git a/Engine/Source/Runtime/Engine/Private/Components/HeterogeneousVolumeComponent.cpp b/Engine/Source/Runtime/Engine/Private/Components/HeterogeneousVolumeComponent.cpp
--- a/Engine/Source/Runtime/Engine/Private/Components/HeterogeneousVolumeComponent.cpp
+++ b/Engine/Source/Runtime/Engine/Private/Components/HeterogeneousVolumeComponent.cpp
@@ -19,6 +19,7 @@
 #include "SparseVolumeTexture/SparseVolumeTexture.h"
 #include "StaticMeshResources.h"
 #include "UObject/ConstructorHelpers.h"
+#include "Engine/World.h"
 
 #include UE_INLINE_GENERATED_CPP_BY_NAME(HeterogeneousVolumeComponent)
 
@@ -644,7 +645,7 @@ void UHeterogeneousVolumeComponent::TickComponent(float DeltaTime, ELevelTick Ti
 		}
 	}
 
-	if (ShouldRender() && MaterialInstanceDynamic)
+	if (ShouldRender() && DoesRuntimeSupportHeterogeneousVolumes(GetWorld()->Scene->GetShaderPlatform()) && MaterialInstanceDynamic)
 	{
 		const int32 SVTParameterIndex = 0;
 		FName SVTParameterName;
diff --git a/Engine/Source/Runtime/RenderCore/Private/RenderUtils.cpp b/Engine/Source/Runtime/RenderCore/Private/RenderUtils.cpp
--- a/Engine/Source/Runtime/RenderCore/Private/RenderUtils.cpp
+++ b/Engine/Source/Runtime/RenderCore/Private/RenderUtils.cpp
@@ -1409,6 +1409,13 @@ bool IsWaterVirtualShadowMapFilteringEnabled(const FStaticShaderPlatform Platfor
 	return !IsForwardShadingEnabled(Platform) && bWaterVSMFilteringSupported && bVirtualShadowMapsSupported;
 }
 
+bool DoesRuntimeSupportHeterogeneousVolumes(EShaderPlatform Platform)
+{
+	static const auto CVarHeterogeneousVolumes = IConsoleManager::Get().FindConsoleVariable(TEXT("r.HeterogeneousVolumes"));
+	bool bHeterogeneousVolumesEnabled = CVarHeterogeneousVolumes && CVarHeterogeneousVolumes->GetInt() > 0;
+	return DoesPlatformSupportHeterogeneousVolumes(Platform) && bHeterogeneousVolumesEnabled;
+}
+
 bool DoesPlatformSupportHeterogeneousVolumes(EShaderPlatform Platform)
 {
 	return IsFeatureLevelSupported(Platform, ERHIFeatureLevel::SM5)
diff --git a/Engine/Source/Runtime/RenderCore/Public/RenderUtils.h b/Engine/Source/Runtime/RenderCore/Public/RenderUtils.h
--- a/Engine/Source/Runtime/RenderCore/Public/RenderUtils.h
+++ b/Engine/Source/Runtime/RenderCore/Public/RenderUtils.h
@@ -453,6 +453,11 @@ RENDERCORE_API bool UseNonNaniteVirtualShadowMaps(EShaderPlatform ShaderPlatform
 /** Returns if water should evaluate virtual shadow maps a second time for the water surface. This is for a platform so can be used at cook time. */
 RENDERCORE_API bool IsWaterVirtualShadowMapFilteringEnabled(const FStaticShaderPlatform Platform);
 
+/*
+ * Detect (at runtime) if the runtime supports heterogeneous volumes.
+ */
+RENDERCORE_API bool DoesRuntimeSupportHeterogeneousVolumes(EShaderPlatform Platform);
+
 /**
  * Returns true if Heterogeneous Volumes should be used for the given shader platform.
  */
