Fixed incorrect transition for transient resources that are also extracted. Use AcquirePass instead of DiscardPass to judge whether to hoist the transition to the first pass.

diff --git a/Engine/Source/Runtime/RenderCore/Private/RenderGraphBuilder.cpp b/Engine/Source/Runtime/RenderCore/Private/RenderGraphBuilder.cpp
--- a/Engine/Source/Runtime/RenderCore/Private/RenderGraphBuilder.cpp
+++ b/Engine/Source/Runtime/RenderCore/Private/RenderGraphBuilder.cpp
@@ -3942,7 +3942,7 @@ void FRDGBuilder::AddLastTextureTransition(FRDGTexture* Texture)
 			if (!SubresourceStateBefore)
 			{
 				SubresourceStateBefore = AllocSubresource();
-				SubresourceStateBefore->SetPass(ERHIPipeline::Graphics, Texture->DiscardPass.IsValid() ? Texture->FirstPass : GetProloguePassHandle());
+				SubresourceStateBefore->SetPass(ERHIPipeline::Graphics, Texture->AcquirePass.IsValid() ? Texture->FirstPass : GetProloguePassHandle());
 			}
 
 			State = SubresourceStateBefore;
