Allocate the global read buffer as Static instead of Dynamic in D3D11, to work around driver inefficiencies.

Fixes frequent long hitches that were happening during map/memcpy.

#tests tested locally
#rnx

diff --git a/Engine/Source/Runtime/RHI/Public/RHIUtilities.h b/Engine/Source/Runtime/RHI/Public/RHIUtilities.h
--- a/Engine/Source/Runtime/RHI/Public/RHIUtilities.h
+++ b/Engine/Source/Runtime/RHI/Public/RHIUtilities.h
@@ -410,10 +410,11 @@ struct FDynamicReadBuffer : public FReadBuffer
 
 	void Initialize(FRHICommandListBase& RHICmdList, const TCHAR* DebugName, uint32 BytesPerElement, uint32 NumElements, EPixelFormat Format, EBufferUsageFlags AdditionalUsage = BUF_None)
 	{
-		ensure(
-			EnumHasAnyFlags(AdditionalUsage, BUF_Dynamic | BUF_Volatile | BUF_Static) &&					// buffer should be Dynamic or Volatile or Static
-			EnumHasAnyFlags(AdditionalUsage, BUF_Dynamic) != EnumHasAnyFlags(AdditionalUsage, BUF_Volatile) // buffer should not be both
-			);
+		// Commented out due to the hack made in FGlobalDynamicReadBuffer::AllocateInternal to use static in D3D11.
+		//ensure(
+		//	EnumHasAnyFlags(AdditionalUsage, BUF_Dynamic | BUF_Volatile | BUF_Static) &&					// buffer should be Dynamic or Volatile or Static
+		//	EnumHasAnyFlags(AdditionalUsage, BUF_Dynamic) != EnumHasAnyFlags(AdditionalUsage, BUF_Volatile) // buffer should not be both
+		//	);
 
 		FReadBuffer::Initialize(RHICmdList, DebugName, BytesPerElement, NumElements, Format, AdditionalUsage);
 	}
diff --git a/Engine/Source/Runtime/RenderCore/Private/DynamicBufferAllocator.cpp b/Engine/Source/Runtime/RenderCore/Private/DynamicBufferAllocator.cpp
--- a/Engine/Source/Runtime/RenderCore/Private/DynamicBufferAllocator.cpp
+++ b/Engine/Source/Runtime/RenderCore/Private/DynamicBufferAllocator.cpp
@@ -32,6 +32,12 @@ FAutoConsoleVariableRef CVarAlignReadBufferSize(
 	GAlignReadBufferRenderingBufferSize,
 	TEXT("The alignment size (in instances) to allocate in blocks for rendering read buffers. i.e. 64*1024 = 256k for a float buffer"));
 
+int32 GReadBufferD3D11StaticAllocation = true;
+FAutoConsoleVariableRef CVarReadBufferD3D11StaticAllocation(
+	TEXT("r.ReadBuffer.D3D11StaticAllocation"),
+	GReadBufferD3D11StaticAllocation,
+	TEXT("Whether to allocate the read buffer as Static instead of Dynamic in D3D11, to work around driver inefficiencies."));
+
 struct FDynamicReadBufferPool
 {
 	/** List of vertex buffers. */
@@ -147,7 +153,14 @@ FGlobalDynamicReadBuffer::FAllocation FGlobalDynamicReadBuffer::AllocateInternal
 			const uint32 NewBufferSize = FMath::Max(AlignedNum, (uint32)GMinReadBufferRenderingBufferSize);
 			Buffer = new FDynamicAllocReadBuffer();
 			BufferPool->Buffers.Add(Buffer);
-			Buffer->Initialize(*RHICmdList, TEXT("FGlobalDynamicReadBuffer_AllocateInternal"), sizeof(Type), NewBufferSize, Format, BUF_Dynamic);
+			// Use static in D3D11 because we're running out of zeroed pages, and that causes the upload of this buffer to take a very long time. Static uses more memory and might be slightly less efficient
+			// on the GPU, but it beats stalling for tens of milliseconds when committing the buffer.
+			EBufferUsageFlags Flags = BUF_Dynamic;
+			if (GReadBufferD3D11StaticAllocation && GDynamicRHI->GetInterfaceType() == ERHIInterfaceType::D3D11)
+			{
+				Flags = BUF_Static;
+			}
+			Buffer->Initialize(*RHICmdList, TEXT("FGlobalDynamicReadBuffer_AllocateInternal"), sizeof(Type), NewBufferSize, Format, Flags);
 		}
 
 		// Lock the buffer if needed.
