Fixed reserved resource commit size being assigned incorrectly when parallel setup is disabled. Reserved commit size needs to be properly pipelined between calls to SetupPassResources, which take place on the async setup queue when parallel setup is enabled, but immediately on the render thread when disabled. Also fixed the cull root buffer / texture commands not being consumed when parallel setup is disabled.

diff --git a/Engine/Source/Runtime/RenderCore/Private/RenderGraphBuilder.cpp b/Engine/Source/Runtime/RenderCore/Private/RenderGraphBuilder.cpp
--- a/Engine/Source/Runtime/RenderCore/Private/RenderGraphBuilder.cpp
+++ b/Engine/Source/Runtime/RenderCore/Private/RenderGraphBuilder.cpp
@@ -1767,12 +1767,8 @@ void FRDGBuilder::Execute()
 	{
 		BeginFlushResourcesRHI();
 		WaitForParallelSetupTasks(ERDGSetupTaskWaitPoint::Compile);
-
-		if (ParallelSetup.bEnabled)
-		{
-			AsyncSetupQueue.Pipe.WaitUntilEmpty();
-			ProcessAsyncSetupQueue();
-		}
+		AsyncSetupQueue.Pipe.WaitUntilEmpty();
+		ProcessAsyncSetupQueue();
 
 		const int32 ParallelCompileResourceThreshold = 32;
 		const int32 NumBuffers   = Buffers.Num();
diff --git a/Engine/Source/Runtime/RenderCore/Public/RenderGraphBuilder.inl b/Engine/Source/Runtime/RenderCore/Public/RenderGraphBuilder.inl
--- a/Engine/Source/Runtime/RenderCore/Public/RenderGraphBuilder.inl
+++ b/Engine/Source/Runtime/RenderCore/Public/RenderGraphBuilder.inl
@@ -418,7 +418,7 @@ inline void FRDGBuilder::QueueBufferUpload(FRDGBufferRef Buffer, FRDGBufferIniti
 inline void FRDGBuilder::QueueCommitReservedBuffer(FRDGBufferRef Buffer, uint64 CommitSizeInBytes)
 {
 	IF_RDG_ENABLE_DEBUG(UserValidation.ValidateCommitBuffer(Buffer, CommitSizeInBytes));
-	if (IsImmediateMode() || GRHICommandList.Bypass())
+	if (!ParallelSetup.bEnabled)
 	{
 		Buffer->PendingCommitSize = CommitSizeInBytes;
 	}
