UE: Fix crash at shutdown

diff --git a/Engine/Source/Runtime/RenderCore/Private/GPUMessaging.cpp b/Engine/Source/Runtime/RenderCore/Private/GPUMessaging.cpp
--- a/Engine/Source/Runtime/RenderCore/Private/GPUMessaging.cpp
+++ b/Engine/Source/Runtime/RenderCore/Private/GPUMessaging.cpp
@@ -267,7 +267,7 @@ static TGlobalResource<FSystem> GSystem;
 
 void FSocket::Reset()
 {
-	if (MessageId.IsValid())
+	if (MessageId.IsValid() && GSystem.IsInitialized())
 	{
 		GSystem.RemoveHandler(MessageId);
 	}
