Make sure all entries of the ShaderBindingLayout are added to the shader source because they are required to be all defined by certain shader compilers (can't have missing cbuffer indices)

diff --git a/Engine/Source/Runtime/RenderCore/Private/ShaderCore.cpp b/Engine/Source/Runtime/RenderCore/Private/ShaderCore.cpp
--- a/Engine/Source/Runtime/RenderCore/Private/ShaderCore.cpp
+++ b/Engine/Source/Runtime/RenderCore/Private/ShaderCore.cpp
@@ -842,20 +842,6 @@ uint32 FShaderParameterMap::CountParametersOfType(EShaderParameterType InType) c
 
 #if WITH_EDITOR
 
-FThreadSafeSharedAnsiStringPtr FShaderBindingLayout::GetUniformBufferDeclarationAnsiPtr(const FShaderParametersMetadata* ShaderParametersMetadata) const
-{
-	// Does the binding layout contain a uniform buffer declaration for this shader paramaters
-	FString UniformBufferName(ShaderParametersMetadata->GetShaderVariableName());
-	const FThreadSafeSharedAnsiStringPtr* UniformBufferDeclarationAnsi = UniformBufferMap.Find(UniformBufferName);
-	if (UniformBufferDeclarationAnsi)
-	{
-		return *UniformBufferDeclarationAnsi;
-	}
-
-	// use the shared declaration from the metadata object itself (no fixed register or space defined)
-	return ShaderParametersMetadata->GetUniformBufferDeclarationAnsiPtr();
-}
-
 void FShaderBindingLayout::SetUniformBufferDeclarationAnsiPtr(const FShaderParametersMetadata* ShaderParametersMetadata, FThreadSafeSharedAnsiStringPtr UniformBufferDeclarationAnsi)
 {
 	FString UniformBufferName(ShaderParametersMetadata->GetShaderVariableName());
diff --git a/Engine/Source/Runtime/RenderCore/Private/ShaderParameters.cpp b/Engine/Source/Runtime/RenderCore/Private/ShaderParameters.cpp
--- a/Engine/Source/Runtime/RenderCore/Private/ShaderParameters.cpp
+++ b/Engine/Source/Runtime/RenderCore/Private/ShaderParameters.cpp
@@ -443,33 +443,45 @@ void UE::ShaderParameters::AddUniformBufferIncludesToEnvironment(FShaderCompiler
 
 	FString UniformBufferIncludes;
 
-	for (const FShaderParametersMetadata* Metadata : InUniformBuffers)	
+	auto AddUniformBufferData = [&OutEnvironment, &UniformBufferIncludes](const FShaderParametersMetadata* Metadata, FThreadSafeSharedAnsiStringPtr UniformBufferDeclaration)
 	{
 		FStringView UniformBufferNameView(Metadata->GetShaderVariableName());
 		uint32 UniformBufferNameHash = GetTypeHash(UniformBufferNameView);
 		if (!OutEnvironment.UniformBufferMap.FindByHash(UniformBufferNameHash, UniformBufferNameView))
 		{
-			FThreadSafeSharedAnsiStringPtr UniformBufferDeclaration;
-
-			// Use the shader binding layout to retrieve the uniform buffer declaration if available
-			if (OutEnvironment.ShaderBindingLayout)
-			{
-				UniformBufferDeclaration = OutEnvironment.ShaderBindingLayout->GetUniformBufferDeclarationAnsiPtr(Metadata);
-			}
-			else
-			{
-				UniformBufferDeclaration = Metadata->GetUniformBufferDeclarationAnsiPtr();
-			}
-			
 			check(UniformBufferDeclaration.Get() != NULL);
 			check(!UniformBufferDeclaration.Get()->IsEmpty());
-			
+
 			UniformBufferIncludes += Metadata->GetUniformBufferInclude();
-			
+
 			OutEnvironment.IncludeVirtualPathToSharedContentsMap.AddByHash(Metadata->GetUniformBufferPathHash(), Metadata->GetUniformBufferPath(), UniformBufferDeclaration);
-			
+
 			Metadata->AddResourceTableEntries(OutEnvironment.ResourceTableMap, OutEnvironment.UniformBufferMap);
 		}
+	};
+
+	// If there is a shader binding layout then always add all the entries (duplicated from InUniformBuffers will not be added twice because of Find function below)
+	if (OutEnvironment.ShaderBindingLayout)
+	{
+		for (auto Iter = OutEnvironment.ShaderBindingLayout->GetUniformBufferDeclarations().CreateConstIterator(); Iter; ++Iter)
+		{
+			const FShaderParametersMetadata* UniformBufferStruct = FindUniformBufferStructByName(*Iter.Key());
+			if (UniformBufferStruct)
+			{
+				AddUniformBufferData(UniformBufferStruct, Iter.Value());
+			}
+		}
+	}
+
+	for (const FShaderParametersMetadata* Metadata : InUniformBuffers)	
+	{
+		FStringView UniformBufferNameView(Metadata->GetShaderVariableName());
+		uint32 UniformBufferNameHash = GetTypeHash(UniformBufferNameView);
+		if (!OutEnvironment.UniformBufferMap.FindByHash(UniformBufferNameHash, UniformBufferNameView))
+		{
+			FThreadSafeSharedAnsiStringPtr UniformBufferDeclaration = Metadata->GetUniformBufferDeclarationAnsiPtr();
+			AddUniformBufferData(Metadata, UniformBufferDeclaration);
+		}
 	}
 
 	FString& GeneratedUniformBuffersInclude = OutEnvironment.IncludeVirtualPathToContentsMap.FindOrAdd("/Engine/Generated/GeneratedUniformBuffers.ush");
diff --git a/Engine/Source/Runtime/RenderCore/Public/ShaderCore.h b/Engine/Source/Runtime/RenderCore/Public/ShaderCore.h
--- a/Engine/Source/Runtime/RenderCore/Public/ShaderCore.h
+++ b/Engine/Source/Runtime/RenderCore/Public/ShaderCore.h
@@ -456,8 +456,8 @@ public:
 	FRHIShaderBindingLayout RHILayout;
 
 #if WITH_EDITOR
-	FThreadSafeSharedAnsiStringPtr GetUniformBufferDeclarationAnsiPtr(const FShaderParametersMetadata* ShaderParametersMetadata) const;
 	void SetUniformBufferDeclarationAnsiPtr(const FShaderParametersMetadata* ShaderParametersMetadata, FThreadSafeSharedAnsiStringPtr UniformBufferDeclarationAnsi);
+	const TMap<FString, FThreadSafeSharedAnsiStringPtr>& GetUniformBufferDeclarations() const { return UniformBufferMap; }
 
 	RENDERCORE_API void AddRequiredSymbols(TArray<FString>& RequiredSymbols) const;
 
