Prevent recursive calls to FFrameEndSync::Sync. Just do a full RT flush instead.

diff --git a/Engine/Source/Runtime/RenderCore/Private/RenderingThread.cpp b/Engine/Source/Runtime/RenderCore/Private/RenderingThread.cpp
--- a/Engine/Source/Runtime/RenderCore/Private/RenderingThread.cpp
+++ b/Engine/Source/Runtime/RenderCore/Private/RenderingThread.cpp
@@ -2031,6 +2031,18 @@ namespace FFrameEndSync
 
 	void Sync(EFlushMode FlushMode)
 	{
+		static bool bRecursive = false;
+		TGuardValue<bool> RecursionGuard(bRecursive, true);
+
+		if (RecursionGuard.GetOriginalValue())
+		{
+			// This is a recursive call to FFrameEndSync::Sync(). Use a standard render fence and do a full sync.
+			FRenderCommandFence Fence;
+			Fence.BeginFence();
+			Fence.Wait();
+			return;
+		}
+
 		bool bFullSync = FlushMode == EFlushMode::Threads;
 
 		// The "r.OneFrameThreadLag" cvar forces a full sync, meaning the game thread will
