Added dump GPU scope to control dumping a single frame.

diff --git a/Engine/Source/Runtime/RenderCore/Private/DumpGPU.cpp b/Engine/Source/Runtime/RenderCore/Private/DumpGPU.cpp
--- a/Engine/Source/Runtime/RenderCore/Private/DumpGPU.cpp
+++ b/Engine/Source/Runtime/RenderCore/Private/DumpGPU.cpp
@@ -2653,6 +2653,19 @@ bool FRDGBuilder::IsDumpingFrame()
 namespace UE::RenderCore::DumpGPU
 {
 
+RENDERCORE_API FDumpScope::FDumpScope(bool bInCondition)
+	: bCondition(bInCondition && !IsDumpingFrame())
+{
+	if (bCondition)
+	{
+		// Pass "-oneframe" to override CVar that could enable multiple frames of capture
+		FRDGBuilder::BeginResourceDump(TEXT("-oneframe"));
+
+		// Tick the DumpGPU system, which will start the dump
+		TickEndFrame();
+	}
+}
+
 void TickEndFrame()
 {
 	check(IsInGameThread());
diff --git a/Engine/Source/Runtime/RenderCore/Public/DumpGPU.h b/Engine/Source/Runtime/RenderCore/Public/DumpGPU.h
--- a/Engine/Source/Runtime/RenderCore/Public/DumpGPU.h
+++ b/Engine/Source/Runtime/RenderCore/Public/DumpGPU.h
@@ -42,6 +42,32 @@ RENDERCORE_API bool IsDumpingFrame();
 
 RENDERCORE_API bool ShouldCameraCut();
 
+class FDumpScope
+{
+public:
+	RENDERCORE_API FDumpScope(bool bInCondition = true);
+
+	~FDumpScope()
+	{
+		if (bCondition)
+		{
+			// Tick the dump GPU system again, which will end the active dump.
+			TickEndFrame();
+		}
+	}
+
+private:
+	bool bCondition;
+};
+
 } // namespace  UE::RenderCore::DumpGPU
 
-#endif // WITH_DUMPGPU
+#else
+
+class FDumpScope
+{
+public:
+	FDumpScope(bool bCondition = true) {}
+};
+
+#endif // WITH_DUMPGPU
\ No newline at end of file
