Fixed use-after-free with the RDG deallocation task freeing RDG scopes that reference realtime GPU profiler events.

diff --git a/Engine/Source/Runtime/RenderCore/Private/ProfilingDebugging/RealtimeGPUProfiler.cpp b/Engine/Source/Runtime/RenderCore/Private/ProfilingDebugging/RealtimeGPUProfiler.cpp
--- a/Engine/Source/Runtime/RenderCore/Private/ProfilingDebugging/RealtimeGPUProfiler.cpp
+++ b/Engine/Source/Runtime/RenderCore/Private/ProfilingDebugging/RealtimeGPUProfiler.cpp
@@ -942,13 +942,14 @@ void FRealtimeGPUProfiler::EndFrame(FRHICommandListImmediate& RHICmdList)
 #endif
 		))
 		{
-			// Launch an inline async task that will free the frame once RDG async task execution is complete to keep it alive.
+			// Launch an inline async task that will free the frame once RDG async deletion is complete
+			// since RDG scopes reference the profiler events by raw pointer.
 			CleanUpTask = UE::Tasks::Launch(UE_SOURCE_LOCATION, [Frame = MoveTemp(*FramePtr)] () mutable
 			{
 				FOptionalTaskTagScope TaskTag(ETaskTag::EParallelRenderingThread);
 				Frame = {};
 
-			}, MakeArrayView({ CleanUpTask, FRDGBuilder::GetAsyncExecuteTask() }), UE::Tasks::ETaskPriority::High, UE::Tasks::EExtendedTaskPriority::Inline);
+			}, MakeArrayView({ CleanUpTask, FRDGBuilder::GetAsyncDeleteTask() }), UE::Tasks::ETaskPriority::High, UE::Tasks::EExtendedTaskPriority::Inline);
 
 			PendingFrames.Pop();
 		}
diff --git a/Engine/Source/Runtime/RenderCore/Private/RenderGraphBuilder.cpp b/Engine/Source/Runtime/RenderCore/Private/RenderGraphBuilder.cpp
--- a/Engine/Source/Runtime/RenderCore/Private/RenderGraphBuilder.cpp
+++ b/Engine/Source/Runtime/RenderCore/Private/RenderGraphBuilder.cpp
@@ -648,6 +648,11 @@ void FRDGBuilder::WaitForAsyncDeleteTask()
 	FAsyncDeleter::LastTask.Wait();
 }
 
+const UE::Tasks::FTask& FRDGBuilder::GetAsyncDeleteTask()
+{
+	return FAsyncDeleter::LastTask;
+}
+
 UE::Tasks::FTask FRDGBuilder::FParallelExecute::LastAsyncExecuteTask;
 
 void FRDGBuilder::WaitForAsyncExecuteTask()
diff --git a/Engine/Source/Runtime/RenderCore/Public/RenderGraphBuilder.h b/Engine/Source/Runtime/RenderCore/Public/RenderGraphBuilder.h
--- a/Engine/Source/Runtime/RenderCore/Public/RenderGraphBuilder.h
+++ b/Engine/Source/Runtime/RenderCore/Public/RenderGraphBuilder.h
@@ -432,6 +432,9 @@ public:
 
 	/** Waits for the last RDG async delete task that was launched. If the event is valid it is consumed. Thus, it is not thread safe to wait from multiple threads at once. */
 	static RENDERCORE_API void WaitForAsyncDeleteTask();
+	
+	/** Returns the last RDG chained deletion task that was launched. */
+	static RENDERCORE_API const UE::Tasks::FTask& GetAsyncDeleteTask();
 
 	/** Waits for the last RDG chained execution task that was launched. */
 	static RENDERCORE_API void WaitForAsyncExecuteTask();
