Fixed DumpGPU scope not using the correct namespace when compiled out.

diff --git a/Engine/Source/Runtime/RenderCore/Public/DumpGPU.h b/Engine/Source/Runtime/RenderCore/Public/DumpGPU.h
--- a/Engine/Source/Runtime/RenderCore/Public/DumpGPU.h
+++ b/Engine/Source/Runtime/RenderCore/Public/DumpGPU.h
@@ -64,10 +64,15 @@ private:
 
 #else
 
+namespace UE::RenderCore::DumpGPU
+{
+
 class FDumpScope
 {
 public:
 	FDumpScope(bool bCondition = true) {}
 };
 
+} // namespace  UE::RenderCore::DumpGPU
+
 #endif // WITH_DUMPGPU
\ No newline at end of file
