Backported multiple critical fixes to RDG and Slate for 5.5.1 patch

diff --git a/Engine/Source/Runtime/RenderCore/Private/RenderGraphBuilder.cpp b/Engine/Source/Runtime/RenderCore/Private/RenderGraphBuilder.cpp
--- a/Engine/Source/Runtime/RenderCore/Private/RenderGraphBuilder.cpp
+++ b/Engine/Source/Runtime/RenderCore/Private/RenderGraphBuilder.cpp
@@ -3942,8 +3942,10 @@ void FRDGBuilder::AddLastTextureTransition(FRDGTexture* Texture)
 		{
 			if (!SubresourceStateBefore)
 			{
+				const FRDGPassHandle AcquirePass = Texture->AcquirePass.IsValid() ? Texture->FirstPass : GetProloguePassHandle();
+
 				SubresourceStateBefore = AllocSubresource();
-				SubresourceStateBefore->SetPass(ERHIPipeline::Graphics, Texture->AcquirePass.IsValid() ? Texture->FirstPass : GetProloguePassHandle());
+				SubresourceStateBefore->SetPass(GetPassPipeline(AcquirePass), AcquirePass);
 			}
 
 			State = SubresourceStateBefore;
