Do not compile RT PSOs from the file cache to avoid bug where cache entries are going to cause a compilation failure.

diff --git a/Engine/Source/Runtime/RenderCore/Private/ShaderPipelineCache.cpp b/Engine/Source/Runtime/RenderCore/Private/ShaderPipelineCache.cpp
--- a/Engine/Source/Runtime/RenderCore/Private/ShaderPipelineCache.cpp
+++ b/Engine/Source/Runtime/RenderCore/Private/ShaderPipelineCache.cpp
@@ -270,6 +270,16 @@ static FAutoConsoleVariableRef CVarShaderPipelineCacheDoNotPrecompileComputePSO(
 												ECVF_Default
 												);
 
+// Ray tracing PSOs are recorded into the cache incorrectly due to a bug.
+bool GShaderPipelineCacheDoNotPrecompileRayTracingPSO = true;
+static FAutoConsoleVariableRef CVarShaderPipelineCacheDoNotPrecompileRayTracingPSO(
+	TEXT("r.ShaderPipelineCache.DoNotPrecompileRayTracingPSO"),
+	GShaderPipelineCacheDoNotPrecompileRayTracingPSO,
+	TEXT("Disables precompilation of raytracing PSOs (from a file cache) on start."),
+	ECVF_Default
+);
+
+
 uint32 FShaderPipelineCache::BatchSize = 0;
 float FShaderPipelineCache::BatchTime = 0.0f;
 
@@ -1132,7 +1142,7 @@ bool FShaderPipelineCacheTask::Precompile(FRHICommandListImmediate& RHICmdList,
 			// it is possible for the wrong size to be used here, which leads to a D3D run-time error when attempting to create the PSO.
 			// Ray tracing shader pre-compilation is disabled until a robust solution is found.
 			// Consider setting r.RayTracing.NonBlockingPipelineCreation=1 meanwhile to avoid most of the RTPSO creation stalls during gameplay.
-			if (IsRayTracingEnabled())
+			if (!GShaderPipelineCacheDoNotPrecompileRayTracingPSO && IsRayTracingEnabled())
 			{
 				FRayTracingShaderRHIRef RayTracingShader = FShaderCodeLibrary::CreateRayTracingShader(Platform, PSO.RayTracingDesc.ShaderHash, PSO.RayTracingDesc.Frequency);
 				if (RayTracingShader.IsValid())
@@ -1169,6 +1179,10 @@ bool FShaderPipelineCacheTask::Precompile(FRHICommandListImmediate& RHICmdList,
 					bOk = RayTracingPipeline != nullptr;
 				}
 			}
+			else
+			{
+				bOk = GShaderPipelineCacheDoNotPrecompileRayTracingPSO;
+			}
 		}
 		else
 		{
