[NNE] NNERuntimeIREE add additional shader type.

#rnx

diff --git a/Engine/Source/Runtime/RenderCore/Public/Shader.h b/Engine/Source/Runtime/RenderCore/Public/Shader.h
--- a/Engine/Source/Runtime/RenderCore/Public/Shader.h
+++ b/Engine/Source/Runtime/RenderCore/Public/Shader.h
@@ -63,6 +63,7 @@ class ITargetPlatform;
 class FCbFieldView;
 class FCbWriter;
 class FComputeKernelShaderType;
+class FNNERuntimeIREEShaderType;
 class FGlobalShaderType;
 class FMaterialShaderType;
 class FMemoryImageWriter;
@@ -1250,6 +1251,7 @@ public:
 		Niagara,
 		OCIO,
 		ComputeKernel,
+		NNERuntimeIREE,
 		NumShaderTypes,
 	};
 
@@ -1403,6 +1405,14 @@ public:
 	{
 		return (ShaderTypeForDynamicCast == EShaderTypeForDynamicCast::ComputeKernel) ? reinterpret_cast<FComputeKernelShaderType*>(this) : nullptr;
 	}
+	FORCEINLINE const FNNERuntimeIREEShaderType* GetNNERuntimeIREEShaderType() const
+	{
+		return (ShaderTypeForDynamicCast == EShaderTypeForDynamicCast::NNERuntimeIREE) ? reinterpret_cast<const FNNERuntimeIREEShaderType*>(this) : nullptr;
+	}
+	FORCEINLINE FNNERuntimeIREEShaderType* GetNNERuntimeIREEShaderType()
+	{
+		return (ShaderTypeForDynamicCast == EShaderTypeForDynamicCast::NNERuntimeIREE) ? reinterpret_cast<FNNERuntimeIREEShaderType*>(this) : nullptr;
+	}
 
 	FORCEINLINE const FGlobalShaderType* AsGlobalShaderType() const
 	{
