[shaders] keep a deterministic single copy of platform debug data; some platforms embed full source into the debug data so this can differ when we have multiple jobs that deduplicate to the same bytecode after compilation (but it doesn't matter which we keep, they necessarily only differ by line numbers/dead code)

diff --git a/Engine/Shaders/Public/ShaderVersion.ush b/Engine/Shaders/Public/ShaderVersion.ush
--- a/Engine/Shaders/Public/ShaderVersion.ush
+++ b/Engine/Shaders/Public/ShaderVersion.ush
@@ -3,5 +3,5 @@
 // Each time the console command is executed it generates a new GUID. As this file is included
 // in Platform.ush (which should be included in any shader) it allows to invalidate the shader DDC.
 // 
-#pragma message("UESHADERMETADATA_VERSION 2664417A-6712-4C43-AC7C-81D53946C03F")
+#pragma message("UESHADERMETADATA_VERSION 58C33AF1-D8A8-455A-BA24-22BF5CAF6AFC")
 
diff --git a/Engine/Source/Runtime/RenderCore/Private/ShaderResource.cpp b/Engine/Source/Runtime/RenderCore/Private/ShaderResource.cpp
--- a/Engine/Source/Runtime/RenderCore/Private/ShaderResource.cpp
+++ b/Engine/Source/Runtime/RenderCore/Private/ShaderResource.cpp
@@ -242,6 +242,13 @@ int32 FShaderMapResourceCode::FindShaderIndex(const FSHAHash& InHash) const
 	return Algo::BinarySearch(ShaderHashes, InHash);
 }
 
+#if WITH_EDITORONLY_DATA
+uint64 FShaderMapResourceCode::FShaderEditorOnlyDataEntry::GetDebugDataHash(TConstArrayView<uint8> InPlatformDebugData)
+{
+	return CityHash64(reinterpret_cast<const char*>(InPlatformDebugData.GetData()), InPlatformDebugData.Num());
+}
+#endif
+
 void FShaderMapResourceCode::AddShaderCompilerOutput(const FShaderCompilerOutput& Output, const FString& DebugName, FString DebugInfo)
 {
 	TRACE_CPUPROFILER_EVENT_SCOPE(FShaderMapResourceCode::AddShaderCode);
@@ -265,7 +272,7 @@ void FShaderMapResourceCode::AddShaderCompilerOutput(const FShaderCompilerOutput
 		// Output.Errors contains warnings in the case any exist (no errors since if there were the job would have failed)
 		// We append the warnings and deduplicate other data like DebugInfo for any additional jobs which resulted in the
 		// same bytecode for the sake of determinism in the results saved to DDC. 
-		UpdateEditorOnlyData(Index, DebugName, Output.Errors, DebugInfo);
+		UpdateEditorOnlyData(Index, DebugName, Output.PlatformDebugData, Output.Errors, DebugInfo);
 		ValidateShaderStatisticsEditorOnlyData(Index, Output.ShaderStatistics);
 	}
 #endif
@@ -275,16 +282,18 @@ void FShaderMapResourceCode::AddShaderCompilerOutput(const FShaderCompilerOutput
 void FShaderMapResourceCode::AddEditorOnlyData(int32 Index, const FString& DebugName, TConstArrayView<uint8> InPlatformDebugData, TConstArrayView<FShaderCompilerError> InCompilerWarnings, const TArray<FGenericShaderStat>& ShaderStatistics, const FString& DebugInfo)
 {
 	FShaderEditorOnlyDataEntry& Entry = ShaderEditorOnlyDataEntries.InsertDefaulted_GetRef(Index);
+
 	Entry.PlatformDebugData = InPlatformDebugData;
+	Entry.PlatformDebugDataHash = Entry.GetDebugDataHash(InPlatformDebugData);
 
 	// This should be a newly created shader entry.
 	check(Entry.ShaderStatistics.Num() == 0);
 	Entry.ShaderStatistics = ShaderStatistics;
 
-	UpdateEditorOnlyData(Index, DebugName, InCompilerWarnings, DebugInfo);
+	UpdateEditorOnlyData(Index, DebugName, InPlatformDebugData, InCompilerWarnings, DebugInfo);
 }
 
-void FShaderMapResourceCode::UpdateEditorOnlyData(int32 Index, const FString& DebugName, TConstArrayView<FShaderCompilerError> InCompilerWarnings, const FString& DebugInfo)
+void FShaderMapResourceCode::UpdateEditorOnlyData(int32 Index, const FString& DebugName, TConstArrayView<uint8> InPlatformDebugData, TConstArrayView<FShaderCompilerError> InCompilerWarnings, const FString& DebugInfo)
 {
 	FShaderEditorOnlyDataEntry& Entry = ShaderEditorOnlyDataEntries[Index];
 
@@ -294,6 +303,15 @@ void FShaderMapResourceCode::UpdateEditorOnlyData(int32 Index, const FString& De
 		Entry.DebugInfo = DebugInfo;
 	}
 
+	// As above, keep only a single copy of the platform debug data, it doesn't matter which, but use the one with the smallest
+	// hash to ensure DDC data is deterministic
+	uint64 InPlatformDebugDataHash = Entry.GetDebugDataHash(InPlatformDebugData);
+	if (InPlatformDebugDataHash < Entry.PlatformDebugDataHash)
+	{
+		Entry.PlatformDebugDataHash = InPlatformDebugDataHash;
+		Entry.PlatformDebugData = InPlatformDebugData;
+	}
+
 	for (const FShaderCompilerError& Warning : InCompilerWarnings)
 	{
 		FString ModifiedWarning = !DebugName.IsEmpty() ? FString::Printf(TEXT("%s [%s]"), *Warning.GetErrorString(), *DebugName) : Warning.GetErrorString();
diff --git a/Engine/Source/Runtime/RenderCore/Public/Shader.h b/Engine/Source/Runtime/RenderCore/Public/Shader.h
--- a/Engine/Source/Runtime/RenderCore/Public/Shader.h
+++ b/Engine/Source/Runtime/RenderCore/Public/Shader.h
@@ -524,6 +524,9 @@ public:
 #if WITH_EDITORONLY_DATA
 	struct FShaderEditorOnlyDataEntry
 	{
+		static uint64 GetDebugDataHash(TConstArrayView<uint8> PlatformDebugData);
+	
+		uint64 PlatformDebugDataHash = 0u; // Used only to ensure the copy of debug data saved is stable for determinism's sake; does not need to be serialized.
 		TArray<uint8> PlatformDebugData;
 
 		/** This field contains a debug string stored in a ShaderSymbols.info file when using r.Shaders.SymbolInfo=1
@@ -542,7 +545,12 @@ public:
 
 		friend FArchive& operator<<(FArchive& Ar, FShaderEditorOnlyDataEntry& Entry)
 		{
-			return Ar << Entry.PlatformDebugData << Entry.DebugInfo << Entry.CompilerWarnings << Entry.ShaderStatistics;
+			Ar << Entry.PlatformDebugData << Entry.DebugInfo << Entry.CompilerWarnings << Entry.ShaderStatistics;
+			if (Ar.IsLoading())
+			{
+				Entry.PlatformDebugDataHash = GetDebugDataHash(Entry.PlatformDebugData);
+			}
+			return Ar;
 		}
 	};
 #endif // WITH_EDITORONLY_DATA
@@ -574,7 +582,7 @@ public:
 
 #if WITH_EDITORONLY_DATA
 	void AddEditorOnlyData(int32 Index, const FString& DebugName, TConstArrayView<uint8> InPlatformDebugData, TConstArrayView<FShaderCompilerError> InCompilerWarnings, const TArray<FGenericShaderStat>& ShaderStatistics, const FString& DebugInfo = FString());
-	void UpdateEditorOnlyData(int32 Index, const FString& DebugName, TConstArrayView<FShaderCompilerError> InCompilerWarnings, const FString& DebugInfo);
+	void UpdateEditorOnlyData(int32 Index, const FString& DebugName, TConstArrayView<uint8> InPlatformDebugData, TConstArrayView<FShaderCompilerError> InCompilerWarnings, const FString& DebugInfo);
 
 	RENDERCORE_API void LogShaderCompilerWarnings();
 #endif
