Add missing texture methods to FShaderParametersMetadataBuilder for non-compute shader passes.

diff --git a/Engine/Source/Runtime/RenderCore/Private/ShaderParameterMetadataBuilder.cpp b/Engine/Source/Runtime/RenderCore/Private/ShaderParameterMetadataBuilder.cpp
--- a/Engine/Source/Runtime/RenderCore/Private/ShaderParameterMetadataBuilder.cpp
+++ b/Engine/Source/Runtime/RenderCore/Private/ShaderParameterMetadataBuilder.cpp
@@ -171,6 +171,93 @@ void FShaderParametersMetadataBuilder::AddRDGBufferUAV(
 	NextMemberOffset += SHADER_PARAMETER_POINTER_ALIGNMENT;
 }
 
+void FShaderParametersMetadataBuilder::AddTexture(
+	const TCHAR* Name,
+	const TCHAR* ShaderType,
+	EShaderPrecisionModifier::Type Precision)
+{
+	NextMemberOffset = Align(NextMemberOffset, SHADER_PARAMETER_POINTER_ALIGNMENT);
+
+	Members.Emplace(
+		Name,
+		ShaderType,
+		__LINE__,
+		NextMemberOffset,
+		UBMT_TEXTURE,
+		Precision,
+		TShaderResourceParameterTypeInfo<FRHITexture*>::NumRows,
+		TShaderResourceParameterTypeInfo<FRHITexture*>::NumColumns,
+		TShaderResourceParameterTypeInfo<FRHITexture*>::NumElements,
+		TShaderResourceParameterTypeInfo<FRHITexture*>::GetStructMetadata()
+	);
+
+	NextMemberOffset += SHADER_PARAMETER_POINTER_ALIGNMENT;
+}
+
+void FShaderParametersMetadataBuilder::AddRDGTexture(
+	const TCHAR* Name,
+	const TCHAR* ShaderType,
+	EShaderPrecisionModifier::Type Precision /* = EShaderPrecisionModifier::Float */
+)
+{
+	NextMemberOffset = Align(NextMemberOffset, SHADER_PARAMETER_POINTER_ALIGNMENT);
+
+	Members.Emplace(
+		Name,
+		ShaderType,
+		__LINE__,
+		NextMemberOffset,
+		UBMT_RDG_TEXTURE,
+		Precision,
+		TShaderResourceParameterTypeInfo<FRDGTexture*>::NumRows,
+		TShaderResourceParameterTypeInfo<FRDGTexture*>::NumColumns,
+		TShaderResourceParameterTypeInfo<FRDGTexture*>::NumElements,
+		TShaderResourceParameterTypeInfo<FRDGTexture*>::GetStructMetadata()
+	);
+
+	NextMemberOffset += SHADER_PARAMETER_POINTER_ALIGNMENT;
+}
+
+void FShaderParametersMetadataBuilder::AddSampler(const TCHAR* Name, const TCHAR* ShaderType, EShaderPrecisionModifier::Type Precision)
+{
+	NextMemberOffset = Align(NextMemberOffset, SHADER_PARAMETER_POINTER_ALIGNMENT);
+
+	Members.Emplace(
+		Name,
+		ShaderType,
+		__LINE__,
+		NextMemberOffset,
+		UBMT_SAMPLER,
+		Precision,
+		TShaderResourceParameterTypeInfo<FRHISamplerState*>::NumRows,
+		TShaderResourceParameterTypeInfo<FRHISamplerState*>::NumColumns,
+		TShaderResourceParameterTypeInfo<FRHISamplerState*>::NumElements,
+		TShaderResourceParameterTypeInfo<FRHISamplerState*>::GetStructMetadata()
+	);
+
+	NextMemberOffset += SHADER_PARAMETER_POINTER_ALIGNMENT;
+}
+
+void FShaderParametersMetadataBuilder::AddRenderTargetBindingSlots(const TCHAR* Name, const TCHAR* ShaderType, EShaderPrecisionModifier::Type Precision)
+{
+	NextMemberOffset = Align(NextMemberOffset, FRenderTargetBindingSlots::FTypeInfo::Alignment);
+
+	Members.Emplace(
+		Name,
+		ShaderType,
+		__LINE__,
+		NextMemberOffset,
+		UBMT_RENDER_TARGET_BINDING_SLOTS,
+		Precision,
+		FRenderTargetBindingSlots::FTypeInfo::NumRows,
+		FRenderTargetBindingSlots::FTypeInfo::NumColumns,
+		FRenderTargetBindingSlots::FTypeInfo::NumElements,
+		FRenderTargetBindingSlots::FTypeInfo::GetStructMetadata()
+	);
+
+	NextMemberOffset += sizeof(FRenderTargetBindingSlots::FTypeInfo::TAlignedType);
+}
+
 FShaderParametersMetadata* FShaderParametersMetadataBuilder::Build(
 	FShaderParametersMetadata::EUseCase UseCase,
 	const TCHAR* ShaderParameterName
diff --git a/Engine/Source/Runtime/RenderCore/Public/ShaderParameterMetadataBuilder.h b/Engine/Source/Runtime/RenderCore/Public/ShaderParameterMetadataBuilder.h
--- a/Engine/Source/Runtime/RenderCore/Public/ShaderParameterMetadataBuilder.h
+++ b/Engine/Source/Runtime/RenderCore/Public/ShaderParameterMetadataBuilder.h
@@ -164,6 +164,30 @@ public:
 		EShaderPrecisionModifier::Type Precision = EShaderPrecisionModifier::Float
 		);
 
+	RENDERCORE_API void AddTexture(
+		const TCHAR* Name,
+		const TCHAR* ShaderType,
+		EShaderPrecisionModifier::Type Precision = EShaderPrecisionModifier::Float
+	);
+
+	RENDERCORE_API void AddRDGTexture(
+		const TCHAR* Name,
+		const TCHAR* ShaderType,
+		EShaderPrecisionModifier::Type Precision = EShaderPrecisionModifier::Float
+	);
+
+	RENDERCORE_API void AddSampler(
+		const TCHAR* Name,
+		const TCHAR* ShaderType,
+		EShaderPrecisionModifier::Type Precision = EShaderPrecisionModifier::Float
+	);
+
+	RENDERCORE_API void AddRenderTargetBindingSlots(
+		const TCHAR* Name = TEXT("RenderTargets"),
+		const TCHAR* ShaderType = TEXT(""),
+		EShaderPrecisionModifier::Type Precision = EShaderPrecisionModifier::Float
+	);
+
 	void AlignNextMemberToStruct()
 	{
 		NextMemberOffset = Align(NextMemberOffset, SHADER_PARAMETER_STRUCT_ALIGNMENT);
