[shaders] non-determinism fixes for shader job cache when debug info is enabled
- eliminate flag in output object indicating whether modified source needs to be serialized or not, just clear this out as part of the compile step if it's not needed so it's always set prior to serialization (simplification)
- tell UBA about the possibility of output to the debug info folders when running jobs which output debug info; without this IFileManager::DirectoryExists can incorrectly report the folder doesn't exist on the host machine (and so the compile jobs think debug info export is disabled due to the DirectoryExists check in FShaderCompilerInput::DumpDebugInfoEnabled)

diff --git a/Engine/Plugins/UbaController/Source/Private/UbaJobProcessor.cpp b/Engine/Plugins/UbaController/Source/Private/UbaJobProcessor.cpp
--- a/Engine/Plugins/UbaController/Source/Private/UbaJobProcessor.cpp
+++ b/Engine/Plugins/UbaController/Source/Private/UbaJobProcessor.cpp
@@ -252,6 +252,11 @@ void FUbaJobProcessor::RunTaskWithUba(FTask* Task)
 	FTaskCommandData& Data = Task->CommandData;
 	SessionServer_RegisterNewFile(UbaSessionServer, TCHAR_TO_UBASTRING(*Data.InputFileName));
 
+	for (const FString& AdditionalOutputFolder : Data.AdditionalOutputFolders)
+	{
+		SessionServer_RegisterNewDirectory(UbaSessionServer, TCHAR_TO_UBASTRING(*AdditionalOutputFolder));
+	}
+
 	FString InputFileName = FPaths::GetCleanFilename(Data.InputFileName);
 	FString OutputFileName = FPaths::GetCleanFilename(Data.OutputFileName);
 	FString Parameters = FString::Printf(TEXT("\"%s/\" %d 0 \"%s\" \"%s\" %s "), *Data.WorkingDirectory, Data.DispatcherPID, *InputFileName, *OutputFileName, *Data.ExtraCommandArgs);
diff --git a/Engine/Source/Developer/DistributedBuildInterface/Public/DistributedBuildControllerInterface.h b/Engine/Source/Developer/DistributedBuildInterface/Public/DistributedBuildControllerInterface.h
--- a/Engine/Source/Developer/DistributedBuildInterface/Public/DistributedBuildControllerInterface.h
+++ b/Engine/Source/Developer/DistributedBuildInterface/Public/DistributedBuildControllerInterface.h
@@ -28,6 +28,7 @@ struct FTaskCommandData
 	FString Description; // Optional string describing the task. Shows up in UBA trace files for each job.
 	uint32 DispatcherPID = 0;
 	TArray<FString> Dependencies;
+	TArray<FString> AdditionalOutputFolders; // Optional additional folder(s) which task may write artifacts to
 };
 
 struct FTask
diff --git a/Engine/Source/Runtime/Engine/Private/ShaderCompiler/ShaderCompilerDistributed.cpp b/Engine/Source/Runtime/Engine/Private/ShaderCompiler/ShaderCompilerDistributed.cpp
--- a/Engine/Source/Runtime/Engine/Private/ShaderCompiler/ShaderCompilerDistributed.cpp
+++ b/Engine/Source/Runtime/Engine/Private/ShaderCompiler/ShaderCompilerDistributed.cpp
@@ -191,6 +191,19 @@ void FShaderCompileDistributedThreadRunnable_Interface::DispatchShaderCompileJob
 	TaskCommandData.ExtraCommandArgs = FString::Printf(TEXT("%s%s"), *FCommandLine::GetSubprocessCommandline(), GIsBuildMachine ? TEXT(" -buildmachine") : TEXT(""));
 	TaskCommandData.Dependencies = GetDependencyFilesForJobs(JobsToSerialize);
 
+	// Register any debug info paths that may be written to as additional output folders
+	// Without this remote tasks can incorrectly report that the debug info paths do not exist
+	for (const FShaderCommonCompileJobPtr& Job : JobsToSerialize)
+	{
+		Job->ForEachSingleShaderJob([&TaskCommandData](FShaderCompileJob& SingleJob)
+			{
+				if (SingleJob.Input.DumpDebugInfoEnabled())
+				{
+					TaskCommandData.AdditionalOutputFolders.Add(SingleJob.Input.DumpDebugInfoPath);
+				};
+			});
+	}
+
 	if (DistributedShaderCompilerVariables::GDistributedJobDescriptionLevel > 0)
 	{
 		TaskCommandData.Description = BuildCompactTaskDescription(JobsToSerialize);
diff --git a/Engine/Source/Runtime/RenderCore/Private/ShaderCore.cpp b/Engine/Source/Runtime/RenderCore/Private/ShaderCore.cpp
--- a/Engine/Source/Runtime/RenderCore/Private/ShaderCore.cpp
+++ b/Engine/Source/Runtime/RenderCore/Private/ShaderCore.cpp
@@ -1680,6 +1680,14 @@ public:
 		Job.Output.Target = Job.Input.Target;
 		// also set the job input hash on the output struct for validation purposes
 		Job.Output.ValidateInputHash = Job.Input.Hash;
+
+		// clear out the modified source/entry point fields if they aren't needed to avoid the data transfer/ddc overhead
+		if (!Job.Input.NeedsOriginalShaderSource())
+		{
+			Job.Output.ModifiedShaderSource.Empty();
+			Job.Output.ModifiedEntryPointName.Empty();
+		}
+
 		if (Job.Output.bSucceeded)
 		{
 			Job.Output.GenerateOutputHash();
@@ -1728,7 +1736,7 @@ bool PreprocessShader(FShaderCommonCompileJob* Job)
 				FString Error = FString::Printf(
 					TEXT("Preprocessing %s stage skipped due to earlier stage preprocessing failure."),
 					GetShaderFrequencyString(StageJob->Input.Target.GetFrequency()));
-				StageJob->Output.Errors.Add(FShaderCompilerError(*Error));
+				StageJob->Output.Errors.Add(FShaderCompilerError(MoveTemp(Error)));
 			}
 		}
 		return !bAnyFailed;
@@ -4108,11 +4116,6 @@ void FShaderCompileJob::AppendDebugName(FStringBuilderBase& OutName) const
 
 void FShaderCompileJob::SerializeWorkerOutput(FArchive& Ar)
 {
-	// Only serialize the modified source/entry point if either:
-	// (a) this job requires debug output (since this is done in the cooker process the modified source is needed outside of SCW)
-	// or (b) the "extract shader source" setting is enabled (i.e. something upstream explicitly wants the final source passed to the compiler).
-	Output.bSerializeModifiedSource = Input.DumpDebugInfoEnabled() || Input.ExtraSettings.bExtractShaderSource;
-
 	Ar << Output;
 
 	bool bSecondaryOutput = SecondaryOutput.IsValid();
diff --git a/Engine/Source/Runtime/RenderCore/Public/ShaderCompilerCore.h b/Engine/Source/Runtime/RenderCore/Public/ShaderCompilerCore.h
--- a/Engine/Source/Runtime/RenderCore/Public/ShaderCompilerCore.h
+++ b/Engine/Source/Runtime/RenderCore/Public/ShaderCompilerCore.h
@@ -30,7 +30,7 @@ typedef TSharedPtr<TArray<ANSICHAR>, ESPMode::ThreadSafe> FShaderSharedAnsiStrin
 // this is for the protocol, not the data, bump if FShaderCompilerInput/FShaderPreprocessOutput serialization, SerializeWorkerInput or ProcessInputFromArchive changes.
 inline const int32 ShaderCompileWorkerInputVersion = 29;
 // this is for the protocol, not the data, bump if FShaderCompilerOutput or WriteToOutputArchive changes.
-inline const int32 ShaderCompileWorkerOutputVersion = 26;
+inline const int32 ShaderCompileWorkerOutputVersion = 27;
 // this is for the protocol, not the data.
 inline const int32 ShaderCompileWorkerSingleJobHeader = 'S';
 // this is for the protocol, not the data.
@@ -38,7 +38,7 @@ inline const int32 ShaderCompileWorkerPipelineJobHeader = 'P';
 
 // Modify this to invalidate _just_ the cache/DDC entries for individual shaders (will not cause shadermaps to rebuild if they are not otherwise out-of-date).
 // This should be bumped for changes to the FShaderCompilerOutput data structure (in addition to ShaderCompileWorkerOutputVersion)
-inline static const FGuid UE_SHADER_CACHE_VERSION = FGuid("5E33C861-2B87-4B58-A4B4-741C3E064881");
+inline static const FGuid UE_SHADER_CACHE_VERSION = FGuid("902C35E9-5B0A-460C-96FF-C988651E9B50");
 
 namespace UE::ShaderCompiler
 {
@@ -553,7 +553,6 @@ struct FShaderCompilerOutput
 	,	PreprocessTime(0.0)
 	,	bSucceeded(false)
 	,	bSupportsQueryingUsedAttributes(false)
-	,	bSerializeModifiedSource(false)
 	,	bSerializingForCache(false)
 	{
 	}
@@ -572,13 +571,14 @@ struct FShaderCompilerOutput
 	double PreprocessTime;
 	bool bSucceeded;
 	bool bSupportsQueryingUsedAttributes;
+	UE_DEPRECATED(5.6, "bSerializeModifiedSource is no longer used")
 	bool bSerializeModifiedSource;
 	bool bSerializingForCache;
 	TArray<FString> UsedAttributes;
 
 	PRAGMA_DISABLE_DEPRECATION_WARNINGS
 	// Explicitly-defaulted copy/move ctors & assignment operators are needed temporarily due to deprecation
-	// of PragmaDirectives field. These can be removed once the deprecation window for said field ends.
+	// of PragmaDirectives and bSerializeModifiedSource fields. These can be removed once the deprecation window for said fields ends.
 	FShaderCompilerOutput(FShaderCompilerOutput&&) = default;
 	FShaderCompilerOutput(const FShaderCompilerOutput&) = default;
 	FShaderCompilerOutput& operator=(FShaderCompilerOutput&&) = default;
@@ -658,12 +658,8 @@ struct FShaderCompilerOutput
 		Ar << Output.NumInstructions;
 		Ar << Output.NumTextureSamplers;
 		Ar << Output.bSucceeded;
-		Ar << Output.bSerializeModifiedSource;
-		if (Output.bSerializeModifiedSource)
-		{
-			Ar << Output.ModifiedShaderSource;
-			Ar << Output.ModifiedEntryPointName;
-		}
+		Ar << Output.ModifiedShaderSource;
+		Ar << Output.ModifiedEntryPointName;
 		Ar << Output.PlatformDebugData;
 		Ar << Output.ShaderStatistics;
 
