Simplify logic to test for workgraph shader.

diff --git a/Engine/Source/Runtime/RenderCore/Public/Shader.h b/Engine/Source/Runtime/RenderCore/Public/Shader.h
--- a/Engine/Source/Runtime/RenderCore/Public/Shader.h
+++ b/Engine/Source/Runtime/RenderCore/Public/Shader.h
@@ -1159,7 +1159,7 @@ public:
 		if (ShaderContent)
 		{
 			const EShaderFrequency Frequency = ShaderContent->GetFrequency();
-			checkSlow(Frequency == SF_WorkGraphRoot || Frequency == SF_WorkGraphComputeNode);
+			checkSlow(IsWorkGraphShaderFrequency(Frequency));
 			RHIShader = static_cast<FRHIWorkGraphShader*>(GetRHIShaderBase(Frequency));
 		}
 		return RHIShader;
