Fixed edge case with reserved buffers in RDG when the buffer is not used in any pass. Before it would be culled and not committed. Now it will always be committed.

diff --git a/Engine/Source/Runtime/RenderCore/Private/RenderGraphBuilder.cpp b/Engine/Source/Runtime/RenderCore/Private/RenderGraphBuilder.cpp
--- a/Engine/Source/Runtime/RenderCore/Private/RenderGraphBuilder.cpp
+++ b/Engine/Source/Runtime/RenderCore/Private/RenderGraphBuilder.cpp
@@ -3843,6 +3843,10 @@ void FRDGBuilder::FinalizeResources()
 					AddLastBufferTransition(Buffer);
 				}
 			}
+			else if (Buffer->PendingCommitSize != 0)
+			{
+				AddCulledReservedCommitTransition(Buffer);
+			}
 
 			if (Buffer->Allocation)
 			{
@@ -4029,6 +4033,28 @@ void FRDGBuilder::AddLastBufferTransition(FRDGBuffer* Buffer)
 	AddBufferTransition(Buffer, Buffer->State, StateAfter);
 }
 
+void FRDGBuilder::AddCulledReservedCommitTransition(FRDGBufferRef Buffer)
+{
+	check(Buffer->HasRHI() && Buffer->bExternal && Buffer->PendingCommitSize > 0);
+	check(Buffer->ReferenceCount == 0);
+
+	const FRDGPassHandle ProloguePassHandle = GetProloguePassHandle();
+	const FRDGPassHandle EpiloguePassHandle = GetEpiloguePassHandle();
+
+	FRDGSubresourceState* StateBefore = AllocSubresource();
+	StateBefore->SetPass(ERHIPipeline::Graphics, IsImmediateMode() ? EpiloguePassHandle : ProloguePassHandle);
+
+	FRDGSubresourceState* StateAfter = AllocSubresource();
+	StateAfter->SetPass(ERHIPipeline::Graphics, EpiloguePassHandle);
+	StateAfter->Access = Buffer->EpilogueAccess;
+	StateAfter->ReservedCommitHandle = AcquireReservedCommitHandle(Buffer);
+
+	EpilogueResourceAccesses.Emplace(Buffer->GetRHI(), StateAfter->Access);
+	Buffer->Allocation = nullptr;
+
+	AddBufferTransition(Buffer, StateBefore, StateAfter);
+}
+
 template <typename FilterSubresourceLambdaType>
 void FRDGBuilder::AddTextureTransition(FRDGTexture* Texture, FRDGTextureSubresourceState& StateBefore, FRDGTextureSubresourceState& StateAfter, FilterSubresourceLambdaType&& FilterSubresourceLambda)
 {
diff --git a/Engine/Source/Runtime/RenderCore/Public/RenderGraphBuilder.h b/Engine/Source/Runtime/RenderCore/Public/RenderGraphBuilder.h
--- a/Engine/Source/Runtime/RenderCore/Public/RenderGraphBuilder.h
+++ b/Engine/Source/Runtime/RenderCore/Public/RenderGraphBuilder.h
@@ -755,6 +755,8 @@ private:
 	void AddLastTextureTransition(FRDGTextureRef Texture);
 	void AddLastBufferTransition(FRDGBufferRef Buffer);
 
+	void AddCulledReservedCommitTransition(FRDGBufferRef Buffer);
+
 	template <typename FilterSubresourceLambdaType>
 	void AddTextureTransition(
 		FRDGTextureRef Texture,
diff --git a/Engine/Source/Runtime/RenderCore/Public/RenderGraphResources.h b/Engine/Source/Runtime/RenderCore/Public/RenderGraphResources.h
--- a/Engine/Source/Runtime/RenderCore/Public/RenderGraphResources.h
+++ b/Engine/Source/Runtime/RenderCore/Public/RenderGraphResources.h
@@ -300,11 +300,6 @@ public:
 		return bExtracted;
 	}
 
-	bool IsCulled() const
-	{
-		return ReferenceCount == 0;
-	}
-
 	/** Whether a prior pass added to the graph produced contents for this resource. External resources are not considered produced
 	 *  until used for a write operation. This is a union of all subresources, so any subresource write will set this to true.
 	 */
@@ -604,6 +599,11 @@ public:
 
 	RENDERCORE_API FRDGTextureSubresourceRange GetSubresourceRangeSRV() const;
 
+	bool IsCulled() const
+	{
+		return ReferenceCount == 0;
+	}
+
 private:
 	RENDERCORE_API FRDGTexture(const TCHAR* InName, const FRDGTextureDesc& InDesc, ERDGTextureFlags InFlags);
 
@@ -1352,6 +1352,11 @@ public:
 		return Desc.BytesPerElement;
 	}
 
+	bool IsCulled() const
+	{
+		return ReferenceCount == 0 && PendingCommitSize == 0;
+	}
+
 private:
 	RENDERCORE_API FRDGBuffer(const TCHAR* InName, const FRDGBufferDesc& InDesc, ERDGBufferFlags InFlags);
 	RENDERCORE_API FRDGBuffer(const TCHAR* InName, const FRDGBufferDesc& InDesc, ERDGBufferFlags InFlags, FRDGBufferNumElementsCallback&& InNumElementsCallback);
