[shaders] fix include of internal ShaderSerialization.h from public header Shader.h

diff --git a/Engine/Plugins/Compositing/OpenColorIO/Source/OpenColorIO/Private/OpenColorIOShaderMap.cpp b/Engine/Plugins/Compositing/OpenColorIO/Source/OpenColorIO/Private/OpenColorIOShaderMap.cpp
--- a/Engine/Plugins/Compositing/OpenColorIO/Source/OpenColorIO/Private/OpenColorIOShaderMap.cpp
+++ b/Engine/Plugins/Compositing/OpenColorIO/Source/OpenColorIO/Private/OpenColorIOShaderMap.cpp
@@ -11,6 +11,7 @@
 #include "Serialization/MemoryWriter.h"
 #include "Serialization/MemoryReader.h"
 #include "ShaderCompiler.h"
+#include "ShaderSerialization.h"
 #include "Stats/StatsMisc.h"
 #include "UObject/UObjectGlobals.h"
 
diff --git a/Engine/Plugins/Experimental/NNERuntimeIREE/Source/NNERuntimeIREEShader/Private/NNERuntimeIREEShaderMap.cpp b/Engine/Plugins/Experimental/NNERuntimeIREE/Source/NNERuntimeIREEShader/Private/NNERuntimeIREEShaderMap.cpp
--- a/Engine/Plugins/Experimental/NNERuntimeIREE/Source/NNERuntimeIREEShader/Private/NNERuntimeIREEShaderMap.cpp
+++ b/Engine/Plugins/Experimental/NNERuntimeIREE/Source/NNERuntimeIREEShader/Private/NNERuntimeIREEShaderMap.cpp
@@ -12,6 +12,7 @@
 #include "Serialization/MemoryWriter.h"
 #include "Serialization/MemoryReader.h"
 #include "ShaderCompiler.h"
+#include "ShaderSerialization.h"
 #include "Stats/StatsMisc.h"
 #include "UObject/UObjectGlobals.h"
 #include "DataDrivenShaderPlatformInfo.h"
diff --git a/Engine/Plugins/FX/Niagara/Source/NiagaraEditor/Private/NiagaraCompilationTasks.cpp b/Engine/Plugins/FX/Niagara/Source/NiagaraEditor/Private/NiagaraCompilationTasks.cpp
--- a/Engine/Plugins/FX/Niagara/Source/NiagaraEditor/Private/NiagaraCompilationTasks.cpp
+++ b/Engine/Plugins/FX/Niagara/Source/NiagaraEditor/Private/NiagaraCompilationTasks.cpp
@@ -24,6 +24,7 @@
 #include "Serialization/CompactBinaryWriter.h"
 #include "Serialization/ObjectAndNameAsStringProxyArchive.h"
 #include "ShaderParameterMetadataBuilder.h"
+#include "ShaderSerialization.h"
 
 // needed for AsType...pretty ugly
 #include "NiagaraNodeOutput.h"
diff --git a/Engine/Plugins/FX/Niagara/Source/NiagaraShader/Private/NiagaraShader.cpp b/Engine/Plugins/FX/Niagara/Source/NiagaraShader/Private/NiagaraShader.cpp
--- a/Engine/Plugins/FX/Niagara/Source/NiagaraShader/Private/NiagaraShader.cpp
+++ b/Engine/Plugins/FX/Niagara/Source/NiagaraShader/Private/NiagaraShader.cpp
@@ -8,6 +8,7 @@
 #include "Serialization/MemoryWriter.h"
 #include "Serialization/MemoryReader.h"
 #include "ShaderCompiler.h"
+#include "ShaderSerialization.h"
 #include "DataDrivenShaderPlatformInfo.h"
 #include "UObject/DevObjectVersion.h"
 #include "NiagaraShaderCompilationManager.h"
@@ -1058,6 +1059,14 @@ FNiagaraShaderMap::~FNiagaraShaderMap()
 	check(!bRegistered);
 }
 
+
+bool FNiagaraShaderMap::Serialize(FArchive& Ar, bool bInlineShaderResources, bool bLoadingCooked)
+{
+	FShaderSerializeContext Ctx(Ar);
+	Ctx.bLoadingCooked = bLoadingCooked;
+	return Serialize(Ctx);
+}
+
 bool FNiagaraShaderMap::Serialize(FShaderSerializeContext& Ctx)
 {
 	// Note: This is saved to the DDC, not into packages (except when cooked)
diff --git a/Engine/Plugins/FX/Niagara/Source/NiagaraShader/Private/NiagaraShared.cpp b/Engine/Plugins/FX/Niagara/Source/NiagaraShader/Private/NiagaraShared.cpp
--- a/Engine/Plugins/FX/Niagara/Source/NiagaraShader/Private/NiagaraShared.cpp
+++ b/Engine/Plugins/FX/Niagara/Source/NiagaraShader/Private/NiagaraShared.cpp
@@ -18,6 +18,7 @@
 #include "UObject/UObjectHash.h"
 #include "PSOPrecache.h"
 #include "ShaderCompiler.h"
+#include "ShaderSerialization.h"
 #include "NiagaraShaderCompilationManager.h"
 #include "Modules/ModuleManager.h"
 #include "NiagaraCustomVersion.h"
diff --git a/Engine/Plugins/FX/Niagara/Source/NiagaraShader/Public/NiagaraShared.h b/Engine/Plugins/FX/Niagara/Source/NiagaraShader/Public/NiagaraShared.h
--- a/Engine/Plugins/FX/Niagara/Source/NiagaraShader/Public/NiagaraShared.h
+++ b/Engine/Plugins/FX/Niagara/Source/NiagaraShader/Public/NiagaraShared.h
@@ -600,12 +600,7 @@ public:
 
 	/** Serializes the shader map. */
 	UE_DEPRECATED(5.5, "Use overload taking a FShaderSerializeContext")
-	NIAGARASHADER_API bool Serialize(FArchive& Ar, bool bInlineShaderResources = true, bool bLoadingCooked = false)
-	{
-		FShaderSerializeContext Ctx(Ar);
-		Ctx.bLoadingCooked = bLoadingCooked;
-		return Serialize(Ctx);
-	}
+	NIAGARASHADER_API bool Serialize(FArchive& Ar, bool bInlineShaderResources = true, bool bLoadingCooked = false);
 
 	NIAGARASHADER_API bool Serialize(FShaderSerializeContext& Ctx);
 
diff --git a/Engine/Plugins/Runtime/ComputeFramework/Source/ComputeFramework/Private/ComputeKernelShaderMap.cpp b/Engine/Plugins/Runtime/ComputeFramework/Source/ComputeFramework/Private/ComputeKernelShaderMap.cpp
--- a/Engine/Plugins/Runtime/ComputeFramework/Source/ComputeFramework/Private/ComputeKernelShaderMap.cpp
+++ b/Engine/Plugins/Runtime/ComputeFramework/Source/ComputeFramework/Private/ComputeKernelShaderMap.cpp
@@ -14,6 +14,7 @@
 #include "Serialization/MemoryWriter.h"
 #include "Serialization/MemoryReader.h"
 #include "ShaderCompiler.h"
+#include "ShaderSerialization.h"
 #include "Stats/StatsMisc.h"
 #include "UObject/UObjectGlobals.h"
 #include "DataDrivenShaderPlatformInfo.h"
diff --git a/Engine/Source/Runtime/Engine/Private/Materials/MaterialShader.cpp b/Engine/Source/Runtime/Engine/Private/Materials/MaterialShader.cpp
--- a/Engine/Source/Runtime/Engine/Private/Materials/MaterialShader.cpp
+++ b/Engine/Source/Runtime/Engine/Private/Materials/MaterialShader.cpp
@@ -27,6 +27,7 @@
 #include "SceneTexturesConfig.h"
 #include "Serialization/ShaderKeyGenerator.h"
 #include "ShaderCodeLibrary.h"
+#include "ShaderSerialization.h"
 #include "UObject/UE5ReleaseStreamObjectVersion.h"
 #include "Materials/Material.h"
 #include "MaterialHLSLGenerator.h"
@@ -52,7 +53,6 @@
 #include "Serialization/CompactBinarySerialization.h"
 #include "Serialization/CompactBinaryWriter.h"
 #include "Serialization/MemoryReader.h"
-#include "ShaderSerialization.h"
 #endif
 
 CSV_DECLARE_CATEGORY_MODULE_EXTERN(RENDERCORE_API, Shaders);
@@ -3615,6 +3615,14 @@ FMaterialShaderMap* FMaterialShaderMap::GetFinalizedClone() const
 }
 #endif // WITH_EDITOR
 
+bool FMaterialShaderMap::Serialize(FArchive& Ar, bool bInlineShaderResources, bool bLoadingCooked, bool bInlineShaderCode, const FName& SerializingAsset)
+{
+	FShaderSerializeContext Ctx(Ar);
+	Ctx.bLoadingCooked = bLoadingCooked;
+	Ctx.SerializingAsset = SerializingAsset;
+	return Serialize(Ctx);
+}
+
 bool FMaterialShaderMap::Serialize(FShaderSerializeContext& Ctx)
 {
 	SCOPED_LOADTIMER(FMaterialShaderMap_Serialize);
diff --git a/Engine/Source/Runtime/Engine/Private/Materials/MaterialShared.cpp b/Engine/Source/Runtime/Engine/Private/Materials/MaterialShared.cpp
--- a/Engine/Source/Runtime/Engine/Private/Materials/MaterialShared.cpp
+++ b/Engine/Source/Runtime/Engine/Private/Materials/MaterialShared.cpp
@@ -22,6 +22,7 @@
 #include "Materials/MaterialExpressionBreakMaterialAttributes.h"
 #include "Materials/MaterialExpressionRerouteBase.h"
 #include "ShaderCompiler.h"
+#include "ShaderSerialization.h"
 #include "MeshMaterialShader.h"
 #include "MeshMaterialShaderType.h"
 #include "RendererInterface.h"
diff --git a/Engine/Source/Runtime/Engine/Public/MaterialShared.h b/Engine/Source/Runtime/Engine/Public/MaterialShared.h
--- a/Engine/Source/Runtime/Engine/Public/MaterialShared.h
+++ b/Engine/Source/Runtime/Engine/Public/MaterialShared.h
@@ -1567,13 +1567,7 @@ public:
 	ENGINE_API void Release();
 
 	UE_DEPRECATED(5.5, "Please use overload accepting an FShaderSerializeContext. Note that inlining of shader code into runtime shadermaps is no longer supported. If you rely on this please reach out to the UE rendering team.")
-	bool Serialize(FArchive& Ar, bool bInlineShaderResources = true, bool bLoadingCooked = false, bool bInlineShaderCode = false, const FName& SerializingAsset = NAME_None)
-	{
-		FShaderSerializeContext Ctx(Ar);
-		Ctx.bLoadingCooked = bLoadingCooked;
-		Ctx.SerializingAsset = SerializingAsset;
-		return Serialize(Ctx);
-	}
+	bool Serialize(FArchive& Ar, bool bInlineShaderResources = true, bool bLoadingCooked = false, bool bInlineShaderCode = false, const FName& SerializingAsset = NAME_None);
 
 	/** Serializes the shader map. */
 	bool Serialize(FShaderSerializeContext& Ctx);
diff --git a/Engine/Source/Runtime/RenderCore/Private/ShaderMap.cpp b/Engine/Source/Runtime/RenderCore/Private/ShaderMap.cpp
--- a/Engine/Source/Runtime/RenderCore/Private/ShaderMap.cpp
+++ b/Engine/Source/Runtime/RenderCore/Private/ShaderMap.cpp
@@ -5,17 +5,19 @@
 =============================================================================*/
 
 #include "Shader.h"
+
+#include "DataDrivenShaderPlatformInfo.h"
+#include "HAL/LowLevelMemTracker.h"
 #include "Misc/App.h"
 #include "Misc/CoreMisc.h"
+#include "Misc/ScopeLock.h"
 #include "Misc/StringBuilder.h"
-#include "VertexFactory.h"
+#include "ProfilingDebugging/CsvProfiler.h"
 #include "ShaderCodeLibrary.h"
 #include "ShaderCore.h"
-#include "Misc/ScopeLock.h"
+#include "ShaderSerialization.h"
 #include "UObject/RenderingObjectVersion.h"
-#include "HAL/LowLevelMemTracker.h"
-#include "DataDrivenShaderPlatformInfo.h"
-#include "ProfilingDebugging/CsvProfiler.h"
+#include "VertexFactory.h"
 
 static EShaderPermutationFlags GetCurrentShaderPermutationFlags()
 {
@@ -235,6 +237,14 @@ void FShaderMapBase::UnfreezeContent()
 
 #define CHECK_SHADERMAP_DEPENDENCIES (WITH_EDITOR || !(UE_BUILD_SHIPPING || UE_BUILD_TEST))
 
+bool FShaderMapBase::Serialize(FArchive& Ar, bool bInlineShaderResources, bool bLoadingCooked, bool bInlineShaderCode, const FName& SerializingAsset)
+{
+	FShaderSerializeContext Ctx(Ar);
+	Ctx.bLoadingCooked = bLoadingCooked;
+	Ctx.SerializingAsset = SerializingAsset;
+	return Serialize(Ctx);
+}
+
 bool FShaderMapBase::Serialize(FShaderSerializeContext& Ctx)
 {
 	FArchive& Ar = Ctx.GetMainArchive();
diff --git a/Engine/Source/Runtime/RenderCore/Private/ShaderResource.cpp b/Engine/Source/Runtime/RenderCore/Private/ShaderResource.cpp
--- a/Engine/Source/Runtime/RenderCore/Private/ShaderResource.cpp
+++ b/Engine/Source/Runtime/RenderCore/Private/ShaderResource.cpp
@@ -20,6 +20,7 @@
 #include "RHIResources.h"	// Access to FRHIRayTracingShader::RayTracingPayloadType requires this
 #include "ShaderCompilerCore.h"
 #include "ShaderCore.h"
+#include "ShaderSerialization.h"
 #include "UObject/RenderingObjectVersion.h"
 
 #if WITH_EDITORONLY_DATA
@@ -376,6 +377,13 @@ void FShaderMapResourceCode::ToString(FStringBuilderBase& OutString) const
 	}
 }
 
+void FShaderMapResourceCode::Serialize(FArchive& Ar, bool bLoadingCooked)
+{
+	FShaderSerializeContext Ctx(Ar);
+	Ctx.bLoadingCooked = bLoadingCooked;
+	Serialize(Ctx);
+}
+
 void FShaderMapResourceCode::Serialize(FShaderSerializeContext& Ctx)
 {
 	FArchive& Ar = Ctx.GetMainArchive();
diff --git a/Engine/Source/Runtime/RenderCore/Public/Shader.h b/Engine/Source/Runtime/RenderCore/Public/Shader.h
--- a/Engine/Source/Runtime/RenderCore/Public/Shader.h
+++ b/Engine/Source/Runtime/RenderCore/Public/Shader.h
@@ -42,7 +42,6 @@
 #include "ShaderParameterMetadata.h"
 #include "ShaderParameters.h"
 #include "ShaderPermutation.h"
-#include "ShaderSerialization.h"
 #include "Templates/RefCounting.h"
 #include "Templates/TypeHash.h"
 #include "Templates/UniquePtr.h"
@@ -83,6 +82,7 @@ class FVertexFactoryType;
 class FRHIShaderBindingLayout;
 struct FShaderCompiledShaderInitializerType;
 struct FShaderCompilerOutput;
+struct FShaderSerializeContext;
 using FShaderMapAssetPaths = TSet<FName>; // Copied from ShaderCodeLibrary.h
 
 /** Define a shader permutation uniquely according to its type, and permutation id.*/
@@ -562,12 +562,7 @@ public:
 	RENDERCORE_API void Finalize();
 
 	UE_DEPRECATED(5.5, "Use overload accepting an FShaderSerializeContext object")
-		void Serialize(FArchive& Ar, bool bLoadingCooked)
-	{
-		FShaderSerializeContext Ctx(Ar);
-		Ctx.bLoadingCooked = bLoadingCooked;
-		Serialize(Ctx);
-	}
+	RENDERCORE_API void Serialize(FArchive& Ar, bool bLoadingCooked);
 
 	RENDERCORE_API void Serialize(FShaderSerializeContext& Ctx);
 #if WITH_EDITORONLY_DATA
@@ -2386,13 +2381,7 @@ public:
 	RENDERCORE_API void FinalizeContent();
 	RENDERCORE_API void UnfreezeContent();
 	UE_DEPRECATED(5.5, "Please use overload accepting an FShaderSerializeContext. Note that inlining of shader code into runtime shadermaps is no longer supported. If you rely on this please reach out to the UE rendering team.")
-	RENDERCORE_API bool Serialize(FArchive& Ar, bool bInlineShaderResources, bool bLoadingCooked, bool bInlineShaderCode = false, const FName& SerializingAsset = NAME_None)
-	{
-		FShaderSerializeContext Ctx(Ar);
-		Ctx.bLoadingCooked = bLoadingCooked;
-		Ctx.SerializingAsset = SerializingAsset;
-		return Serialize(Ctx);
-	}
+	RENDERCORE_API bool Serialize(FArchive& Ar, bool bInlineShaderResources, bool bLoadingCooked, bool bInlineShaderCode = false, const FName& SerializingAsset = NAME_None);
 
 	RENDERCORE_API bool Serialize(FShaderSerializeContext& Ctx);
 
