[shaders] add helper in RenderCore for initializing shader debug info path so it can be called when initializing symbol override path; add shaderdebuginfo replacement token option in symbol override path

diff --git a/Engine/Source/Runtime/Engine/Private/ShaderCompiler/ShaderCompiler.cpp b/Engine/Source/Runtime/Engine/Private/ShaderCompiler/ShaderCompiler.cpp
--- a/Engine/Source/Runtime/Engine/Private/ShaderCompiler/ShaderCompiler.cpp
+++ b/Engine/Source/Runtime/Engine/Private/ShaderCompiler/ShaderCompiler.cpp
@@ -513,20 +513,6 @@ static FAutoConsoleVariableRef CVarCreateShadersOnLoad(
 	TEXT("Whether to create shaders on load, which can reduce hitching, but use more memory.  Otherwise they will be created as needed.")
 );
 
-static TAutoConsoleVariable<FString> CVarShaderOverrideDebugDir(
-	TEXT("r.OverrideShaderDebugDir"),
-	"",
-	TEXT("Override output location of shader debug files\n")
-	TEXT("Empty: use default location Saved\\ShaderDebugInfo.\n"),
-	ECVF_ReadOnly);
-
-static TAutoConsoleVariable<bool> CVarDisambiguateShaderDebugDir(
-	TEXT("r.DisambiguateShaderDebugDir"),
-	false,
-	TEXT("If true, appends a folder containing the full project path with directory separators/drive qualifiers replaced with _ to the root debug info folder.\n")
-	TEXT("Intended for use in conjunction with r.OverrideShaderDebugDir to avoid shaderdebuginfo output clashing across workspaces/projects."),
-	ECVF_ReadOnly);
-
 static TAutoConsoleVariable<int32> CVarShadersValidation(
 	TEXT("r.Shaders.Validation"),
 	1,
@@ -701,10 +687,6 @@ void FShaderCompilingManager::ReportMemoryUsage()
 	}
 }
 
-static FString GetBuildMachineArtifactBasePath()
-{
-	return FPaths::Combine(*FPaths::EngineDir(), TEXT("Programs"), TEXT("AutomationTool"), TEXT("Saved"), TEXT("Logs"));
-}
 
 static bool FindShaderCompileWorkerExecutableInLaunchDir(const FString& ExecutableName, FString& OutFilename)
 {
@@ -844,27 +826,6 @@ FShaderCompilingManager::FShaderCompilingManager() :
 	FPaths::NormalizeDirectoryName(AbsoluteBaseDirectory);
 	AbsoluteShaderBaseWorkingDirectory = AbsoluteBaseDirectory + TEXT("/");
 
-	// Build machines should dump to the AutomationTool/Saved/Logs directory and they will upload as build artifacts via the AutomationTool.
-	const FString BaseDebugInfoPath = GIsBuildMachine ? GetBuildMachineArtifactBasePath() : FPaths::ProjectSavedDir();
-
-	FString AbsoluteDebugInfoDirectory = IFileManager::Get().ConvertToAbsolutePathForExternalAppForWrite(*(BaseDebugInfoPath / TEXT("ShaderDebugInfo")));
-	const FString OverrideShaderDebugDir = CVarShaderOverrideDebugDir.GetValueOnAnyThread();
-	if (!OverrideShaderDebugDir.IsEmpty())
-	{
-		AbsoluteDebugInfoDirectory = OverrideShaderDebugDir;
-	}
-
-	if (CVarDisambiguateShaderDebugDir.GetValueOnAnyThread())
-	{
-		FString AppendFolder = FPaths::ConvertRelativePathToFull(FPaths::ProjectDir());
-		FPaths::NormalizeDirectoryName(AppendFolder);
-		AppendFolder.ReplaceInline(TEXT(":/"), TEXT("_"));
-		AppendFolder.ReplaceCharInline('/', '_');
-		AbsoluteDebugInfoDirectory /= AppendFolder;
-	}
-
-	FPaths::NormalizeDirectoryName(AbsoluteDebugInfoDirectory);
-	AbsoluteShaderDebugInfoDirectory = AbsoluteDebugInfoDirectory;
 
 	CalculateNumberOfCompilingThreads(FPlatformMisc::NumberOfCores(), FPlatformMisc::NumberOfCoresIncludingHyperthreads());
 
diff --git a/Engine/Source/Runtime/Engine/Public/ShaderCompiler.h b/Engine/Source/Runtime/Engine/Public/ShaderCompiler.h
--- a/Engine/Source/Runtime/Engine/Public/ShaderCompiler.h
+++ b/Engine/Source/Runtime/Engine/Public/ShaderCompiler.h
@@ -1101,7 +1101,7 @@ public:
 
 	const FString& GetAbsoluteShaderDebugInfoDirectory() const
 	{
-		return AbsoluteShaderDebugInfoDirectory;
+		return GetShaderDebugInfoPath();
 	}
 
 	bool AreWarningsSuppressed(const EShaderPlatform Platform) const
diff --git a/Engine/Source/Runtime/RenderCore/Private/ShaderCore.cpp b/Engine/Source/Runtime/RenderCore/Private/ShaderCore.cpp
--- a/Engine/Source/Runtime/RenderCore/Private/ShaderCore.cpp
+++ b/Engine/Source/Runtime/RenderCore/Private/ShaderCore.cpp
@@ -112,6 +112,20 @@ static TAutoConsoleVariable<bool> CVarDumpDebugInfoForCacheHits(
 	TEXT("If true, debug info (via IShaderFormat::OutputDebugData) will be output for all jobs including duplicates and cache/DDC hits. If false, only jobs that actually executed compilation will dump debug info."),
 	ECVF_Default);
 
+static TAutoConsoleVariable<FString> CVarShaderOverrideDebugDir(
+	TEXT("r.OverrideShaderDebugDir"),
+	"",
+	TEXT("Override output location of shader debug files\n")
+	TEXT("Empty: use default location Saved\\ShaderDebugInfo.\n"),
+	ECVF_ReadOnly);
+
+static TAutoConsoleVariable<bool> CVarDisambiguateShaderDebugDir(
+	TEXT("r.DisambiguateShaderDebugDir"),
+	false,
+	TEXT("If true, appends a folder containing the full project path with directory separators/drive qualifiers replaced with _ to the root debug info folder.\n")
+	TEXT("Intended for use in conjunction with r.OverrideShaderDebugDir to avoid shaderdebuginfo output clashing across workspaces/projects."),
+	ECVF_ReadOnly);
+
 void UpdateShaderDevelopmentMode()
 {
 	// Keep LogShaders verbosity in sync with r.ShaderDevelopmentMode
@@ -681,6 +695,49 @@ bool ShouldAllowUniqueShaderSymbols(FName ShaderFormat)
 	return AllowUniqueSymbols.IsEnabled(ShaderFormat);
 }
 
+const FString& GetBuildMachineArtifactBasePath()
+{
+	// no need for thread safety here; this will be hit on a single thread during shader compiling manager init long before it's called anywhere else
+	static FString ArtifactBasePath;
+	if (ArtifactBasePath.IsEmpty())
+	{
+		ArtifactBasePath = FPaths::Combine(*FPaths::EngineDir(), TEXT("Programs"), TEXT("AutomationTool"), TEXT("Saved"), TEXT("Logs"));
+	}
+	return ArtifactBasePath;
+}
+
+
+const FString& GetShaderDebugInfoPath()
+{
+	static FString DebugInfoPath;
+	// no need for thread safety here; this will be hit on a single thread during shader compiling manager init long before it's called anywhere else
+	if (DebugInfoPath.IsEmpty())
+	{
+		// Build machines should dump to the AutomationTool/Saved/Logs directory and they will upload as build artifacts via the AutomationTool.
+		const FString BaseDebugInfoPath = GIsBuildMachine ? GetBuildMachineArtifactBasePath() : FPaths::ProjectSavedDir();
+
+		FString AbsoluteDebugInfoDirectory = IFileManager::Get().ConvertToAbsolutePathForExternalAppForWrite(*(BaseDebugInfoPath / TEXT("ShaderDebugInfo")));
+		const FString OverrideShaderDebugDir = CVarShaderOverrideDebugDir.GetValueOnAnyThread();
+		if (!OverrideShaderDebugDir.IsEmpty())
+		{
+			AbsoluteDebugInfoDirectory = OverrideShaderDebugDir;
+		}
+
+		if (CVarDisambiguateShaderDebugDir.GetValueOnAnyThread())
+		{
+			FString AppendFolder = FPaths::ConvertRelativePathToFull(FPaths::ProjectDir());
+			FPaths::NormalizeDirectoryName(AppendFolder);
+			AppendFolder.ReplaceInline(TEXT(":/"), TEXT("_"));
+			AppendFolder.ReplaceCharInline('/', '_');
+			AbsoluteDebugInfoDirectory /= AppendFolder;
+		}
+
+		FPaths::NormalizeDirectoryName(AbsoluteDebugInfoDirectory);
+		DebugInfoPath = AbsoluteDebugInfoDirectory;
+	}
+	return DebugInfoPath;
+}
+
 bool GetShaderSymbolPathOverride(FString& OutPathOverride, FName ShaderFormat)
 {
 	static const FShaderSymbolSettingHelper SymbolPathOverride(TEXT("r.Shaders.SymbolPathOverride"));
@@ -694,6 +751,8 @@ bool GetShaderSymbolPathOverride(FString& OutPathOverride, FName ShaderFormat)
 			OutPathOverride = OutPathOverride.Replace(TEXT("{ProjectDir}"), *FPaths::ProjectDir(), ESearchCase::IgnoreCase);
 			// Allow the user to specify the location of the per-project saved folder string.
 			OutPathOverride = OutPathOverride.Replace(TEXT("{ProjectSavedDir}"), *FPaths::ProjectSavedDir(), ESearchCase::IgnoreCase);
+			// Allow the user to specify the configured shader debug info folder.
+			OutPathOverride = OutPathOverride.Replace(TEXT("{ShaderDebugInfoDir}"), *GetShaderDebugInfoPath(), ESearchCase::IgnoreCase);
 		}
 		return !OutPathOverride.IsEmpty();
 	}
diff --git a/Engine/Source/Runtime/RenderCore/Public/ShaderCompilerCore.h b/Engine/Source/Runtime/RenderCore/Public/ShaderCompilerCore.h
--- a/Engine/Source/Runtime/RenderCore/Public/ShaderCompilerCore.h
+++ b/Engine/Source/Runtime/RenderCore/Public/ShaderCompilerCore.h
@@ -46,6 +46,12 @@ namespace UE::ShaderCompiler
 	RENDERCORE_API ERHIBindlessConfiguration GetBindlessSamplersConfiguration(FName ShaderFormat);
 }
 
+/** Returns the path where shader compilation related artifacts should be stored when running on a build machine */
+extern RENDERCORE_API const FString& GetBuildMachineArtifactBasePath();
+
+/** Returns the base path where any shader debug information should be written to */
+extern RENDERCORE_API const FString& GetShaderDebugInfoPath();
+
 /** Returns true if shader symbols should be kept for a given platform. */
 extern RENDERCORE_API bool ShouldGenerateShaderSymbols(FName ShaderFormat);
 
