[HWRT] Remove "experimental" from r.RayTracing.EnableOnDemand description since it's now enabled by default.

diff --git a/Engine/Source/Runtime/RenderCore/Private/RenderUtils.cpp b/Engine/Source/Runtime/RenderCore/Private/RenderUtils.cpp
--- a/Engine/Source/Runtime/RenderCore/Private/RenderUtils.cpp
+++ b/Engine/Source/Runtime/RenderCore/Private/RenderUtils.cpp
@@ -46,7 +46,7 @@ static int32 GRayTracingEnableOnDemand = 1;
 static FAutoConsoleVariableRef CVarRayTracingEnableOnDemand(
 	TEXT("r.RayTracing.EnableOnDemand"),
 	GRayTracingEnableOnDemand,
-	TEXT("Controls whether ray tracing features can be toggled on demand at runtime without restarting the game (experimental).\n")
+	TEXT("Controls whether ray tracing features can be toggled on demand at runtime without restarting the game.\n")
 	TEXT("Requires r.RayTracing=1. Has a small performance and memory overhead.\n")
 	TEXT(" 0: off\n")
 	TEXT(" 1: on"),
