RHI: Missing usage flag in pooled index buffers for niagara ribbons.

#rnx

diff --git a/Engine/Source/Runtime/RenderCore/Private/GlobalRenderResources.cpp b/Engine/Source/Runtime/RenderCore/Private/GlobalRenderResources.cpp
--- a/Engine/Source/Runtime/RenderCore/Private/GlobalRenderResources.cpp
+++ b/Engine/Source/Runtime/RenderCore/Private/GlobalRenderResources.cpp
@@ -778,7 +778,9 @@ public:
 	{
 		check(!IsValidRef(BufferType::GetRHI()));
 		FRHIResourceCreateInfo CreateInfo(TEXT("FDynamicBuffer"));
-		BufferType::SetRHI(RHICmdList.CreateBuffer(BufferSize, Stride == 0 ? EBufferUsageFlags::VertexBuffer : EBufferUsageFlags::IndexBuffer, Stride, ERHIAccess::VertexOrIndexBuffer | ERHIAccess::SRVCompute, CreateInfo));
+		EBufferUsageFlags UsageFlags = (Stride == 0) ? EBufferUsageFlags::VertexBuffer : EBufferUsageFlags::IndexBuffer;
+		UsageFlags |= EBufferUsageFlags::ShaderResource;
+		BufferType::SetRHI(RHICmdList.CreateBuffer(BufferSize, UsageFlags, Stride, ERHIAccess::VertexOrIndexBuffer | ERHIAccess::SRVCompute, CreateInfo));
 		MappedBuffer = NULL;
 		AllocatedByteCount = 0;
 	}
