[Shaders] Add diagnostics to detect leaking references to shader map resources in ReleaseRHI().

Some plugins or other shader libraries are leaking memory and keep references to shader maps after they are released in the render thread.
This change adds diagnostics to find and report all shader libraries with their logical name (e.g. a plugin name) that still hold a reference to such a shader map resource.

#rnx

diff --git a/Engine/Source/Runtime/RenderCore/Private/ShaderCodeLibrary.cpp b/Engine/Source/Runtime/RenderCore/Private/ShaderCodeLibrary.cpp
--- a/Engine/Source/Runtime/RenderCore/Private/ShaderCodeLibrary.cpp
+++ b/Engine/Source/Runtime/RenderCore/Private/ShaderCodeLibrary.cpp
@@ -1426,6 +1426,8 @@ void FShaderMapResource_SharedCode::PreloadShaderMap(FGraphEventArray& OutComple
 	}
 }
 
+void DiagnoseLeakingShaderMapResourceRefs(const FShaderMapResource_SharedCode* InResource);
+
 void FShaderMapResource_SharedCode::ReleaseRHI()
 {
 	if (LibraryInstance && ensureMsgf(LibraryInstance->Library, TEXT("LibraryInstance->Library pointer is expected to be valid as long as library's FShaderMapResource are alive.")))
@@ -1476,6 +1478,9 @@ void FShaderMapResource_SharedCode::ReleaseRHI()
 
 	if (GetNumRefs() > 0)
 	{
+		// There shouldn't be any references left to this resource at this point. Diagnose what shader libraries are still holding a reference to this shader map resource.
+		DiagnoseLeakingShaderMapResourceRefs(this);
+
 		ensureMsgf(false, TEXT("FShaderMapResource_SharedCode::ReleaseRHI is still referenced (Num of references %d, owner %s). Invoking OnSharedShaderMapResourceExplicitRelease delegate."), GetNumRefs(), *GetOwnerName().ToString());
 		UE_LOG(LogShaderLibrary, Warning, TEXT("FShaderMapResource_SharedCode::ReleaseRHI is still referenced (Num of references %d, owner %s). Invoking OnSharedShaderMapResourceExplicitRelease delegate."), 
 			GetNumRefs(), *GetOwnerName().ToString());
@@ -2823,6 +2828,26 @@ public:
 #endif
 	}
 
+	void GatherShaderMapResourceReferences(const FShaderMapResource_SharedCode* InResource, TArray<FString>& OutLibrariesWithRefToResource) const
+	{
+		using namespace UE::ShaderLibrary::Private;
+
+		for (auto It = NamedLibrariesStack.CreateConstIterator(); It; ++It)
+		{
+			const FNamedShaderLibrary* Library = It.Value().Get();
+			for (const TUniquePtr<FShaderLibraryInstance>& LibraryInstance : Library->Components)
+			{
+				for (int32 LibraryResourceIndex = 0; LibraryResourceIndex < LibraryInstance->GetNumResources(); ++LibraryResourceIndex)
+				{
+					if (LibraryInstance->GetResource(LibraryResourceIndex).GetReference() == InResource)
+					{
+						OutLibrariesWithRefToResource.Add(It.Value()->LogicalName);
+					}
+				}
+			}
+		}
+	}
+
 	bool OpenLibrary(FString const& Name, FString const& Directory, const bool bMonolithicOnly = false)
 	{
 		TRACE_CPUPROFILER_EVENT_SCOPE(OpenShaderLibrary);
@@ -3674,6 +3699,26 @@ public:
 	}
 };
 
+// Used to diagnose leaking references to shader map resources in FShaderMapResource_SharedCode::ReleaseRHI()
+void DiagnoseLeakingShaderMapResourceRefs(const FShaderMapResource_SharedCode* InResource)
+{
+	TArray<FString> LibsWithRefToResource;
+	FShaderLibrariesCollection::Impl->GatherShaderMapResourceReferences(InResource, LibsWithRefToResource);
+	FString Info;
+	for (const FString& Name : LibsWithRefToResource)
+	{
+		if (!Info.IsEmpty())
+		{
+			Info += TEXT(", ");
+		}
+		Info += Name;
+	}
+	if (!Info.IsEmpty())
+	{
+		UE_LOG(LogShaderLibrary, Warning, TEXT("Leaking reference to shader map resource detected; Resource \"%s\" still referenced in: %s"), *InResource->GetFriendlyName(), *Info);
+	}
+}
+
 static FSharedShaderCodeRequest OnSharedShaderCodeRequest;
 
 FShaderLibrariesCollection* FShaderLibrariesCollection::Impl = nullptr;
diff --git a/Engine/Source/Runtime/RenderCore/Public/Shader.h b/Engine/Source/Runtime/RenderCore/Public/Shader.h
--- a/Engine/Source/Runtime/RenderCore/Public/Shader.h
+++ b/Engine/Source/Runtime/RenderCore/Public/Shader.h
@@ -342,7 +342,7 @@ public:
 	inline int32 GetNumRefs() const { return NumRefs.load(std::memory_order_relaxed); }
 
 	// FRenderResource interface.
-	RENDERCORE_API virtual void ReleaseRHI();
+	RENDERCORE_API virtual void ReleaseRHI() override;
 
 	inline int32 GetNumShaders() const
 	{
