Fixed render commands not inheriting LLM and mem scopes.

diff --git a/Engine/Source/Runtime/RenderCore/Private/RenderingThread.cpp b/Engine/Source/Runtime/RenderCore/Private/RenderingThread.cpp
--- a/Engine/Source/Runtime/RenderCore/Private/RenderingThread.cpp
+++ b/Engine/Source/Runtime/RenderCore/Private/RenderingThread.cpp
@@ -1543,9 +1543,9 @@ void FRenderThreadCommandPipe::EnqueueAndLaunch(const TCHAR* Name, uint32& SpecI
 			{
 				TRACE_CPUPROFILER_EVENT_SCOPE_USE_ON_CHANNEL(*Command.SpecId, Command.Name, EventScope, RenderCommandsChannel, true);
 				FScopeCycleCounter Scope(Command.StatId, true);
-				Command.Function(RHICmdList);
 
-				// Release the command immediately to match destruction order with task version.
+				UE::FInheritedContextScope InheritedContextScope = Command.RestoreInheritedContext();
+				Command.Function(RHICmdList);
 				Command.Function = {};
 			}
 			ConsumeCommands.Reset();
@@ -1984,7 +1984,9 @@ void FRenderCommandPipe::EnqueueAndLaunch(FFunctionVariant&& FunctionVariant, co
 	
 			for (FCommand& Command : PoppedQueue)
 			{
+				UE::FInheritedContextScope InheritedContextScope = Command.RestoreInheritedContext();
 				ExecuteCommand(MoveTemp(Command.Function), Command.Name, *Command.SpecId);
+				Command.Function = {};
 			}
 	
 			UE::RenderCommandPipe::ReplayingPipe = PreviousReplayingPipe;
diff --git a/Engine/Source/Runtime/RenderCore/Public/RenderingThread.h b/Engine/Source/Runtime/RenderCore/Public/RenderingThread.h
--- a/Engine/Source/Runtime/RenderCore/Public/RenderingThread.h
+++ b/Engine/Source/Runtime/RenderCore/Public/RenderingThread.h
@@ -344,14 +344,16 @@ private:
 
 	RENDERCORE_API void EnqueueAndLaunch(const TCHAR* Name, uint32& SpecId, TStatId StatId, TUniqueFunction<void(FRHICommandListImmediate&)>&& Function);
 
-	struct FCommand
+	struct FCommand : public UE::FInheritedContextBase
 	{
 		FCommand(const TCHAR* InName, uint32& OutSpecId, TStatId InStatId, TUniqueFunction<void(FRHICommandListImmediate&)>&& InFunction)
 			: Name(InName)
 			, SpecId(&OutSpecId)
 			, StatId(InStatId)
 			, Function(MoveTemp(InFunction))
-		{}
+		{
+			CaptureInheritedContext();
+		}
 
 		const TCHAR* Name;
 		uint32* SpecId;
@@ -526,13 +528,15 @@ private:
 		ExecuteCommand(FFunctionVariant(TInPlaceType<FEmptyFunction>(), MoveTemp(Function)), CommandName, CommandSpecId);
 	}
 
-	struct FCommand
+	struct FCommand : public UE::FInheritedContextBase
 	{
 		FCommand(FFunctionVariant&& InFunction, const TCHAR* InName, uint32& InOutSpecId)
 			: Function(MoveTemp(InFunction))
 			, Name(InName)
 			, SpecId(&InOutSpecId)
-		{}
+		{
+			CaptureInheritedContext();
+		}
 
 		FFunctionVariant Function;
 		const TCHAR* Name;
