[shaders] make addition of DebugHash_ comment in final stripped shader source optional (useful but maybe only to me when A/B testing shader symbol-related code changes)

diff --git a/Engine/Source/Runtime/RenderCore/Private/ShaderCore.cpp b/Engine/Source/Runtime/RenderCore/Private/ShaderCore.cpp
--- a/Engine/Source/Runtime/RenderCore/Private/ShaderCore.cpp
+++ b/Engine/Source/Runtime/RenderCore/Private/ShaderCore.cpp
@@ -1630,8 +1630,11 @@ public:
 				int32 NewlineLoc = SourceView.Find(SHADER_SOURCE_VIEWLITERAL("\n"), DebugHashLoc);
 				FShaderSource::TStringBuilder<2 * sizeof(FShaderCompilerInputHash::ByteArray) + 1> HashStr;
 				HashStr << Hash;
-				check(NewlineLoc - DebugHashLoc == HashStr.Len());
-				FMemory::Memcpy(Job.PreprocessOutput.EditSource().GetData() + DebugHashLoc, HashStr.GetData(), sizeof(FShaderSource::CharType) * HashStr.Len());
+				// handle case where we didn't add the debug hash association (this can be disabled by a cvar for debugging purposes)
+				if (NewlineLoc - DebugHashLoc == HashStr.Len())
+				{
+					FMemory::Memcpy(Job.PreprocessOutput.EditSource().GetData() + DebugHashLoc, HashStr.GetData(), sizeof(FShaderSource::CharType) * HashStr.Len());
+				}
 			}
 
 			// always compress the code after stripping to minimize memory footprint
diff --git a/Engine/Source/Runtime/RenderCore/Private/ShaderPreprocessTypes.cpp b/Engine/Source/Runtime/RenderCore/Private/ShaderPreprocessTypes.cpp
--- a/Engine/Source/Runtime/RenderCore/Private/ShaderPreprocessTypes.cpp
+++ b/Engine/Source/Runtime/RenderCore/Private/ShaderPreprocessTypes.cpp
@@ -6,6 +6,15 @@
 #include "Containers/AnsiString.h"
 #include "Templates/UnrealTemplate.h"
 
+static TAutoConsoleVariable<bool> CVarEnableDebugHashAssociation(
+	TEXT("r.ShaderCompiler.EnableDebugHashAssociation"),
+	true,
+	TEXT("If true, a // DebugHash_<hash> comment will be appended to all finalized shader source matching the DebugHash_<hash>.txt file output with "
+		 "shader debug artifacts. This can be used to find debug artifacts from shader source visible in Pix/RenderDoc/etc when symbols are enabled. "
+		 "Note that currently the only known reason to disable this is when comparing shader source or symbols across streams where versions may differ "
+		 "but source would otherwise match."),
+	ECVF_ReadOnly);
+
 struct FPreprocessConstants
 {
 	FPreprocessConstants()
@@ -14,7 +23,10 @@ struct FPreprocessConstants
 		// will be used to associate shader source back to debug output). This will be replaced inline
 		// by code calling into the preprocessing after the fact (necessary since the hash of the preprocessed
 		// code itself is used in constructing this debug hash).
-		StrippedCodeHeaderBuilder << DebugHashCommentStart << DebugHashPrefix << FShaderCompilerInputHash() << DebugHashCommentEnd;
+		if (CVarEnableDebugHashAssociation.GetValueOnAnyThread())
+		{
+			StrippedCodeHeaderBuilder << DebugHashCommentStart << DebugHashPrefix << FShaderCompilerInputHash() << DebugHashCommentEnd;
+		}
 		// Also add a sentinel line directive so we get the correct line offsets when remapping error/warning messages.
 		StrippedCodeHeaderBuilder << LineDirectiveStart << FilenameSentinel << LineDirectiveEnd;
 
@@ -49,21 +61,25 @@ private:
 	FShaderSource::TStringBuilder<HeaderLen> StrippedCodeHeaderBuilder;
 };
 
-static FPreprocessConstants GPreprocessConstants;
+static FPreprocessConstants& GetPreprocessConstants()
+{
+	static FPreprocessConstants PreprocessConstants;
+	return PreprocessConstants;
+}
 
 FShaderSource::FViewType GetShaderSourceDebugHashPrefix()
 {
-	return GPreprocessConstants.DebugHashPrefix;
+	return GetPreprocessConstants().DebugHashPrefix;
 }
 
 FStringView GetShaderSourceDebugHashPrefixWide()
 {
-	return GPreprocessConstants.DebugHashPrefixWide;
+	return GetPreprocessConstants().DebugHashPrefixWide;
 }
 
 void FShaderDiagnosticRemapper::Remap(FShaderCompilerError& Diagnostic) const
 {
-	if (!Diagnostic.ErrorLineString.IsEmpty() && !Diagnostic.ErrorVirtualFilePath.IsEmpty() && Diagnostic.ErrorVirtualFilePath == GPreprocessConstants.FilenameSentinel.GetData())
+	if (!Diagnostic.ErrorLineString.IsEmpty() && !Diagnostic.ErrorVirtualFilePath.IsEmpty() && Diagnostic.ErrorVirtualFilePath == GetPreprocessConstants().FilenameSentinel.GetData())
 	{
 		// only attempt remapping if the filename sentinel matches; this will bypass remapping in the case where compilation was retried on FXC
 		// with a DXC precompile (this compilation will report error messages relative to an intermediate hlsl file, which is not a change in behaviour)
@@ -95,7 +111,7 @@ void FShaderDiagnosticRemapper::Remap(FShaderCompilerError& Diagnostic) const
 		}
 	}
 
-	int32 FilenameIndex = Diagnostic.StrippedErrorMessage.Find(GPreprocessConstants.FilenameSentinel.GetData());
+	int32 FilenameIndex = Diagnostic.StrippedErrorMessage.Find(GetPreprocessConstants().FilenameSentinel.GetData());
 	// if we find our sentinel filename in the diagnostic message it needs remapping
 	// need to loop in case the message contains the filename multiple times
 	while (FilenameIndex >= 0)
@@ -103,7 +119,7 @@ void FShaderDiagnosticRemapper::Remap(FShaderCompilerError& Diagnostic) const
 		const FString& OriginalMessage = Diagnostic.StrippedErrorMessage;
 		// Parse the line number from the message and record start and end indices
 		// Code assumes that the next integral number found in the message string after the filename is always line number
-		int32 LineNumberStart = FilenameIndex + GPreprocessConstants.FilenameSentinel.Len();
+		int32 LineNumberStart = FilenameIndex + GetPreprocessConstants().FilenameSentinel.Len();
 		while (!FChar::IsDigit(OriginalMessage[LineNumberStart]))
 		{
 			++LineNumberStart;
@@ -126,18 +142,18 @@ void FShaderDiagnosticRemapper::Remap(FShaderCompilerError& Diagnostic) const
 
 		// Assume line number is the same number of digits for simplicity in reserve allocation size;
 		// it's an upper bound but in most practical cases worst case we're allocating an extra byte
-		int32 RemappedLen = OriginalMessage.Len() - GPreprocessConstants.FilenameSentinel.Len() + RemapData.Filename.Len();
+		int32 RemappedLen = OriginalMessage.Len() - GetPreprocessConstants().FilenameSentinel.Len() + RemapData.Filename.Len();
 		FString RemappedMessage;
 		RemappedMessage.Reserve(RemappedLen);
 		RemappedMessage.AppendChars(OriginalMessage.GetCharArray().GetData(), FilenameIndex);
 		RemappedMessage.Append(RemapData.Filename);
-		int32 SeparatorStart = FilenameIndex + GPreprocessConstants.FilenameSentinel.Len();
+		int32 SeparatorStart = FilenameIndex + GetPreprocessConstants().FilenameSentinel.Len();
 		RemappedMessage.Append(OriginalMessage.GetCharArray().GetData() + SeparatorStart, LineNumberStart - SeparatorStart);
 		RemappedMessage.AppendInt(RemapData.LineNumber);
 		RemappedMessage.Append(OriginalMessage.GetCharArray().GetData() + LineNumberEnd, OriginalMessage.Len() - LineNumberEnd);
 
 		Diagnostic.StrippedErrorMessage = MoveTemp(RemappedMessage);
-		FilenameIndex = Diagnostic.StrippedErrorMessage.Find(GPreprocessConstants.FilenameSentinel.GetData());
+		FilenameIndex = Diagnostic.StrippedErrorMessage.Find(GetPreprocessConstants().FilenameSentinel.GetData());
 	}
 }
 
@@ -226,9 +242,9 @@ inline void SkipNewLine(const FShaderSource::CharType*& Current, const FShaderSo
 void FShaderPreprocessOutput::StripCode(bool bCopyOriginalPreprocessdSource)
 {
 	// Reserve worst case slack (i.e. assuming there is nothing to strip) to avoid reallocation
-	FShaderSource PreprocessedSourceStripped(GPreprocessConstants.StrippedCodeHeader, PreprocessedSource.Len());
+	FShaderSource PreprocessedSourceStripped(GetPreprocessConstants().StrippedCodeHeader, PreprocessedSource.Len());
 	FShaderSource::CharType* OutStrippedData = PreprocessedSourceStripped.GetData();
-	FShaderSource::CharType* OutStripped = OutStrippedData + GPreprocessConstants.StrippedCodeHeader.Len();
+	FShaderSource::CharType* OutStripped = OutStrippedData + GetPreprocessConstants().StrippedCodeHeader.Len();
 
 	const FShaderSource::CharType* Begin = PreprocessedSource.GetData(), * Current = Begin;
 	const FShaderSource::CharType* End = Current + PreprocessedSource.Len();
