Deprecated GRDGDebug and other debug stuff that hasn't been useful or is unmaintained.

diff --git a/Engine/Source/Runtime/RenderCore/Private/RenderGraphBuilder.cpp b/Engine/Source/Runtime/RenderCore/Private/RenderGraphBuilder.cpp
--- a/Engine/Source/Runtime/RenderCore/Private/RenderGraphBuilder.cpp
+++ b/Engine/Source/Runtime/RenderCore/Private/RenderGraphBuilder.cpp
@@ -2591,7 +2591,7 @@ void FRDGBuilder::SetupAuxiliaryPasses(FRDGPass* Pass)
 
 FRDGPass* FRDGBuilder::SetupParameterPass(FRDGPass* Pass)
 {
-	IF_RDG_ENABLE_DEBUG(UserValidation.ValidateAddPass(Pass, AuxiliaryPasses.IsActive()));
+	IF_RDG_ENABLE_DEBUG(UserValidation.ValidateAddPass(Pass));
 	CSV_SCOPED_TIMING_STAT_EXCLUSIVE_CONDITIONAL(RDGBuilder_SetupPass, GRDGVerboseCSVStats != 0);
 
 #if RDG_EVENTS
@@ -2621,7 +2621,7 @@ FRDGPass* FRDGBuilder::SetupParameterPass(FRDGPass* Pass)
 
 FRDGPass* FRDGBuilder::SetupEmptyPass(FRDGPass* Pass)
 {
-	IF_RDG_ENABLE_DEBUG(UserValidation.ValidateAddPass(Pass, AuxiliaryPasses.IsActive()));
+	IF_RDG_ENABLE_DEBUG(UserValidation.ValidateAddPass(Pass));
 	CSV_SCOPED_TIMING_STAT_EXCLUSIVE_CONDITIONAL(RDGBuilder_SetupPass, GRDGVerboseCSVStats != 0);
 
 	Pass->bEmptyParameters = true;
diff --git a/Engine/Source/Runtime/RenderCore/Private/RenderGraphPrivate.cpp b/Engine/Source/Runtime/RenderCore/Private/RenderGraphPrivate.cpp
--- a/Engine/Source/Runtime/RenderCore/Private/RenderGraphPrivate.cpp
+++ b/Engine/Source/Runtime/RenderCore/Private/RenderGraphPrivate.cpp
@@ -29,16 +29,6 @@ FAutoConsoleVariableRef CVarRDGValidation(
 	TEXT(" 1: enabled (default);\n"),
 	ECVF_RenderThreadSafe);
 
-int32 GRDGDebug = 0;
-FAutoConsoleVariableRef CVarRDGDebug(
-	TEXT("r.RDG.Debug"),
-	GRDGDebug,
-	TEXT("Allow to output warnings for inefficiencies found during wiring and execution of the passes.\n")
-	TEXT(" 0: disabled;\n")
-	TEXT(" 1: emit warning once (default);\n")
-	TEXT(" 2: emit warning everytime issue is detected."),
-	ECVF_RenderThreadSafe);
-
 int32 GRDGDebugFlushGPU = 0;
 FAutoConsoleVariableRef CVarRDGDebugFlushGPU(
 	TEXT("r.RDG.Debug.FlushGPU"),
@@ -213,31 +203,6 @@ uint8 GetClobberStencil()
 	return 123;
 }
 
-void EmitRDGWarning(const FString& WarningMessage)
-{
-	if (!GRDGDebug)
-	{
-		return;
-	}
-
-	static TSet<FString> GAlreadyEmittedWarnings;
-
-	const int32 kRDGEmitWarningsOnce = 1;
-
-	if (GRDGDebug == kRDGEmitWarningsOnce)
-	{
-		if (!GAlreadyEmittedWarnings.Contains(WarningMessage))
-		{
-			GAlreadyEmittedWarnings.Add(WarningMessage);
-			UE_LOG(LogRDG, Warning, TEXT("%s"), *WarningMessage);
-		}
-	}
-	else
-	{
-		UE_LOG(LogRDG, Warning, TEXT("%s"), *WarningMessage);
-	}
-}
-
 bool GRDGAllowRHIAccess = false;
 bool GRDGAllowRHIAccessAsync = false;
 
@@ -523,11 +488,6 @@ void InitRenderGraph()
 		GRDGValidation = ValidationValue;
 	}
 
-	if (FParse::Param(FCommandLine::Get(), TEXT("rdgdebug")))
-	{
-		GRDGDebug = 1;
-	}
-
 	if (FParse::Param(FCommandLine::Get(), TEXT("rdgdebugextendresourcelifetimes")))
 	{
 		GRDGDebugExtendResourceLifetimes = 1;
@@ -626,7 +586,6 @@ bool IsParallelExecuteEnabled(EShaderPlatform ShaderPlatform)
 	return GRDGParallelExecute > 0
 		&& !GRHICommandList.Bypass()
 		&& !IsImmediateMode()
-		&& !GRDGDebug
 		&& !GRDGDebugFlushGPU
 		&& !GRDGTransitionLog
 		&& !IsMobilePlatform(ShaderPlatform)
@@ -646,7 +605,6 @@ bool IsParallelSetupEnabled(EShaderPlatform ShaderPlatform)
 	return GRDGParallelSetup > 0
 		&& !GRHICommandList.Bypass()
 		&& !IsImmediateMode()
-		&& !GRDGDebug
 		&& !GRDGTransitionLog
 		&& !IsMobilePlatform(ShaderPlatform)
 		&& !IsOpenGLPlatform(ShaderPlatform)
@@ -684,7 +642,7 @@ FRDGScopeState::FState::FState(bool bInImmediate, bool bInParallelExecute)
 		case ERDGScopeMode::TopLevelOnly:
 		case ERDGScopeMode::AllEvents:
 			// Override to a higher level in some cases
-			if (GRDGDebug != 0 || bRDGChannelEnabled != 0)
+			if (bRDGChannelEnabled != 0)
 			{
 				LocalScopeMode = ERDGScopeMode::AllEventsAndPassNames;
 			}
diff --git a/Engine/Source/Runtime/RenderCore/Private/RenderGraphPrivate.h b/Engine/Source/Runtime/RenderCore/Private/RenderGraphPrivate.h
--- a/Engine/Source/Runtime/RenderCore/Private/RenderGraphPrivate.h
+++ b/Engine/Source/Runtime/RenderCore/Private/RenderGraphPrivate.h
@@ -29,7 +29,6 @@ DEFINE_LOG_CATEGORY_STATIC(LogRDG, Log, All);
 extern int32 GRDGAsyncCompute;
 extern int32 GRDGClobberResources;
 extern int32 GRDGValidation;
-extern int32 GRDGDebug;
 extern int32 GRDGDebugFlushGPU;
 extern int32 GRDGDebugExtendResourceLifetimes;
 extern int32 GRDGDebugDisableTransientResources;
@@ -69,14 +68,10 @@ bool IsDebugAllowedForResource(const TCHAR* ResourceName);
 
 void EmitRDGWarning(const FString& WarningMessage);
 
-#define EmitRDGWarningf(WarningMessageFormat, ...) \
-	EmitRDGWarning(FString::Printf(WarningMessageFormat, ##__VA_ARGS__));
-
 #else // !RDG_ENABLE_DEBUG
 
 const int32 GRDGClobberResources = 0;
 const int32 GRDGValidation = 0;
-const int32 GRDGDebug = 0;
 const int32 GRDGDebugFlushGPU = 0;
 const int32 GRDGDebugExtendResourceLifetimes = 0;
 const int32 GRDGDebugDisableTransientResources = 0;
@@ -86,8 +81,6 @@ const int32 GRDGOverlapUAVs = 1;
 
 #define RDG_ALLOW_RHI_ACCESS_SCOPE()
 
-#define EmitRDGWarningf(WarningMessageFormat, ...)
-
 #endif
 
 extern int32 GRDGAsyncCompute;
diff --git a/Engine/Source/Runtime/RenderCore/Private/RenderGraphResources.cpp b/Engine/Source/Runtime/RenderCore/Private/RenderGraphResources.cpp
--- a/Engine/Source/Runtime/RenderCore/Private/RenderGraphResources.cpp
+++ b/Engine/Source/Runtime/RenderCore/Private/RenderGraphResources.cpp
@@ -54,6 +54,9 @@ FRDGViewableResource::FRDGViewableResource(const TCHAR* InName, const ERDGViewab
 	, bCollectForAllocate(1)
 #if RHI_USE_RESOURCE_DEBUG_NAME
 	, bHeapAllocatedDebugName(0)
+#endif
+#if RDG_ENABLE_DEBUG
+	, bClobbered(0)
 #endif
 	, TransientExtractionHint(ETransientExtractionHint::None)
 	, ReferenceCount(IsImmediateMode() ? 1 : 0)
diff --git a/Engine/Source/Runtime/RenderCore/Private/RenderGraphValidation.cpp b/Engine/Source/Runtime/RenderCore/Private/RenderGraphValidation.cpp
--- a/Engine/Source/Runtime/RenderCore/Private/RenderGraphValidation.cpp
+++ b/Engine/Source/Runtime/RenderCore/Private/RenderGraphValidation.cpp
@@ -21,15 +21,6 @@ bool GRDGBuilderActive = false;
 
 } //! namespace
 
-struct FRDGResourceDebugData
-{
-	/** Boolean to track at runtime whether a resource is actually used by the lambda of a pass or not, to detect unnecessary resource dependencies on passes. */
-	bool bIsActuallyUsedByPass = false;
-
-	/** Boolean to track at pass execution whether the underlying RHI resource is allowed to be accessed. */
-	bool bAllowRHIAccess = false;
-};
-
 void FRDGResource::MarkResourceAsUsed()
 {
 	if (!GRDGValidation)
@@ -38,8 +29,6 @@ void FRDGResource::MarkResourceAsUsed()
 	}
 
 	ValidateRHIAccess();
-
-	GetDebugData().bIsActuallyUsedByPass = true;
 }
 
 void FRDGResource::ValidateRHIAccess() const
@@ -49,68 +38,12 @@ void FRDGResource::ValidateRHIAccess() const
 		return;
 	}
 
-	check(DebugData);
-	checkf(DebugData->bAllowRHIAccess || GRDGAllowRHIAccess || GRDGAllowRHIAccessAsync,
+	checkf(bAllowRHIAccess || GRDGAllowRHIAccess || GRDGAllowRHIAccessAsync,
 		TEXT("Accessing the RHI resource of %s at this time is not allowed. If you hit this check in pass, ")
 		TEXT("that is due to this resource not being referenced in the parameters of your pass."),
 		Name);
 }
 
-FRDGResourceDebugData& FRDGResource::GetDebugData() const
-{
-	check(GRDGValidation != 0);
-	check(DebugData);
-	return *DebugData;
-}
-
-struct FRDGViewableResourceDebugData
-{
-	/** Pointer towards the pass that is the first to produce it, for even more convenient error message. */
-	const FRDGPass* FirstProducer = nullptr;
-
-	/** Count the number of times it has been used by a pass (without culling). */
-	uint32 PassAccessCount = 0;
-
-	/** Tracks whether this resource was clobbered by the builder prior to use. */
-	bool bHasBeenClobbered = false;
-};
-
-FRDGViewableResourceDebugData& FRDGViewableResource::GetViewableDebugData() const
-{
-	check(GRDGValidation != 0);
-	check(ViewableDebugData);
-	return *ViewableDebugData;
-}
-
-struct FRDGTextureDebugData
-{
-	/** Tracks whether a UAV has ever been allocated to catch when TexCreate_UAV was unneeded. */
-	bool bHasNeededUAV = false;
-
-	/** Tracks whether has ever been bound as a render target to catch when TexCreate_RenderTargetable was unneeded. */
-	bool bHasBeenBoundAsRenderTarget = false;
-};
-
-FRDGTextureDebugData& FRDGTexture::GetTextureDebugData() const
-{
-	check(GRDGValidation != 0);
-	check(TextureDebugData);
-	return *TextureDebugData;
-}
-
-struct FRDGBufferDebugData
-{
-	/** Tracks state changes in order of execution. */
-	TArray<TPair<FRDGPassHandle, FRDGSubresourceState>, FRDGArrayAllocator> States;
-};
-
-FRDGBufferDebugData& FRDGBuffer::GetBufferDebugData() const
-{
-	check(GRDGValidation != 0);
-	check(BufferDebugData);
-	return *BufferDebugData;
-}
-
 void FRDGUniformBuffer::MarkResourceAsUsed()
 {
 	if (!GRDGValidation)
@@ -150,7 +83,6 @@ void FRDGUserValidation::ExecuteGuard(const TCHAR* Operation, const TCHAR* Resou
 void FRDGUserValidation::ValidateCreateResource(FRDGResourceRef Resource)
 {
 	check(Resource);
-	Resource->DebugData = Allocator.Alloc<FRDGResourceDebugData>();
 
 	bool bAlreadyInSet = false;
 	ResourceMap.Emplace(Resource, &bAlreadyInSet);
@@ -160,7 +92,6 @@ void FRDGUserValidation::ValidateCreateResource(FRDGResourceRef Resource)
 void FRDGUserValidation::ValidateCreateViewableResource(FRDGViewableResource* Resource)
 {
 	ValidateCreateResource(Resource);
-	Resource->ViewableDebugData = Allocator.AllocNoDestruct<FRDGViewableResourceDebugData>();
 }
 
 void FRDGUserValidation::ValidateCreateTexture(FRDGTextureRef Texture)
@@ -171,11 +102,6 @@ void FRDGUserValidation::ValidateCreateTexture(FRDGTextureRef Texture)
 	}
 
 	ValidateCreateViewableResource(Texture);
-	Texture->TextureDebugData = Allocator.AllocNoDestruct<FRDGTextureDebugData>();
-	if (GRDGDebug)
-	{
-		TrackedTextures.Add(Texture);
-	}
 }
 
 void FRDGUserValidation::ValidateCreateBuffer(FRDGBufferRef Buffer)
@@ -188,12 +114,6 @@ void FRDGUserValidation::ValidateCreateBuffer(FRDGBufferRef Buffer)
 	ValidateCreateViewableResource(Buffer);
 
 	checkf(Buffer->Desc.GetSize() > 0 || Buffer->NumElementsCallback, TEXT("Creating buffer '%s' is zero bytes in size."), Buffer->Name);
-
-	Buffer->BufferDebugData = Allocator.Alloc<FRDGBufferDebugData>();
-	if (GRDGDebug)
-	{
-		TrackedBuffers.Add(Buffer);
-	}
 }
 
 void FRDGUserValidation::ValidateCreateSRV(FRDGTextureSRVRef SRV)
@@ -505,12 +425,6 @@ void FRDGUserValidation::ValidateExtractResource(FRDGViewableResource* Resource)
 	checkf(Resource->bProduced || Resource->bExternal || Resource->bQueuedForUpload,
 		TEXT("Unable to queue the extraction of the resource %s because it has not been produced by any pass."),
 		Resource->Name);
-
-	/** Increment pass access counts for externally registered buffers and textures to avoid
-	 *  emitting a 'produced but never used' warning. We don't have the history of registered
-	 *  resources to be able to emit a proper warning.
-	 */
-	Resource->GetViewableDebugData().PassAccessCount++;
 }
 
 void FRDGUserValidation::ValidateConvertToExternalResource(FRDGViewableResource* Resource)
@@ -539,23 +453,6 @@ void FRDGUserValidation::ValidateConvertToExternalUniformBuffer(FRDGUniformBuffe
 		UniformBuffer->GetLayoutName());
 }
 
-void FRDGUserValidation::RemoveUnusedWarning(FRDGViewableResource* Resource)
-{
-	if (!GRDGValidation)
-	{
-		return;
-	}
-
-	check(Resource);
-	ExecuteGuard(TEXT("RemoveUnusedResourceWarning"), Resource->Name);
-
-	// Removes 'produced but not used' warning.
-	Resource->GetViewableDebugData().PassAccessCount++;
-
-	// Removes 'not used' warning.
-	Resource->GetDebugData().bIsActuallyUsedByPass = true;
-}
-
 bool FRDGUserValidation::TryMarkForClobber(FRDGViewableResource* Resource) const
 {
 	if (!GRDGValidation)
@@ -564,13 +461,11 @@ bool FRDGUserValidation::TryMarkForClobber(FRDGViewableResource* Resource) const
 	}
 
 	check(Resource);
-	FRDGViewableResourceDebugData& DebugData = Resource->GetViewableDebugData();
-
-	const bool bClobber = !DebugData.bHasBeenClobbered && !Resource->bExternal && !Resource->bQueuedForUpload && IsDebugAllowedForResource(Resource->Name);
+	const bool bClobber = !Resource->bClobbered && !Resource->bExternal && !Resource->bQueuedForUpload && IsDebugAllowedForResource(Resource->Name);
 
 	if (bClobber)
 	{
-		DebugData.bHasBeenClobbered = true;
+		Resource->bClobbered = true;
 	}
 
 	return bClobber;
@@ -727,7 +622,7 @@ void FRDGUserValidation::ValidateAddPass(const void* ParameterStruct, const FSha
 		TEXT("Pass %s specified NeverMerge without Raster. Only raster passes support this flag."), Name.GetTCHAR());
 }
 
-void FRDGUserValidation::ValidateAddPass(const FRDGPass* Pass, bool bSkipPassAccessMarking)
+void FRDGUserValidation::ValidateAddPass(const FRDGPass* Pass)
 {
 	if (!GRDGValidation)
 	{
@@ -748,15 +643,7 @@ void FRDGUserValidation::ValidateAddPass(const FRDGPass* Pass, bool bSkipPassAcc
 
 	const auto MarkAsProduced = [&](FRDGViewableResource* Resource)
 	{
-		if (!bSkipPassAccessMarking)
-		{
-			auto& Debug = Resource->GetViewableDebugData();
-			if (!Debug.FirstProducer)
-			{
-				Debug.FirstProducer = Pass;
-			}
-			Debug.PassAccessCount++;
-		}
+		// Currently a no-op.
 	};
 
 	const auto MarkAsConsumed = [&] (FRDGViewableResource* Resource)
@@ -764,11 +651,6 @@ void FRDGUserValidation::ValidateAddPass(const FRDGPass* Pass, bool bSkipPassAcc
 		ensureMsgf(Resource->bProduced || Resource->bExternal || Resource->bQueuedForUpload,
 			TEXT("Pass %s has a read dependency on %s, but it was never written to."),
 			PassName, Resource->Name);
-
-		if (!bSkipPassAccessMarking)
-		{
-			Resource->GetViewableDebugData().PassAccessCount++;
-		}
 	};
 
 	const auto CheckValidResource = [&](FRDGResourceRef Resource)
@@ -797,7 +679,6 @@ void FRDGUserValidation::ValidateAddPass(const FRDGPass* Pass, bool bSkipPassAcc
 
 		if (IsWritableAccess(Access))
 		{
-			MarkAsProduced(Buffer);
 			bCanProduce = true;
 		}
 	};
@@ -808,7 +689,6 @@ void FRDGUserValidation::ValidateAddPass(const FRDGPass* Pass, bool bSkipPassAcc
 
 		if (IsWritableAccess(Access))
 		{
-			MarkAsProduced(Texture);
 			bCanProduce = true;
 		}
 	};
@@ -1049,56 +929,6 @@ void FRDGUserValidation::ValidateExecuteEnd()
 
 	bHasExecuted = true;
 	GRDGBuilderActive = false;
-
-	if (GRDGDebug && GRDGValidation)
-	{
-		auto ValidateResourceAtExecuteEnd = [](const FRDGViewableResource* Resource)
-		{
-			const auto& ViewableDebugData = Resource->GetViewableDebugData();
-			const bool bProducedButNeverUsed = ViewableDebugData.PassAccessCount == 1 && ViewableDebugData.FirstProducer;
-
-			if (bProducedButNeverUsed)
-			{
-				check(Resource->bProduced || Resource->bExternal || Resource->bExtracted);
-
-				EmitRDGWarningf(
-					TEXT("Resource %s has been produced by the pass %s, but never used by another pass."),
-					Resource->Name, ViewableDebugData.FirstProducer->GetName());
-			}
-		};
-
-		for (const FRDGTextureRef Texture : TrackedTextures)
-		{
-			ValidateResourceAtExecuteEnd(Texture);
-
-			const auto& ViewableDebugData = Texture->GetViewableDebugData();
-			const auto& TextureDebugData = Texture->GetTextureDebugData();
-
-			const bool bHasBeenProducedByGraph = !Texture->bExternal && ViewableDebugData.PassAccessCount > 0;
-
-			if (bHasBeenProducedByGraph && !TextureDebugData.bHasNeededUAV && EnumHasAnyFlags(Texture->Desc.Flags, TexCreate_UAV))
-			{
-				EmitRDGWarningf(
-					TEXT("Resource %s first produced by the pass %s had the TexCreate_UAV flag, but no UAV has been used."),
-					Texture->Name, ViewableDebugData.FirstProducer->GetName());
-			}
-
-			if (bHasBeenProducedByGraph && !TextureDebugData.bHasBeenBoundAsRenderTarget && EnumHasAnyFlags(Texture->Desc.Flags, TexCreate_RenderTargetable))
-			{
-				EmitRDGWarningf(
-					TEXT("Resource %s first produced by the pass %s had the TexCreate_RenderTargetable flag, but has never been bound as a render target of a pass."),
-					Texture->Name, ViewableDebugData.FirstProducer->GetName());
-			}
-		}
-
-		for (const FRDGBufferRef Buffer : TrackedBuffers)
-		{
-			ValidateResourceAtExecuteEnd(Buffer);
-		}
-	}
-
-	TrackedTextures.Empty();
-	TrackedBuffers.Empty();
 }
 
 void FRDGUserValidation::ValidateExecutePassBegin(const FRDGPass* Pass)
@@ -1109,96 +939,6 @@ void FRDGUserValidation::ValidateExecutePassBegin(const FRDGPass* Pass)
 	}
 
 	SetAllowRHIAccess(Pass, true);
-
-	if (GRDGDebug)
-	{
-		Pass->GetParameters().EnumerateUniformBuffers([&](FRDGUniformBufferBinding UniformBuffer)
-		{
-			// Global uniform buffers are always marked as used, because FShader traversal doesn't know about them.
-			if (UniformBuffer.IsStatic())
-			{
-				UniformBuffer->MarkResourceAsUsed();
-			}
-		});
-
-		const auto ValidateTextureAccess = [](FRDGTextureRef Texture, ERHIAccess Access)
-		{
-			if (EnumHasAnyFlags(Access, ERHIAccess::UAVMask))
-			{
-				Texture->GetTextureDebugData().bHasNeededUAV = true;
-			}
-			if (EnumHasAnyFlags(Access, ERHIAccess::RTV | ERHIAccess::DSVRead | ERHIAccess::DSVWrite))
-			{
-				Texture->GetTextureDebugData().bHasBeenBoundAsRenderTarget = true;
-			}
-			Texture->MarkResourceAsUsed();
-		};
-
-		Pass->GetParameters().Enumerate([&](FRDGParameter Parameter)
-		{
-			switch (Parameter.GetType())
-			{
-			case UBMT_RDG_TEXTURE_UAV:
-				if (FRDGTextureUAVRef UAV = Parameter.GetAsTextureUAV())
-				{
-					FRDGTextureRef Texture = UAV->Desc.Texture;
-					Texture->GetTextureDebugData().bHasNeededUAV = true;
-				}
-			break;
-			case UBMT_RDG_TEXTURE_ACCESS:
-				if (FRDGTextureAccess TextureAccess = Parameter.GetAsTextureAccess())
-				{
-					ValidateTextureAccess(TextureAccess.GetTexture(), TextureAccess.GetAccess());
-				}
-			break;
-			case UBMT_RDG_TEXTURE_ACCESS_ARRAY:
-			{
-				const FRDGTextureAccessArray& TextureAccessArray = Parameter.GetAsTextureAccessArray();
-
-				for (FRDGTextureAccess TextureAccess : TextureAccessArray)
-				{
-					ValidateTextureAccess(TextureAccess.GetTexture(), TextureAccess.GetAccess());
-				}
-			}
-			break;
-			case UBMT_RDG_BUFFER_ACCESS:
-				if (FRDGBufferRef Buffer = Parameter.GetAsBuffer())
-				{
-					Buffer->MarkResourceAsUsed();
-				}
-			break;
-			case UBMT_RDG_BUFFER_ACCESS_ARRAY:
-				for (FRDGBufferAccess BufferAccess : Parameter.GetAsBufferAccessArray())
-				{
-					BufferAccess->MarkResourceAsUsed();
-				}
-			break;
-			case UBMT_RENDER_TARGET_BINDING_SLOTS:
-			{
-				const FRenderTargetBindingSlots& RenderTargets = Parameter.GetAsRenderTargetBindingSlots();
-
-				RenderTargets.Enumerate([&](FRenderTargetBinding RenderTarget)
-				{
-					FRDGTextureRef Texture = RenderTarget.GetTexture();
-					Texture->GetTextureDebugData().bHasBeenBoundAsRenderTarget = true;
-					Texture->MarkResourceAsUsed();
-				});
-
-				if (FRDGTextureRef Texture = RenderTargets.DepthStencil.GetTexture())
-				{
-					Texture->GetTextureDebugData().bHasBeenBoundAsRenderTarget = true;
-					Texture->MarkResourceAsUsed();
-				}
-
-				if (FRDGTextureRef Texture = RenderTargets.ShadingRateTexture)
-				{
-					Texture->MarkResourceAsUsed();
-				}
-			}
-			break;
-			}
-		});
-	}
 }
 
 void FRDGUserValidation::ValidateExecutePassEnd(const FRDGPass* Pass)
@@ -1209,60 +949,6 @@ void FRDGUserValidation::ValidateExecutePassEnd(const FRDGPass* Pass)
 	}
 
 	SetAllowRHIAccess(Pass, false);
-
-	const FRDGParameterStruct PassParameters = Pass->GetParameters();
-
-	if (GRDGDebug)
-	{
-		uint32 TrackedResourceCount = 0;
-		uint32 UsedResourceCount = 0;
-
-		PassParameters.Enumerate([&](FRDGParameter Parameter)
-		{
-			if (Parameter.IsResource())
-			{
-				if (FRDGResourceRef Resource = Parameter.GetAsResource())
-				{
-					TrackedResourceCount++;
-					UsedResourceCount += Resource->GetDebugData().bIsActuallyUsedByPass ? 1 : 0;
-				}
-			}
-		});
-
-		if (TrackedResourceCount != UsedResourceCount)
-		{
-			FString WarningMessage = FString::Printf(
-				TEXT("'%d' of the '%d' resources of the pass '%s' were not actually used."),
-				TrackedResourceCount - UsedResourceCount, TrackedResourceCount, Pass->GetName());
-
-			PassParameters.Enumerate([&](FRDGParameter Parameter)
-			{
-				if (Parameter.IsResource())
-				{
-					if (const FRDGResourceRef Resource = Parameter.GetAsResource())
-					{
-						if (!Resource->GetDebugData().bIsActuallyUsedByPass)
-						{
-							WarningMessage += FString::Printf(TEXT("\n    %s"), Resource->Name);
-						}
-					}
-				}
-			});
-
-			EmitRDGWarning(WarningMessage);
-		}
-	}
-
-	PassParameters.Enumerate([&](FRDGParameter Parameter)
-	{
-		if (Parameter.IsResource())
-		{
-			if (FRDGResourceRef Resource = Parameter.GetAsResource())
-			{
-				Resource->GetDebugData().bIsActuallyUsedByPass = false;
-			}
-		}
-	});
 }
 
 void FRDGUserValidation::SetAllowRHIAccess(const FRDGPass* Pass, bool bAllowAccess)
@@ -1278,21 +964,21 @@ void FRDGUserValidation::SetAllowRHIAccess(const FRDGPass* Pass, bool bAllowAcce
 		{
 			if (FRDGResourceRef Resource = Parameter.GetAsResource())
 			{
-				Resource->GetDebugData().bAllowRHIAccess = bAllowAccess;
+				Resource->bAllowRHIAccess = bAllowAccess;
 			}
 		}
 		else if (Parameter.IsBufferAccessArray())
 		{
 			for (FRDGBufferAccess BufferAccess : Parameter.GetAsBufferAccessArray())
 			{
-				BufferAccess->GetDebugData().bAllowRHIAccess = bAllowAccess;
+				BufferAccess->bAllowRHIAccess = bAllowAccess;
 			}
 		}
 		else if (Parameter.IsTextureAccessArray())
 		{
 			for (FRDGTextureAccess TextureAccess : Parameter.GetAsTextureAccessArray())
 			{
-				TextureAccess->GetDebugData().bAllowRHIAccess = bAllowAccess;
+				TextureAccess->bAllowRHIAccess = bAllowAccess;
 			}
 		}
 		else if (Parameter.IsRenderTargetBindingSlots())
@@ -1301,27 +987,27 @@ void FRDGUserValidation::SetAllowRHIAccess(const FRDGPass* Pass, bool bAllowAcce
 
 			RenderTargets.Enumerate([&](FRenderTargetBinding RenderTarget)
 			{
-				RenderTarget.GetTexture()->GetDebugData().bAllowRHIAccess = bAllowAccess;
+				RenderTarget.GetTexture()->bAllowRHIAccess = bAllowAccess;
 
 				if (FRDGTextureRef ResolveTexture = RenderTarget.GetResolveTexture())
 				{
-					ResolveTexture->GetDebugData().bAllowRHIAccess = bAllowAccess;
+					ResolveTexture->bAllowRHIAccess = bAllowAccess;
 				}
 			});
 
 			if (FRDGTextureRef Texture = RenderTargets.DepthStencil.GetTexture())
 			{
-				Texture->GetDebugData().bAllowRHIAccess = bAllowAccess;
+				Texture->bAllowRHIAccess = bAllowAccess;
 			}
 
 			if (FRDGTextureRef ResolveTexture = RenderTargets.DepthStencil.GetResolveTexture())
 			{
-				ResolveTexture->GetDebugData().bAllowRHIAccess = bAllowAccess;
+				ResolveTexture->bAllowRHIAccess = bAllowAccess;
 			}
 
 			if (FRDGTexture* Texture = RenderTargets.ShadingRateTexture)
 			{
-				Texture->GetDebugData().bAllowRHIAccess = bAllowAccess;
+				Texture->bAllowRHIAccess = bAllowAccess;
 			}
 		}
 	});
diff --git a/Engine/Source/Runtime/RenderCore/Public/RenderGraphBuilder.h b/Engine/Source/Runtime/RenderCore/Public/RenderGraphBuilder.h
--- a/Engine/Source/Runtime/RenderCore/Public/RenderGraphBuilder.h
+++ b/Engine/Source/Runtime/RenderCore/Public/RenderGraphBuilder.h
@@ -414,10 +414,6 @@ public:
 		}
 	}
 
-	/** Flag a resource that is produced by a pass but never used or extracted to not emit an 'unused' warning. */
-	void RemoveUnusedTextureWarning(FRDGTextureRef Texture);
-	void RemoveUnusedBufferWarning(FRDGBufferRef Buffer);
-
 	/** Flushes all queued passes to an async task to perform setup work. */
 	RENDERCORE_API void FlushSetupQueue();
 
@@ -452,12 +448,6 @@ public:
 	static bool IsDumpingFrame() { return false; }
 #endif
 
-	UE_DEPRECATED(5.5, "This path is no longer supported.")
-	static void DumpDraw(const FRDGEventName& DrawEventName) {}
-
-	UE_DEPRECATED(5.5, "This path is no longer supported.")
-	static bool IsDumpingDraws() { return false; }
-
 #if WITH_MGPU
 	/** Copy all cross GPU external resources (not marked MultiGPUGraphIgnore) at the end of execution (bad for perf, but useful for debugging). */
 	void EnableForceCopyCrossGPU()
@@ -466,6 +456,18 @@ public:
 	}
 #endif
 
+	UE_DEPRECATED(5.5, "This path is no longer supported.")
+	static void DumpDraw(const FRDGEventName& DrawEventName) {}
+
+	UE_DEPRECATED(5.5, "This path is no longer supported.")
+	static bool IsDumpingDraws() { return false; }
+
+	UE_DEPRECATED(5.6, "RemoveUnusedTextureWarning is no longer necessary.")
+	void RemoveUnusedTextureWarning(FRDGTextureRef Texture) {}
+
+	UE_DEPRECATED(5.6, "RemoveUnusedBufferWarning is no longer necessary.")
+	void RemoveUnusedBufferWarning(FRDGBufferRef Buffer) {}
+
 private:
 	static const char* const kDefaultUnaccountedCSVStat;
 
diff --git a/Engine/Source/Runtime/RenderCore/Public/RenderGraphBuilder.inl b/Engine/Source/Runtime/RenderCore/Public/RenderGraphBuilder.inl
--- a/Engine/Source/Runtime/RenderCore/Public/RenderGraphBuilder.inl
+++ b/Engine/Source/Runtime/RenderCore/Public/RenderGraphBuilder.inl
@@ -744,14 +744,4 @@ inline void FRDGBuilder::SetBufferAccessFinal(FRDGBufferRef Buffer, ERHIAccess A
 {
 	IF_RDG_ENABLE_DEBUG(UserValidation.ValidateSetAccessFinal(Buffer, AccessFinal));
 	Buffer->EpilogueAccess = AccessFinal;
-}
-
-inline void FRDGBuilder::RemoveUnusedTextureWarning(FRDGTextureRef Texture)
-{
-	IF_RDG_ENABLE_DEBUG(UserValidation.RemoveUnusedWarning(Texture));
-}
-
-inline void FRDGBuilder::RemoveUnusedBufferWarning(FRDGBufferRef Buffer)
-{
-	IF_RDG_ENABLE_DEBUG(UserValidation.RemoveUnusedWarning(Buffer));
-}
+}
\ No newline at end of file
diff --git a/Engine/Source/Runtime/RenderCore/Public/RenderGraphResources.h b/Engine/Source/Runtime/RenderCore/Public/RenderGraphResources.h
--- a/Engine/Source/Runtime/RenderCore/Public/RenderGraphResources.h
+++ b/Engine/Source/Runtime/RenderCore/Public/RenderGraphResources.h
@@ -37,10 +37,6 @@ class FRHITransientTexture;
 class FRenderTargetPool;
 class FShaderParametersMetadata;
 struct FPooledRenderTarget;
-struct FRDGBufferDebugData;
-struct FRDGResourceDebugData;
-struct FRDGTextureDebugData;
-struct FRDGViewableResourceDebugData;
 
 /** Used for tracking pass producer / consumer edges in the graph for culling and pipe fencing. */
 struct FRDGProducerState
@@ -144,6 +140,9 @@ public:
 protected:
 	FRDGResource(const TCHAR* InName)
 		: Name(InName)
+#if RDG_ENABLE_DEBUG
+		, bAllowRHIAccess(0)
+#endif
 	{}
 
 	FRHIResource* GetRHIUnchecked() const
@@ -164,8 +163,7 @@ protected:
 
 private:
 #if RDG_ENABLE_DEBUG
-	struct FRDGResourceDebugData* DebugData = nullptr;
-	RENDERCORE_API FRDGResourceDebugData& GetDebugData() const;
+	uint8 bAllowRHIAccess : 1;
 #endif
 
 	friend FRDGBuilder;
@@ -408,6 +406,10 @@ protected:
 	uint8 bHeapAllocatedDebugName : 1;
 #endif
 
+#if RDG_ENABLE_DEBUG
+	uint8 bClobbered : 1;
+#endif
+
 	/** Whether this resource is allowed to be both transient and extracted. */
 	ETransientExtractionHint TransientExtractionHint;
 
@@ -461,11 +463,6 @@ private:
 	TArray<FRDGPassHandle, FRDGArrayAllocator> TracePasses;
 #endif
 
-#if RDG_ENABLE_DEBUG
-	struct FRDGViewableResourceDebugData* ViewableDebugData = nullptr;
-	RENDERCORE_API FRDGViewableResourceDebugData& GetViewableDebugData() const;
-#endif
-
 	friend bool IsExtendedLifetimeResource(FRDGViewableResource*);
 
 	friend FRDGBuilder;
@@ -652,11 +649,6 @@ private:
 	/** Valid strictly when holding a strong reference; use PooledRenderTarget instead. */
 	TRefCountPtr<IPooledRenderTarget> Allocation;
 
-#if RDG_ENABLE_DEBUG
-	struct FRDGTextureDebugData* TextureDebugData = nullptr;
-	RENDERCORE_API FRDGTextureDebugData& GetTextureDebugData() const;
-#endif
-
 #if SUPPORTS_VISUALIZE_TEXTURE
 	FIntPoint VisualizeTextureExtent = { 0, 0 };
 public:
@@ -1406,11 +1398,6 @@ private:
 	/** Tracks the merged subresource state as the graph is built. */
 	FRDGSubresourceState* MergeState = nullptr;
 
-#if RDG_ENABLE_DEBUG
-	struct FRDGBufferDebugData* BufferDebugData = nullptr;
-	RENDERCORE_API FRDGBufferDebugData& GetBufferDebugData() const;
-#endif
-
 	friend FRDGBuilder;
 	friend FRDGBarrierValidation;
 	friend FRDGUserValidation;
diff --git a/Engine/Source/Runtime/RenderCore/Public/RenderGraphValidation.h b/Engine/Source/Runtime/RenderCore/Public/RenderGraphValidation.h
--- a/Engine/Source/Runtime/RenderCore/Public/RenderGraphValidation.h
+++ b/Engine/Source/Runtime/RenderCore/Public/RenderGraphValidation.h
@@ -97,7 +97,7 @@ public:
 	 */
 	RENDERCORE_API void ValidateAddPass(const void* ParameterStruct, const FShaderParametersMetadata* Metadata, const FRDGEventName& Name, ERDGPassFlags Flags);
 	RENDERCORE_API void ValidateAddPass(const FRDGEventName& Name, ERDGPassFlags Flags);
-	RENDERCORE_API void ValidateAddPass(const FRDGPass* Pass, bool bSkipPassAccessMarking);
+	RENDERCORE_API void ValidateAddPass(const FRDGPass* Pass);
 
 	/** Validate pass state before and after execution. */
 	RENDERCORE_API void ValidateExecutePassBegin(const FRDGPass* Pass);
@@ -107,9 +107,6 @@ public:
 	RENDERCORE_API void ValidateExecuteBegin();
 	RENDERCORE_API void ValidateExecuteEnd();
 
-	/** Removes the 'produced but not used' warning from the requested resource. */
-	RENDERCORE_API void RemoveUnusedWarning(FRDGViewableResource* Resource);
-
 	/** Attempts to mark a resource for clobbering. If already marked, returns false.  */
 	RENDERCORE_API bool TryMarkForClobber(FRDGViewableResource* Resource) const;
 
@@ -140,10 +137,6 @@ private:
 
 	FRDGAllocator& Allocator;
 
-	/** List of tracked resources for validation prior to shutdown. */
-	TArray<FRDGTextureRef, FRDGArrayAllocator> TrackedTextures;
-	TArray<FRDGBufferRef, FRDGArrayAllocator> TrackedBuffers;
-
 	/** Map tracking all active resources in the graph. */
 	TSet<FRDGResourceRef, DefaultKeyFuncs<FRDGResourceRef>, FRDGSetAllocator> ResourceMap;
 
diff --git a/Engine/Source/Runtime/Renderer/Private/PostProcess/PostProcessMaterial.cpp b/Engine/Source/Runtime/Renderer/Private/PostProcess/PostProcessMaterial.cpp
--- a/Engine/Source/Runtime/Renderer/Private/PostProcess/PostProcessMaterial.cpp
+++ b/Engine/Source/Runtime/Renderer/Private/PostProcess/PostProcessMaterial.cpp
@@ -480,9 +480,6 @@ FPostProcessMaterialParameters* GetPostProcessMaterialParameters(
 
 	const FScreenPassTexture BlackDummy(GSystemTextures.GetBlackDummy(GraphBuilder));
 
-	// This gets passed in whether or not it's used.
-	GraphBuilder.RemoveUnusedTextureWarning(BlackDummy.Texture);
-
 	FRHISamplerState* PointClampSampler = TStaticSamplerState<SF_Point, AM_Clamp, AM_Clamp, AM_Clamp>::GetRHI();
 
 	int32 NumUserSceneTextures = MaterialShaderMap->GetUserSceneTextureInputs().Num();
diff --git a/Engine/Source/Runtime/Renderer/Private/PostProcess/PostProcessMotionBlur.cpp b/Engine/Source/Runtime/Renderer/Private/PostProcess/PostProcessMotionBlur.cpp
--- a/Engine/Source/Runtime/Renderer/Private/PostProcess/PostProcessMotionBlur.cpp
+++ b/Engine/Source/Runtime/Renderer/Private/PostProcess/PostProcessMotionBlur.cpp
@@ -1095,7 +1095,6 @@ FMotionBlurOutputs AddMotionBlurFilterPass(
 
 			FRDGTextureRef DummyTexture = GraphBuilder.CreateTexture(Desc, TEXT("MotionBlur.DummyOutput"));
 			OriginalPassParameters.SceneColorOutputMip1 = GraphBuilder.CreateUAV(DummyTexture, ERDGUnorderedAccessViewFlags::SkipBarrier);
-			GraphBuilder.RemoveUnusedTextureWarning(DummyTexture);
 		}
 
 		if (OriginalPassParameters.OutputMip2)
@@ -1121,7 +1120,6 @@ FMotionBlurOutputs AddMotionBlurFilterPass(
 
 			FRDGTextureRef DummyTexture = GraphBuilder.CreateTexture(Desc, TEXT("MotionBlur.DummyOutput"));
 			OriginalPassParameters.SceneColorOutputMip2 = GraphBuilder.CreateUAV(DummyTexture, ERDGUnorderedAccessViewFlags::SkipBarrier);
-			GraphBuilder.RemoveUnusedTextureWarning(DummyTexture);
 		}
 
 		OriginalPassParameters.DebugOutput = CreateDebugUAV(GraphBuilder, Output.FullRes.TextureSRV->Desc.Texture->Desc.Extent, TEXT("Debug.MotionBlur.Filter"));
diff --git a/Engine/Source/Runtime/Renderer/Private/PostProcess/TemporalSuperResolution.cpp b/Engine/Source/Runtime/Renderer/Private/PostProcess/TemporalSuperResolution.cpp
--- a/Engine/Source/Runtime/Renderer/Private/PostProcess/TemporalSuperResolution.cpp
+++ b/Engine/Source/Runtime/Renderer/Private/PostProcess/TemporalSuperResolution.cpp
@@ -1273,7 +1273,6 @@ static FRDGTextureUAVRef CreateDummyUAV(FRDGBuilder& GraphBuilder, EPixelFormat
 		/* InFlags = */ TexCreate_ShaderResource | TexCreate_UAV);
 
 	FRDGTextureRef DummyTexture = GraphBuilder.CreateTexture(Desc, TEXT("TSR.DummyOutput"));
-	GraphBuilder.RemoveUnusedTextureWarning(DummyTexture);
 
 	return GraphBuilder.CreateUAV(DummyTexture);
 };
@@ -1288,7 +1287,6 @@ static FRDGTextureUAVRef CreateDummyUAVArray(FRDGBuilder& GraphBuilder, EPixelFo
 		/* ArraySize = */ 1);
 
 	FRDGTextureRef DummyTexture = GraphBuilder.CreateTexture(Desc, TEXT("TSR.DummyOutput"));
-	GraphBuilder.RemoveUnusedTextureWarning(DummyTexture);
 
 	return GraphBuilder.CreateUAV(DummyTexture);
 };
@@ -2541,7 +2539,6 @@ FDefaultTemporalUpscaler::FOutputs AddTemporalSuperResolutionPasses(
 
 				// Create an unused texture for the moire history so that the VisualizeTSR can still display the updated moire history.
 				FRDGTextureRef UnusedMoireHistoryTexture = GraphBuilder.CreateTexture(Desc, TEXT("TSR.History.Moire"));
-				GraphBuilder.RemoveUnusedTextureWarning(UnusedMoireHistoryTexture);
 
 				PassParameters->HistoryMoireOutput = GraphBuilder.CreateUAV(UnusedMoireHistoryTexture);
 				MoireHistoryTexture = GraphBuilder.CreateSRV(FRDGTextureSRVDesc::CreateForSlice(UnusedMoireHistoryTexture, 0));
diff --git a/Engine/Source/Runtime/Renderer/Private/ScreenSpaceDenoise.cpp b/Engine/Source/Runtime/Renderer/Private/ScreenSpaceDenoise.cpp
--- a/Engine/Source/Runtime/Renderer/Private/ScreenSpaceDenoise.cpp
+++ b/Engine/Source/Runtime/Renderer/Private/ScreenSpaceDenoise.cpp
@@ -2252,16 +2252,6 @@ static void DenoiseSignalAtConstantPixelDensity(
 			}
 		} // for (uint32 BatchedSignalId = 0; BatchedSignalId < Settings.SignalBatchSize; BatchedSignalId++)
 	}
-	else if (HistoryTextureCountPerSignal >= 2)
-	{
-		// The SignalHistory1 is always generated for temporal history, but will endup useless if there is no view state,
-		// in witch case we do not extract any textures. Don't support a shader permutation that does not produce it, because
-		// it is already a not ideal case for the denoiser.
-		for (int32 BufferId = 1; BufferId < HistoryTextureCountPerSignal; BufferId++)
-		{
-			GraphBuilder.RemoveUnusedTextureWarning(SignalHistory.Textures[BufferId]);
-		}
-	}
 
 	// Final convolution / output to correct
 	if (SignalUsesFinalConvolution(Settings.SignalProcessing))
@@ -2695,10 +2685,6 @@ public:
 		RDG_EVENT_SCOPE_STAT(GraphBuilder, ReflectionsDenoiser, "DenoiseReflections");
 		RDG_GPU_STAT_SCOPE(GraphBuilder, ReflectionsDenoiser);
 
-		// Imaginary depth is only used for Nvidia denoiser.
-		// TODO(Denoiser): permutation to not generate it?
-		GraphBuilder.RemoveUnusedTextureWarning(ReflectionInputs.RayImaginaryDepth);
-
 		FViewInfoPooledRenderTargets ViewInfoPooledRenderTargets;
 		SetupSceneViewInfoPooledRenderTargets(View, &ViewInfoPooledRenderTargets);
 
@@ -2745,11 +2731,6 @@ public:
 	{
 		RDG_EVENT_SCOPE_STAT(GraphBuilder, ReflectionsDenoiser, "DenoiseWaterReflections");
 		RDG_GPU_STAT_SCOPE(GraphBuilder, ReflectionsDenoiser);
-
-		// Imaginary depth is only used for Nvidia denoiser.
-		// TODO: permutation to not generate it?
-		if (ReflectionInputs.RayImaginaryDepth)
-			GraphBuilder.RemoveUnusedTextureWarning(ReflectionInputs.RayImaginaryDepth);
 		
 		FViewInfoPooledRenderTargets ViewInfoPooledRenderTargets;
 		SetupSceneViewInfoPooledRenderTargets(View, &ViewInfoPooledRenderTargets);
