[shaders] call NotifyShaderCompiled on job completion and when shadermaps are loaded from DDC (in addition to when they are saved) in the editor. This has the effect of shader symbols being written out to disk (when enabled) when using ODSC/running in the editor.

diff --git a/Engine/Source/Runtime/RenderCore/Private/ShaderCore.cpp b/Engine/Source/Runtime/RenderCore/Private/ShaderCore.cpp
--- a/Engine/Source/Runtime/RenderCore/Private/ShaderCore.cpp
+++ b/Engine/Source/Runtime/RenderCore/Private/ShaderCore.cpp
@@ -4142,6 +4142,14 @@ void FShaderCompileJob::OnComplete()
 	// we always run preprocessing as a first step when the job was submitted.
 	Output.PreprocessTime = PreprocessOutput.ElapsedTime;
 
+	const IShaderFormat* ShaderFormat = nullptr;
+	if (!Output.PlatformDebugData.IsEmpty() || Input.DumpDebugInfoEnabled())
+	{
+		// don't bother looking up the shader format unless needed, this can have non-trivial cost due to thread safety
+		// and the possibility of available target platforms needing to be refreshed
+		ShaderFormat = GetTargetPlatformManagerRef().FindShaderFormat(Input.ShaderFormat);
+	}
+
 	// output debug info for the job if either the job was the one which actually executed the compile, or debug info is requested 
 	// for all jobs including ones which hit the job cache (or matched and was deduplicated with another in-flight job)
 	// note that depending on shaderformat implementation this may not necessarily be _all_ debug artifacts, since some require
@@ -4149,7 +4157,7 @@ void FShaderCompileJob::OnComplete()
 	// job which executed compilation)
 	if (Input.DumpDebugInfoEnabled())
 	{
-		const IShaderFormat* ShaderFormat = GetTargetPlatformManagerRef().FindShaderFormat(Input.ShaderFormat);
+		check(ShaderFormat);
 		if (!PreprocessOutput.bSucceeded)
 		{
 			// output minimal debug info for preprocessing failures
@@ -4169,6 +4177,12 @@ void FShaderCompileJob::OnComplete()
 			}
 		}
 	}
+
+	if (!Output.PlatformDebugData.IsEmpty())
+	{
+		check(ShaderFormat);
+		ShaderFormat->NotifyShaderCompiled(Output.PlatformDebugData, Input.ShaderFormat, Input.GenerateDebugInfo());
+	}
 }
 
 void FShaderCompileJob::AppendDebugName(FStringBuilderBase& OutName) const
diff --git a/Engine/Source/Runtime/RenderCore/Private/ShaderMap.cpp b/Engine/Source/Runtime/RenderCore/Private/ShaderMap.cpp
--- a/Engine/Source/Runtime/RenderCore/Private/ShaderMap.cpp
+++ b/Engine/Source/Runtime/RenderCore/Private/ShaderMap.cpp
@@ -249,6 +249,7 @@ bool FShaderMapBase::Serialize(FShaderSerializeContext& Ctx)
 {
 	FArchive& Ar = Ctx.GetMainArchive();
 	LLM_SCOPE(ELLMTag::Shaders);
+	EShaderPlatform ShaderPlatform = EShaderPlatform::SP_NumPlatforms;
 	if (Ar.IsSaving())
 	{
 		check(Content.Object);
@@ -280,26 +281,17 @@ bool FShaderMapBase::Serialize(FShaderSerializeContext& Ctx)
 
 		// Serialize a copy of ShaderPlatform directly into the archive
 		// This will allow us to correctly deserialize the stream, even if we're not able to load the frozen content
-		const EShaderPlatform ShaderPlatform = GetShaderPlatform();
+		ShaderPlatform = GetShaderPlatform();
 		FName ShaderPlatformName = FDataDrivenShaderPlatformInfo::GetName(ShaderPlatform);
 		Ar << ShaderPlatformName;
 
-		if (Ar.IsCooking())
-		{
-			const FName ShaderFormat = LegacyShaderPlatformToShaderFormat(ShaderPlatform);
-			if (ShaderFormat != NAME_None)
-			{
-				Code->NotifyShadersCompiled(ShaderFormat);
-			}
-		}
-
 		if (bShareCode)
 		{
 			FSHAHash ResourceHash = Code->ResourceHash;
 			Ar << ResourceHash;
 			FShaderLibraryCooker::AddShaderCode(ShaderPlatform, Code, GetAssociatedAssets());
 		}
-		else 
+		else
 #endif // WITH_EDITOR
 		{
 			Code->Serialize(Ctx);
@@ -319,8 +311,8 @@ bool FShaderMapBase::Serialize(FShaderSerializeContext& Ctx)
 
 		FName ShaderPlatformName;
 		Ar << ShaderPlatformName;
-		
-		const EShaderPlatform ShaderPlatform = FDataDrivenShaderPlatformInfo::GetShaderPlatformFromName(ShaderPlatformName);
+
+		ShaderPlatform = FDataDrivenShaderPlatformInfo::GetShaderPlatformFromName(ShaderPlatformName);
 
 		if (bShareCode)
 		{
@@ -381,6 +373,19 @@ bool FShaderMapBase::Serialize(FShaderSerializeContext& Ctx)
 		}
 	}
 
+#if WITH_EDITOR
+	bool bNotify = Ar.IsCooking() || GIsEditor;
+	if (bNotify)
+	{
+		check(ShaderPlatform != SP_NumPlatforms); // ShaderPlatform should have been initialized above in either the saving or loading blocks
+		const FName ShaderFormat = LegacyShaderPlatformToShaderFormat(ShaderPlatform);
+		if (ShaderFormat != NAME_None)
+		{
+			Code->NotifyShadersCompiled(ShaderFormat);
+		}
+	}
+#endif
+
 	return (bool)Content.Object;
 }
 
