VulkanRHI: Replace Clamp with a Max that only ensures the value isn't 0.  Fix Nanite rendering in Vulkan.

#rnx

diff --git a/Engine/Source/Runtime/RenderCore/Public/RenderGraphBuilder.inl b/Engine/Source/Runtime/RenderCore/Public/RenderGraphBuilder.inl
--- a/Engine/Source/Runtime/RenderCore/Public/RenderGraphBuilder.inl
+++ b/Engine/Source/Runtime/RenderCore/Public/RenderGraphBuilder.inl
@@ -72,8 +72,8 @@ inline FRDGBufferRef FRDGBuilder::CreateBuffer(
 	IF_RDG_ENABLE_DEBUG(UserValidation.ValidateCreateBuffer(OverrideDesc, Name, Flags));
 
 	// Clamp the buffer size to that which is permissible, otherwise it's a guaranteed crash.
-	OverrideDesc.BytesPerElement = FMath::Clamp<uint32>(OverrideDesc.BytesPerElement, 1u, static_cast<uint32>(GetMaxBufferDimension()));
-	OverrideDesc.NumElements     = FMath::Clamp<uint32>(OverrideDesc.NumElements,     1u, static_cast<uint32>(GetMaxBufferDimension()));
+	OverrideDesc.BytesPerElement = FMath::Max<uint32>(1u, OverrideDesc.BytesPerElement);
+	OverrideDesc.NumElements = FMath::Max<uint32>(1u, OverrideDesc.NumElements);
 
 	FRDGBufferRef Buffer = Buffers.Allocate(Allocators.Root, Name, OverrideDesc, Flags);
 	IF_RDG_ENABLE_DEBUG(UserValidation.ValidateCreateBuffer(Buffer));
