[shaders] add null check for Code object when calling notify to handle cooked editor shader maps, which won't have FShaderMapResourceCode objects - no way to write out symbols for these (for now at least).

diff --git a/Engine/Source/Runtime/RenderCore/Private/ShaderMap.cpp b/Engine/Source/Runtime/RenderCore/Private/ShaderMap.cpp
--- a/Engine/Source/Runtime/RenderCore/Private/ShaderMap.cpp
+++ b/Engine/Source/Runtime/RenderCore/Private/ShaderMap.cpp
@@ -381,7 +381,10 @@ bool FShaderMapBase::Serialize(FShaderSerializeContext& Ctx)
 		const FName ShaderFormat = LegacyShaderPlatformToShaderFormat(ShaderPlatform);
 		if (ShaderFormat != NAME_None)
 		{
-			Code->NotifyShadersCompiled(ShaderFormat);
+			if (Code)
+			{
+				Code->NotifyShadersCompiled(ShaderFormat);
+			}
 		}
 	}
 #endif
