Fix shadows leaving trail in right eye.
The screen space shadow mask was not being cleared in slice 1.

diff --git a/Engine/Source/Runtime/RHI/Public/RHIGlobals.h b/Engine/Source/Runtime/RHI/Public/RHIGlobals.h
--- a/Engine/Source/Runtime/RHI/Public/RHIGlobals.h
+++ b/Engine/Source/Runtime/RHI/Public/RHIGlobals.h
@@ -719,6 +719,9 @@ struct FRHIGlobals
 
 	/** True if lossy/fixed-rate compression of render targets is supporterd. */
 	bool SupportsLossyFramebufferCompression = false;
+
+	/** True if RHI supports supports using any slice or all slices of a texture array as render targets without creating the texture with the TargetArraySlicesIndependently flag. */
+	bool SupportsBindingTexArrayPerSlice = false;
 };
 
 extern RHI_API FRHIGlobals GRHIGlobals;
@@ -894,6 +897,7 @@ extern RHI_API uint32 GGPUFrameTime;
 #define GRHISupportsMSAAShaderResolve						   GRHIGlobals.SupportsMSAAShaderResolve
 #define GRHISupportsDepthStencilResolve                        GRHIGlobals.SupportsDepthStencilResolve
 #define GRHISupportsLossyFramebufferCompression                GRHIGlobals.SupportsLossyFramebufferCompression
+#define GRHISupportsBindingTexArrayPerSlice					   GRHIGlobals.SupportsBindingTexArrayPerSlice
 
 // Utility Getters
 
diff --git a/Engine/Source/Runtime/RenderCore/Private/RenderGraphUtils.cpp b/Engine/Source/Runtime/RenderCore/Private/RenderGraphUtils.cpp
--- a/Engine/Source/Runtime/RenderCore/Private/RenderGraphUtils.cpp
+++ b/Engine/Source/Runtime/RenderCore/Private/RenderGraphUtils.cpp
@@ -658,7 +658,7 @@ void AddClearRenderTargetPass(FRDGBuilder& GraphBuilder, FRDGTextureRef Texture,
 
 	checkf((TextureClearInfo.FirstMipIndex < Texture->Desc.NumMips) && ((TextureClearInfo.FirstMipIndex + TextureClearInfo.NumMips) <= Texture->Desc.NumMips),
 		TEXT("Invalid mip range [%d, %d] for texture %s (%d mips)"), TextureClearInfo.FirstMipIndex, TextureClearInfo.FirstMipIndex + TextureClearInfo.NumMips - 1, Texture->Name, Texture->Desc.NumMips);
-	checkf(((TextureClearInfo.FirstSliceIndex == 0) && (TextureClearInfo.NumSlices == 1)) || (Texture->Desc.IsTextureArray() && EnumHasAnyFlags(Texture->Desc.Flags, ETextureCreateFlags::TargetArraySlicesIndependently)),
+	checkf(((TextureClearInfo.FirstSliceIndex == 0) && (TextureClearInfo.NumSlices == 1)) || (Texture->Desc.IsTextureArray() && (EnumHasAnyFlags(Texture->Desc.Flags, ETextureCreateFlags::TargetArraySlicesIndependently) || GRHISupportsBindingTexArrayPerSlice)),
 		TEXT("Per-slice clear (outside of slice 0, i.e. clearing any other slice than the first one) is only supported on 2DArray at the moment and ETextureCreateFlags::TargetArraySlicesIndependently must be passed when creating the texture (texture %s)."), Texture->Name);
 	checkf((TextureClearInfo.FirstSliceIndex < TextureNumSlicesOrDepth) && ((TextureClearInfo.FirstSliceIndex + TextureClearInfo.NumSlices) <= TextureNumSlicesOrDepth),
 		TEXT("Invalid slice range [%d, %d] for texture %s (%d slices)"), TextureClearInfo.FirstSliceIndex, TextureClearInfo.FirstSliceIndex + TextureClearInfo.NumSlices - 1, Texture->Name, TextureNumSlicesOrDepth);
diff --git a/Engine/Source/Runtime/Renderer/Private/CapsuleShadowRendering.cpp b/Engine/Source/Runtime/Renderer/Private/CapsuleShadowRendering.cpp
--- a/Engine/Source/Runtime/Renderer/Private/CapsuleShadowRendering.cpp
+++ b/Engine/Source/Runtime/Renderer/Private/CapsuleShadowRendering.cpp
@@ -616,7 +616,7 @@ bool FSceneRenderer::RenderCapsuleDirectShadows(
 
 				FUpsampleCapsuleShadowParameters* PassParameters = GraphBuilder.AllocParameters<FUpsampleCapsuleShadowParameters>();
 
-				if (View.Aspects.IsMobileMultiViewEnabled() && !View.Aspects.IsInstancedStereoEnabled()) // Native mobile multi view
+				if (GRHISupportsBindingTexArrayPerSlice) // Native mobile multi view
 				{
 					PassParameters->RenderTargets[0] = FRenderTargetBinding(ScreenShadowMaskTexture, ERenderTargetLoadAction::ELoad, /*MipIndex*/ 0, /*ArraySlice*/ ViewIndex);
 				}
diff --git a/Engine/Source/Runtime/Renderer/Private/ShadowRendering.cpp b/Engine/Source/Runtime/Renderer/Private/ShadowRendering.cpp
--- a/Engine/Source/Runtime/Renderer/Private/ShadowRendering.cpp
+++ b/Engine/Source/Runtime/Renderer/Private/ShadowRendering.cpp
@@ -2415,7 +2415,12 @@ void FMobileSceneRenderer::RenderMobileShadowProjections(
 	RDG_RHI_EVENT_SCOPE(GraphBuilder, RenderMobileShadowProjections);
 
 	FRDGTextureRef ScreenShadowMaskTexture = GraphBuilder.RegisterExternalTexture(GScreenSpaceShadowMaskTextureMobileOutputs.ScreenSpaceShadowMaskTextureMobile, TEXT("ScreenSpaceShadowMaskTextureMobile"));
-	AddClearRenderTargetPass(GraphBuilder, ScreenShadowMaskTexture);
+	FRDGTextureClearInfo TextureClearInfo;
+	if (GRHISupportsBindingTexArrayPerSlice)
+	{
+		TextureClearInfo.NumSlices = ScreenShadowMaskTexture->Desc.ArraySize;
+	}
+	AddClearRenderTargetPass(GraphBuilder, ScreenShadowMaskTexture, TextureClearInfo);
 
 	const FMinimalSceneTextures& SceneTextures = GetActiveSceneTextures();
 	for (auto LightIt = Scene->Lights.CreateConstIterator(); LightIt; ++LightIt)
diff --git a/Engine/Source/Runtime/VulkanRHI/Private/VulkanRHI.cpp b/Engine/Source/Runtime/VulkanRHI/Private/VulkanRHI.cpp
--- a/Engine/Source/Runtime/VulkanRHI/Private/VulkanRHI.cpp
+++ b/Engine/Source/Runtime/VulkanRHI/Private/VulkanRHI.cpp
@@ -938,6 +938,7 @@ void FVulkanDynamicRHI::InitInstance()
 		GRHIMaxDispatchThreadGroupsPerDimension.X = FMath::Min<uint32>(Limits.maxComputeWorkGroupCount[0], 0x7fffffff);
 		GRHIMaxDispatchThreadGroupsPerDimension.Y = FMath::Min<uint32>(Limits.maxComputeWorkGroupCount[1], 0x7fffffff);
 		GRHIMaxDispatchThreadGroupsPerDimension.Z = FMath::Min<uint32>(Limits.maxComputeWorkGroupCount[2], 0x7fffffff);
+		GRHISupportsBindingTexArrayPerSlice = true;
 
 		// Note: While the 2022/2024 profile limits state a minimum of 16, other profiles (even core) go down to 4
 		//       (see https://vulkan.lunarg.com/doc/view/1.3.290.0/windows/profiles_definitions.html)
