GetBasisDeterminantSign - update this to take advantage of the sparsity (6 zeros ) of the matrix.
In this case the determinant of the upper left 3x3 matrix is equal to the determinant of the full 4x4 matrix.
This reduces the multiply count from 40 to 9.

Currently this function is used primarily in mesh importing and computations in editor ( SpeedTreeImport, FbxStaticMeshImport, StaticMeshBuilder, SkeletalMeshLODModel, QuadricMeshReduction etc)

diff --git a/Engine/Source/Runtime/RenderCore/Public/RenderMath.h b/Engine/Source/Runtime/RenderCore/Public/RenderMath.h
--- a/Engine/Source/Runtime/RenderCore/Public/RenderMath.h
+++ b/Engine/Source/Runtime/RenderCore/Public/RenderMath.h
@@ -21,7 +21,9 @@ FORCEINLINE float GetBasisDeterminantSign(const FVector& XAxis, const FVector& Y
 		FPlane(ZAxis, 0),
 		FPlane(0, 0, 0, 1)
 	);
-	return (Basis.Determinant() < 0) ? -1.0f : +1.0f;
+	// RotDeterminant is the determinant of the upper left 3x3 matrix and equivalent to the full 4x4 determinant for this matrix
+	// at one fourth the cost. 
+	return (Basis.RotDeterminant() < 0) ? -1.0f : +1.0f;
 }
 
 /**
