Merging RHI CreateBuffer refactor to Main.

diff --git a/Engine/Source/Runtime/RenderCore/Private/GlobalRenderResources.cpp b/Engine/Source/Runtime/RenderCore/Private/GlobalRenderResources.cpp
--- a/Engine/Source/Runtime/RenderCore/Private/GlobalRenderResources.cpp
+++ b/Engine/Source/Runtime/RenderCore/Private/GlobalRenderResources.cpp
@@ -5,6 +5,7 @@
 #include "Containers/ResourceArray.h"
 #include "RenderCore.h"
 #include "RenderUtils.h"
+#include "RHIResourceUtils.h"
 #include "Algo/Reverse.h"
 #include "Async/Mutex.h"
 
@@ -49,7 +50,7 @@ public:
 	// FResource interface.
 	virtual void InitRHI(FRHICommandListBase& RHICmdList) override
 	{
-		// Create the texture RHI.  		
+		// Create the texture RHI.
 		FColorBulkData BulkData(R, G, B, A);
 
 		// BGRA typed UAV is unsupported per D3D spec, use RGBA here.
@@ -78,10 +79,13 @@ class FEmptyVertexBuffer : public FVertexBufferWithSRV
 public:
 	virtual void InitRHI(FRHICommandListBase& RHICmdList) override
 	{
-		// Create the texture RHI.  		
-		FRHIResourceCreateInfo CreateInfo(TEXT("EmptyVertexBuffer"));
+		// Create the texture RHI.
+		const FRHIBufferCreateDesc CreateDesc =
+			FRHIBufferCreateDesc::CreateVertex(TEXT("EmptyVertexBuffer"), 16)
+			.AddUsage(BUF_Static | BUF_ShaderResource | BUF_UnorderedAccess)
+			.DetermineInitialState();
 
-		VertexBufferRHI = RHICmdList.CreateVertexBuffer(16u, BUF_Static | BUF_ShaderResource | BUF_UnorderedAccess, CreateInfo);
+		VertexBufferRHI = RHICmdList.CreateBuffer(CreateDesc);
 
 		// Create a view of the buffer
 		ShaderResourceViewRHI = RHICmdList.CreateShaderResourceView(
@@ -102,11 +106,13 @@ class FEmptyStructuredBuffer : public FVertexBufferWithSRV
 public:
 	virtual void InitRHI(FRHICommandListBase& RHICmdList) override
 	{
-		// Create the buffer RHI.  		
-		FRHIResourceCreateInfo CreateInfo(TEXT("EmptyStructuredBuffer"));
+		const FRHIBufferCreateDesc CreateDesc =
+			FRHIBufferCreateDesc::CreateStructured(TEXT("EmptyStructuredBuffer"), sizeof(uint32) * 4, sizeof(uint32))
+			.AddUsage(BUF_Static | BUF_ShaderResource | BUF_UnorderedAccess)
+			.DetermineInitialState();
 
-		const uint32 BufferSize = sizeof(uint32) * 4u;
-		VertexBufferRHI = RHICmdList.CreateStructuredBuffer(sizeof(uint32), BufferSize, BUF_Static | BUF_ShaderResource | BUF_UnorderedAccess, CreateInfo);
+		// Create the buffer RHI.
+		VertexBufferRHI = RHICmdList.CreateBuffer(CreateDesc);
 
 		// Create a view of the buffer
 		ShaderResourceViewRHI = RHICmdList.CreateShaderResourceView(VertexBufferRHI, FRHIViewDesc::CreateBufferSRV().SetTypeFromBuffer(VertexBufferRHI));
@@ -119,15 +125,15 @@ class FBlackFloat4StructuredBufferWithSRV : public FVertexBufferWithSRV
 public:
 	virtual void InitRHI(FRHICommandListBase& RHICmdList) override
 	{
-		// Create the buffer RHI.  		
-		FRHIResourceCreateInfo CreateInfo(TEXT("BlackFloat4StructuredBuffer"));
+		const FRHIBufferCreateDesc CreateDesc =
+			FRHIBufferCreateDesc::CreateStructured(TEXT("BlackFloat4StructuredBuffer"), sizeof(FVector4f), sizeof(FVector4f))
+			.AddUsage(BUF_Static | BUF_ShaderResource | BUF_UnorderedAccess)
+			.SetInitialState(ERHIAccess::SRVMask);
 
-		const uint32 BufferSize = sizeof(FVector4f);
-		VertexBufferRHI = RHICmdList.CreateStructuredBuffer(sizeof(FVector4f), BufferSize, BUF_Static | BUF_ShaderResource | BUF_UnorderedAccess, ERHIAccess::SRVMask, CreateInfo);
+		const FVector4f InitialData(0.0f, 0.0f, 0.0f, 0.0f);
 
-		FVector4f* BufferData = (FVector4f*)RHICmdList.LockBuffer(VertexBufferRHI, 0, sizeof(FVector4f), RLM_WriteOnly);
-		*BufferData = FVector4f(0.0f, 0.0f, 0.0f, 0.0f);
-		RHICmdList.UnlockBuffer(VertexBufferRHI);
+		// Create the buffer RHI.
+		VertexBufferRHI = UE::RHIResourceUtils::CreateBufferWithValue(RHICmdList, CreateDesc, InitialData);
 
 		// Create a view of the buffer
 		ShaderResourceViewRHI = RHICmdList.CreateShaderResourceView(VertexBufferRHI, FRHIViewDesc::CreateBufferSRV().SetTypeFromBuffer(VertexBufferRHI));
@@ -140,14 +146,15 @@ class FBlackFloat4VertexBuffer : public FVertexBufferWithSRV
 public:
 	virtual void InitRHI(FRHICommandListBase& RHICmdList) override
 	{
-		// Create the texture RHI.  		
-		FRHIResourceCreateInfo CreateInfo(TEXT("BlackFloat4VertexBuffer"));
+		const FRHIBufferCreateDesc CreateDesc =
+			FRHIBufferCreateDesc::CreateVertex(TEXT("BlackFloat4VertexBuffer"), sizeof(FVector4f))
+			.AddUsage(BUF_Static | BUF_ShaderResource | BUF_UnorderedAccess)
+			.SetInitialState(ERHIAccess::SRVMask);
 
-		VertexBufferRHI = RHICmdList.CreateVertexBuffer(sizeof(FVector4f), BUF_Static | BUF_ShaderResource | BUF_UnorderedAccess, ERHIAccess::SRVMask, CreateInfo);
+		const FVector4f InitialData(0.0f, 0.0f, 0.0f, 0.0f);
 
-		FVector4f* BufferData = (FVector4f*)RHICmdList.LockBuffer(VertexBufferRHI, 0, sizeof(FVector4f), RLM_WriteOnly);
-		*BufferData = FVector4f(0.0f, 0.0f, 0.0f, 0.0f);
-		RHICmdList.UnlockBuffer(VertexBufferRHI);
+		// Create the texture RHI.
+		VertexBufferRHI = UE::RHIResourceUtils::CreateBufferWithValue(RHICmdList, CreateDesc, InitialData);
 
 		// Create a view of the buffer
 		ShaderResourceViewRHI = RHICmdList.CreateShaderResourceView(
@@ -195,14 +202,15 @@ class FWhiteVertexBuffer : public FVertexBufferWithSRV
 public:
 	virtual void InitRHI(FRHICommandListBase& RHICmdList) override
 	{
-		// Create the texture RHI.  		
-		FRHIResourceCreateInfo CreateInfo(TEXT("WhiteVertexBuffer"));
+		const FRHIBufferCreateDesc CreateDesc =
+			FRHIBufferCreateDesc::CreateVertex(TEXT("WhiteVertexBuffer"), sizeof(FVector4f))
+			.AddUsage(BUF_Static | BUF_ShaderResource)
+			.DetermineInitialState();
 
-		VertexBufferRHI = RHICmdList.CreateVertexBuffer(sizeof(FVector4f), BUF_Static | BUF_ShaderResource, CreateInfo);
+		const FVector4f InitialData(1.0f, 1.0f, 1.0f, 1.0f);
 
-		FVector4f* BufferData = (FVector4f*)RHICmdList.LockBuffer(VertexBufferRHI, 0, sizeof(FVector4f), RLM_WriteOnly);
-		*BufferData = FVector4f(1.0f, 1.0f, 1.0f, 1.0f);
-		RHICmdList.UnlockBuffer(VertexBufferRHI);
+		// Create the buffer RHI.
+		VertexBufferRHI = UE::RHIResourceUtils::CreateBufferWithValue(RHICmdList, CreateDesc, InitialData);
 
 		// Create a view of the buffer
 		ShaderResourceViewRHI = RHICmdList.CreateShaderResourceView(
@@ -220,14 +228,15 @@ class FBlackVertexBuffer : public FVertexBufferWithSRV
 public:
 	virtual void InitRHI(FRHICommandListBase& RHICmdList) override
 	{
-		// Create the texture RHI.  		
-		FRHIResourceCreateInfo CreateInfo(TEXT("BlackVertexBuffer"));
+		const FRHIBufferCreateDesc CreateDesc =
+			FRHIBufferCreateDesc::CreateVertex(TEXT("BlackVertexBuffer"), sizeof(FVector4f))
+			.AddUsage(BUF_Static | BUF_ShaderResource)
+			.DetermineInitialState();
 
-		VertexBufferRHI = RHICmdList.CreateVertexBuffer(sizeof(FVector4f), BUF_Static | BUF_ShaderResource, CreateInfo);
+		const FVector4f InitialData(0.0f, 0.0f, 0.0f, 0.0f);
 
-		FVector4f* BufferData = (FVector4f*)RHICmdList.LockBuffer(VertexBufferRHI, 0, sizeof(FVector4f), RLM_WriteOnly);
-		*BufferData = FVector4f(0.0f, 0.0f, 0.0f, 0.0f);
-		RHICmdList.UnlockBuffer(VertexBufferRHI);
+		// Create the buffer RHI.
+		VertexBufferRHI = UE::RHIResourceUtils::CreateBufferWithValue(RHICmdList, CreateDesc, InitialData);
 
 		// Create a view of the buffer
 		ShaderResourceViewRHI = RHICmdList.CreateShaderResourceView(
@@ -618,15 +627,15 @@ FNullColorVertexBuffer::~FNullColorVertexBuffer() = default;
 void FNullColorVertexBuffer::InitRHI(FRHICommandListBase& RHICmdList)
 {
 	// create a static vertex buffer
-	FRHIResourceCreateInfo CreateInfo(TEXT("FNullColorVertexBuffer"));
-
-	VertexBufferRHI = RHICmdList.CreateBuffer(sizeof(uint32) * 4, BUF_Static | BUF_VertexBuffer | BUF_ShaderResource, 0, ERHIAccess::VertexOrIndexBuffer | ERHIAccess::SRVMask, CreateInfo);
-	uint32* Vertices = (uint32*)RHICmdList.LockBuffer(VertexBufferRHI, 0, sizeof(uint32) * 4, RLM_WriteOnly);
-	Vertices[0] = FColor(255, 255, 255, 255).DWColor();
-	Vertices[1] = FColor(255, 255, 255, 255).DWColor();
-	Vertices[2] = FColor(255, 255, 255, 255).DWColor();
-	Vertices[3] = FColor(255, 255, 255, 255).DWColor();
-	RHICmdList.UnlockBuffer(VertexBufferRHI);
+	const FRHIBufferCreateDesc CreateDesc =
+		FRHIBufferCreateDesc::CreateVertex<uint32>(TEXT("FNullColorVertexBuffer"), 4)
+		.AddUsage(EBufferUsageFlags::Static | EBufferUsageFlags::ShaderResource)
+		.SetInitialState(ERHIAccess::VertexOrIndexBuffer | ERHIAccess::SRVMask);
+
+	const uint32 DWWhite = FColor(255, 255, 255, 255).DWColor();
+	const uint32 ColorValues[4] = { DWWhite, DWWhite, DWWhite, DWWhite };
+
+	VertexBufferRHI = UE::RHIResourceUtils::CreateBufferWithArray(RHICmdList, CreateDesc, ColorValues);
 	VertexBufferSRV = RHICmdList.CreateShaderResourceView(
 		VertexBufferRHI,
 		FRHIViewDesc::CreateBufferSRV()
@@ -652,12 +661,14 @@ FNullVertexBuffer::~FNullVertexBuffer() = default;
 void FNullVertexBuffer::InitRHI(FRHICommandListBase& RHICmdList)
 {
 	// create a static vertex buffer
-	FRHIResourceCreateInfo CreateInfo(TEXT("FNullVertexBuffer"));
-	VertexBufferRHI = RHICmdList.CreateBuffer(sizeof(FVector3f), BUF_Static | BUF_VertexBuffer | BUF_ShaderResource, 0, ERHIAccess::VertexOrIndexBuffer | ERHIAccess::SRVMask, CreateInfo);
-	FVector3f* LockedData = (FVector3f*)RHICmdList.LockBuffer(VertexBufferRHI, 0, sizeof(FVector3f), RLM_WriteOnly);
-	*LockedData = FVector3f(0.0f);
-	RHICmdList.UnlockBuffer(VertexBufferRHI);
+	const FRHIBufferCreateDesc CreateDesc =
+		FRHIBufferCreateDesc::CreateVertex<FVector3f>(TEXT("FNullVertexBuffer"), 1)
+		.AddUsage(EBufferUsageFlags::Static | EBufferUsageFlags::ShaderResource)
+		.SetInitialState(ERHIAccess::VertexOrIndexBuffer | ERHIAccess::SRVMask);
 
+	const FVector3f InitialValue(0.0f);
+
+	VertexBufferRHI = UE::RHIResourceUtils::CreateBufferWithValue(RHICmdList, CreateDesc, InitialValue);
 	VertexBufferSRV = RHICmdList.CreateShaderResourceView(
 		VertexBufferRHI,
 		FRHIViewDesc::CreateBufferSRV()
@@ -677,21 +688,18 @@ TGlobalResource<FNullVertexBuffer, FRenderResource::EInitPhase::Pre> GNullVertex
 //////////////////////////////////////////////////////////////////////////////////////////////////////////////////////
 // FScreenSpaceVertexBuffer
 
+static const FVector2f GScreenSpaceVertexBufferData[4] =
+{
+	FVector2f(-1,-1),
+	FVector2f(-1,+1),
+	FVector2f(+1,-1),
+	FVector2f(+1,+1),
+};
+
 void FScreenSpaceVertexBuffer::InitRHI(FRHICommandListBase& RHICmdList)
 {
 	// create a static vertex buffer
-	FRHIResourceCreateInfo CreateInfo(TEXT("FScreenSpaceVertexBuffer"));
-	VertexBufferRHI = RHICmdList.CreateVertexBuffer(sizeof(FVector2f) * 4, BUF_Static, CreateInfo);
-	void* VoidPtr = RHICmdList.LockBuffer(VertexBufferRHI, 0, sizeof(FVector2f) * 4, RLM_WriteOnly);
-	static const FVector2f Vertices[4] =
-	{
-		FVector2f(-1,-1),
-		FVector2f(-1,+1),
-		FVector2f(+1,-1),
-		FVector2f(+1,+1),
-	};
-	FMemory::Memcpy(VoidPtr, Vertices, sizeof(FVector2f) * 4);
-	RHICmdList.UnlockBuffer(VertexBufferRHI);
+	VertexBufferRHI = UE::RHIResourceUtils::CreateVertexBufferFromArray(RHICmdList, TEXT("FScreenSpaceVertexBuffer"), EBufferUsageFlags::Static, MakeConstArrayView(GScreenSpaceVertexBufferData));
 }
 
 TGlobalResource<FScreenSpaceVertexBuffer, FRenderResource::EInitPhase::Pre> GScreenSpaceVertexBuffer;
@@ -722,12 +730,8 @@ TGlobalResource<FTileVertexDeclaration, FRenderResource::EInitPhase::Pre> GTileV
 
 void FCubeIndexBuffer::InitRHI(FRHICommandListBase& RHICmdList)
 {
-	// create a static vertex buffer
-	FRHIResourceCreateInfo CreateInfo(TEXT("FCubeIndexBuffer"));
-	IndexBufferRHI = RHICmdList.CreateIndexBuffer(sizeof(uint16), sizeof(uint16) * NUM_CUBE_VERTICES, BUF_Static, CreateInfo);
-	void* VoidPtr = RHICmdList.LockBuffer(IndexBufferRHI, 0, sizeof(uint16) * NUM_CUBE_VERTICES, RLM_WriteOnly);
-	FMemory::Memcpy(VoidPtr, GCubeIndices, NUM_CUBE_VERTICES * sizeof(uint16));
-	RHICmdList.UnlockBuffer(IndexBufferRHI);
+	// create a static index buffer
+	IndexBufferRHI = UE::RHIResourceUtils::CreateIndexBufferFromArray(RHICmdList, TEXT("FCubeIndexBuffer"), EBufferUsageFlags::Static, MakeConstArrayView(GCubeIndices));
 }
 
 TGlobalResource<FCubeIndexBuffer, FRenderResource::EInitPhase::Pre> GCubeIndexBuffer;
@@ -735,15 +739,12 @@ TGlobalResource<FCubeIndexBuffer, FRenderResource::EInitPhase::Pre> GCubeIndexBu
 //////////////////////////////////////////////////////////////////////////////////////////////////////////////////////
 // FTwoTrianglesIndexBuffer
 
+static const uint16 GTwoTrianglesIndexBufferData[6] = { 0, 1, 3, 0, 3, 2 };
+
 void FTwoTrianglesIndexBuffer::InitRHI(FRHICommandListBase& RHICmdList)
 {
-	// create a static vertex buffer
-	FRHIResourceCreateInfo CreateInfo(TEXT("FTwoTrianglesIndexBuffer"));
-	IndexBufferRHI = RHICmdList.CreateIndexBuffer(sizeof(uint16), sizeof(uint16) * 6, BUF_Static, CreateInfo);
-	void* VoidPtr = RHICmdList.LockBuffer(IndexBufferRHI, 0, sizeof(uint16) * 6, RLM_WriteOnly);
-	static const uint16 Indices[] = { 0, 1, 3, 0, 3, 2 };
-	FMemory::Memcpy(VoidPtr, Indices, 6 * sizeof(uint16));
-	RHICmdList.UnlockBuffer(IndexBufferRHI);
+	// create a static index buffer
+	IndexBufferRHI = UE::RHIResourceUtils::CreateIndexBufferFromArray(RHICmdList, TEXT("FTwoTrianglesIndexBuffer"), EBufferUsageFlags::Static, MakeConstArrayView(GTwoTrianglesIndexBufferData));
 }
 
 TGlobalResource<FTwoTrianglesIndexBuffer, FRenderResource::EInitPhase::Pre> GTwoTrianglesIndexBuffer;
@@ -808,10 +809,13 @@ public:
 	virtual void InitRHI(FRHICommandListBase& RHICmdList) override
 	{
 		check(!IsValidRef(BufferType::GetRHI()));
-		FRHIResourceCreateInfo CreateInfo(TEXT("FDynamicBuffer"));
-		EBufferUsageFlags UsageFlags = (Stride == 0) ? EBufferUsageFlags::VertexBuffer : EBufferUsageFlags::IndexBuffer;
-		UsageFlags |= EBufferUsageFlags::ShaderResource;
-		BufferType::SetRHI(RHICmdList.CreateBuffer(BufferSize, UsageFlags, Stride, ERHIAccess::VertexOrIndexBuffer | ERHIAccess::SRVCompute, CreateInfo));
+
+		const FRHIBufferCreateDesc CreateDesc =
+			FRHIBufferCreateDesc::Create(TEXT("FDynamicBuffer"), BufferSize, Stride, Stride == 0 ? EBufferUsageFlags::VertexBuffer : EBufferUsageFlags::IndexBuffer)
+			.AddUsage(EBufferUsageFlags::ShaderResource)
+			.SetInitialState(ERHIAccess::VertexOrIndexBuffer | ERHIAccess::SRVCompute);
+
+		BufferType::SetRHI(RHICmdList.CreateBuffer(CreateDesc));
 		MappedBuffer = NULL;
 		AllocatedByteCount = 0;
 	}
diff --git a/Engine/Source/Runtime/RenderCore/Private/RenderGraphResourcePool.cpp b/Engine/Source/Runtime/RenderCore/Private/RenderGraphResourcePool.cpp
--- a/Engine/Source/Runtime/RenderCore/Private/RenderGraphResourcePool.cpp
+++ b/Engine/Source/Runtime/RenderCore/Private/RenderGraphResourcePool.cpp
@@ -185,9 +185,11 @@ FRDGPooledBuffer* FRDGBufferPool::CreateBuffer(FRHICommandListBase& RHICmdList,
 	LLM_TAGSET_SCOPE_CLEAR(ELLMTagSet::Assets);
 	LLM_TAGSET_SCOPE_CLEAR(ELLMTagSet::AssetClasses);
 
-	const ERHIAccess InitialAccess = RHIGetDefaultResourceState(Desc.Usage, false);
-	FRHIResourceCreateInfo CreateInfo(InDebugName);
-	TRefCountPtr<FRHIBuffer> BufferRHI = RHICmdList.CreateBuffer(NumBytes, Desc.Usage, Desc.BytesPerElement, InitialAccess, CreateInfo);
+	const FRHIBufferCreateDesc CreateDesc =
+		FRHIBufferCreateDesc::Create(InDebugName, NumBytes, Desc.BytesPerElement, Desc.Usage)
+		.SetInitialState(RHIGetDefaultResourceState(Desc.Usage, false));
+
+	TRefCountPtr<FRHIBuffer> BufferRHI = RHICmdList.CreateBuffer(CreateDesc);
 
 	FRDGPooledBuffer* PooledBuffer = new FRDGPooledBuffer(RHICmdList, MoveTemp(BufferRHI), Desc, Desc.NumElements, InDebugName);
 	AllocatedBuffers.Add(PooledBuffer);
diff --git a/Engine/Source/Runtime/RenderCore/Private/RenderResource.cpp b/Engine/Source/Runtime/RenderCore/Private/RenderResource.cpp
--- a/Engine/Source/Runtime/RenderCore/Private/RenderResource.cpp
+++ b/Engine/Source/Runtime/RenderCore/Private/RenderResource.cpp
@@ -265,15 +265,21 @@ FBufferRHIRef FRenderResource::CreateRHIBufferInternal(
 	bool bWithoutNativeResource)
 {
 	const uint32 SizeInBytes = ResourceArray ? ResourceArray->GetResourceDataSize() : 0;
-	FRHIResourceCreateInfo CreateInfo(InDebugName, ResourceArray);
-	CreateInfo.ClassName = FName(InDebugName);
-	CreateInfo.OwnerName = InOwnerName;
-	CreateInfo.bWithoutNativeResource = bWithoutNativeResource;
 
-	FBufferRHIRef Buffer = RHICmdList.CreateBuffer(SizeInBytes, InBufferUsageFlags | EBufferUsageFlags::VertexBuffer, 0, ERHIAccess::VertexOrIndexBuffer | ERHIAccess::SRVMask, CreateInfo);
+	FRHIBufferCreateDesc CreateDesc = bWithoutNativeResource
+		? FRHIBufferCreateDesc::CreateNull(InDebugName)
+		: (FRHIBufferCreateDesc::CreateVertex(InDebugName, SizeInBytes).AddUsage(InBufferUsageFlags));
 
-	Buffer->SetOwnerName(InOwnerName);
-	return Buffer;
+	CreateDesc.SetClassName(FName(InDebugName))
+		.SetOwnerName(InOwnerName)
+		.SetInitialState(ERHIAccess::VertexOrIndexBuffer | ERHIAccess::SRVMask);
+
+	if (ResourceArray)
+	{
+		CreateDesc.SetInitActionResourceArray(ResourceArray);
+	}
+
+	return RHICmdList.CreateBuffer(CreateDesc);
 }
 
 void BeginInitResource(FRenderResource* Resource, FRenderCommandPipe* RenderCommandPipe)
diff --git a/Engine/Source/Runtime/RenderCore/Private/UnifiedBuffer.cpp b/Engine/Source/Runtime/RenderCore/Private/UnifiedBuffer.cpp
--- a/Engine/Source/Runtime/RenderCore/Private/UnifiedBuffer.cpp
+++ b/Engine/Source/Runtime/RenderCore/Private/UnifiedBuffer.cpp
@@ -1450,7 +1450,7 @@ void FScatterUploadBuffer::Init( uint32 NumElements, uint32 InNumBytesPerElement
 	NumBytesPerElement = InNumBytesPerElement;
 	bFloat4Buffer = bInFloat4Buffer;
 
-	const EBufferUsageFlags Usage = bInFloat4Buffer ? BUF_None : BUF_ByteAddressBuffer;
+	const EBufferUsageFlags Usage = BUF_ShaderResource | BUF_Volatile | (bInFloat4Buffer ? BUF_None : BUF_ByteAddressBuffer);
 	const uint32 TypeSize = bInFloat4Buffer ? 16 : 4;
 
 	uint32 ScatterBytes = NumElements * sizeof( uint32 );
@@ -1486,8 +1486,11 @@ void FScatterUploadBuffer::Init( uint32 NumElements, uint32 InNumBytesPerElement
 			ScatterBuffer.Release();
 			ScatterBuffer.NumBytes = ScatterBufferSize;
 
-			FRHIResourceCreateInfo CreateInfo(DebugName);
-			ScatterBuffer.Buffer = RHICmdList.CreateStructuredBuffer(sizeof(uint32), ScatterBuffer.NumBytes, BUF_ShaderResource | BUF_Volatile | Usage, CreateInfo);
+			const FRHIBufferCreateDesc ScatterBufferDesc =
+				FRHIBufferCreateDesc::CreateStructured(DebugName, ScatterBuffer.NumBytes, sizeof(uint32))
+				.AddUsage(Usage);
+
+			ScatterBuffer.Buffer = RHICmdList.CreateBuffer(ScatterBufferDesc);
 			ScatterBuffer.SRV = RHICmdList.CreateShaderResourceView(ScatterBuffer.Buffer, FRHIViewDesc::CreateBufferSRV().SetTypeFromBuffer(ScatterBuffer.Buffer));
 		}
 
@@ -1497,8 +1500,11 @@ void FScatterUploadBuffer::Init( uint32 NumElements, uint32 InNumBytesPerElement
 			UploadBuffer.Release();
 			UploadBuffer.NumBytes = UploadBufferSize;
 
-			FRHIResourceCreateInfo CreateInfo(DebugName);
-			UploadBuffer.Buffer = RHICmdList.CreateStructuredBuffer(TypeSize, UploadBuffer.NumBytes, BUF_ShaderResource | BUF_Volatile | Usage, CreateInfo);
+			const FRHIBufferCreateDesc UploadBufferDesc =
+				FRHIBufferCreateDesc::CreateStructured(DebugName, UploadBuffer.NumBytes, TypeSize)
+				.AddUsage(Usage);
+
+			UploadBuffer.Buffer = RHICmdList.CreateBuffer(UploadBufferDesc);
 			UploadBuffer.SRV = RHICmdList.CreateShaderResourceView(UploadBuffer.Buffer, FRHIViewDesc::CreateBufferSRV().SetTypeFromBuffer(UploadBuffer.Buffer));
 		}
 
diff --git a/Engine/Source/Runtime/RenderCore/Private/VertexFactory.cpp b/Engine/Source/Runtime/RenderCore/Private/VertexFactory.cpp
--- a/Engine/Source/Runtime/RenderCore/Private/VertexFactory.cpp
+++ b/Engine/Source/Runtime/RenderCore/Private/VertexFactory.cpp
@@ -10,6 +10,7 @@
 #include "ShaderCompilerCore.h"
 #include "RenderUtils.h"
 #include "DataDrivenShaderPlatformInfo.h"
+#include "RHIResourceUtils.h"
 #if WITH_EDITOR
 #include "Serialization/CompactBinary.h"
 #include "Serialization/CompactBinaryWriter.h"
@@ -403,12 +404,12 @@ void FVertexFactory::InitDeclaration(const FVertexDeclarationElementList& Elemen
 void FPrimitiveIdDummyBuffer::InitRHI(FRHICommandListBase& RHICmdList) 
 {
 	// create a static vertex buffer
-	FRHIResourceCreateInfo CreateInfo(TEXT("FPrimitiveIdDummyBuffer"));
+	const FRHIBufferCreateDesc CreateDesc =
+		FRHIBufferCreateDesc::CreateVertex(TEXT("FPrimitiveIdDummyBuffer"), sizeof(uint32))
+		.AddUsage(EBufferUsageFlags::Static | EBufferUsageFlags::ShaderResource)
+		.SetInitialState(ERHIAccess::VertexOrIndexBuffer | ERHIAccess::SRVMask);
 
-	VertexBufferRHI = RHICmdList.CreateBuffer(sizeof(uint32), BUF_Static | BUF_VertexBuffer | BUF_ShaderResource, 0, ERHIAccess::VertexOrIndexBuffer | ERHIAccess::SRVMask, CreateInfo);
-	uint32* Vertices = (uint32*)RHICmdList.LockBuffer(VertexBufferRHI, 0, sizeof(uint32), RLM_WriteOnly);
-	Vertices[0] = 0;
-	RHICmdList.UnlockBuffer(VertexBufferRHI);
+	VertexBufferRHI = UE::RHIResourceUtils::CreateBufferZeroed(RHICmdList, CreateDesc);
 	VertexBufferSRV = RHICmdList.CreateShaderResourceView(
 		VertexBufferRHI,
 		FRHIViewDesc::CreateBufferSRV()
diff --git a/Engine/Source/Runtime/RenderCore/Public/MediaShaders.h b/Engine/Source/Runtime/RenderCore/Public/MediaShaders.h
--- a/Engine/Source/Runtime/RenderCore/Public/MediaShaders.h
+++ b/Engine/Source/Runtime/RenderCore/Public/MediaShaders.h
@@ -103,15 +103,16 @@ struct FMediaElementVertex
 	{ }
 };
 
-inline FBufferRHIRef CreateTempMediaVertexBuffer(float ULeft = 0.0f, float URight = 1.0f, float VTop = 0.0f, float VBottom = 1.0f)
+inline FBufferRHIRef CreateTempMediaVertexBuffer(FRHICommandListBase& RHICmdList, float ULeft = 0.0f, float URight = 1.0f, float VTop = 0.0f, float VBottom = 1.0f)
 {
-	FRHICommandListBase& RHICmdList = FRHICommandListImmediate::Get();
+	const FRHIBufferCreateDesc CreateDesc =
+		FRHIBufferCreateDesc::CreateVertex<FMediaElementVertex>(TEXT("TempMediaVertexBuffer"), 4)
+		.AddUsage(EBufferUsageFlags::Volatile)
+		.DetermineInitialState()
+		.SetInitActionInitializer();
 
-	FRHIResourceCreateInfo CreateInfo(TEXT("TempMediaVertexBuffer"));
-	FBufferRHIRef VertexBufferRHI = RHICmdList.CreateVertexBuffer(sizeof(FMediaElementVertex) * 4, BUF_Volatile, CreateInfo);
-	void* VoidPtr = RHICmdList.LockBuffer(VertexBufferRHI, 0, sizeof(FMediaElementVertex) * 4, RLM_WriteOnly);
+	TRHIBufferInitializer<FMediaElementVertex> Vertices = RHICmdList.CreateBufferInitializer(CreateDesc);
 
-	FMediaElementVertex* Vertices = (FMediaElementVertex*)VoidPtr;
 	Vertices[0].Position.Set(-1.0f, 1.0f, 1.0f, 1.0f); // Top Left
 	Vertices[1].Position.Set(1.0f, 1.0f, 1.0f, 1.0f); // Top Right
 	Vertices[2].Position.Set(-1.0f, -1.0f, 1.0f, 1.0f); // Bottom Left
@@ -121,9 +122,14 @@ inline FBufferRHIRef CreateTempMediaVertexBuffer(float ULeft = 0.0f, float URigh
 	Vertices[1].TextureCoordinate.Set(URight, VTop);
 	Vertices[2].TextureCoordinate.Set(ULeft, VBottom);
 	Vertices[3].TextureCoordinate.Set(URight, VBottom);
-	RHICmdList.UnlockBuffer(VertexBufferRHI);
 
-	return VertexBufferRHI;
+	return Vertices.Finalize();
+}
+
+UE_DEPRECATED(5.6, "CreateTempMediaVertexBuffer requires a Command List argument")
+inline FBufferRHIRef CreateTempMediaVertexBuffer(float ULeft = 0.0f, float URight = 1.0f, float VTop = 0.0f, float VBottom = 1.0f)
+{
+	return CreateTempMediaVertexBuffer(FRHICommandListImmediate::Get(), ULeft, URight, VTop, VBottom);
 }
 
 /**
diff --git a/Engine/Source/Runtime/Renderer/Private/FXRenderingUtils.cpp b/Engine/Source/Runtime/Renderer/Private/FXRenderingUtils.cpp
--- a/Engine/Source/Runtime/Renderer/Private/FXRenderingUtils.cpp
+++ b/Engine/Source/Runtime/Renderer/Private/FXRenderingUtils.cpp
@@ -7,6 +7,7 @@
 #include "Lumen/LumenScreenProbeGather.h"
 #include "MaterialShared.h"
 #include "PrimitiveSceneProxy.h"
+#include "RHIResourceUtils.h"
 #include "SceneInterface.h"
 #include "ScenePrivate.h"
 #include "SystemTextures.h"
@@ -169,15 +170,14 @@ public:
 		FRDGBufferDesc BufferDesc = FRDGBufferDesc::CreateStructuredDesc(sizeof(uint32), 1);
 		BufferDesc.Usage = BUF_Static | BUF_ShaderResource | BUF_ByteAddressBuffer;
 
-		FRHIResourceCreateInfo CreateInfo(TEXT("FDFDummyByteAddress"));
-		FBufferRHIRef RHIBuffer = RHICmdList.CreateStructuredBuffer(BufferDesc.BytesPerElement, BufferDesc.GetSize(), BufferDesc.Usage, CreateInfo);
-		{
-			void* Data = RHICmdList.LockBuffer(RHIBuffer, 0, BufferDesc.GetSize(), RLM_WriteOnly);
-			FMemory::Memset(Data, 0, BufferDesc.GetSize());
-			RHICmdList.UnlockBuffer(RHIBuffer);
-		}
+		const FRHIBufferCreateDesc BufferCreateDesc =
+			FRHIBufferCreateDesc::CreateStructured(TEXT("FDFDummyByteAddress"), BufferDesc.GetSize(), BufferDesc.BytesPerElement)
+			.AddUsage(BufferDesc.Usage)
+			.DetermineInitialState();
+
+		FBufferRHIRef RHIBuffer = UE::RHIResourceUtils::CreateBufferZeroed(RHICmdList, BufferCreateDesc);
 
-		PooledBuffer = new FRDGPooledBuffer(RHICmdList, RHIBuffer, BufferDesc, BufferDesc.NumElements, CreateInfo.DebugName);
+		PooledBuffer = new FRDGPooledBuffer(RHICmdList, RHIBuffer, BufferDesc, BufferDesc.NumElements, BufferCreateDesc.DebugName);
 	}
 
 	virtual void ReleaseRHI() override
diff --git a/Engine/Source/Runtime/Renderer/Private/GPUBenchmark.cpp b/Engine/Source/Runtime/Renderer/Private/GPUBenchmark.cpp
--- a/Engine/Source/Runtime/Renderer/Private/GPUBenchmark.cpp
+++ b/Engine/Source/Runtime/Renderer/Private/GPUBenchmark.cpp
@@ -25,6 +25,7 @@
 #include "VisualizeTexture.h"
 #include "CommonRenderResources.h"
 #include "SceneView.h"
+#include "RHIResourceUtils.h"
 
 static const uint32 GBenchmarkResolution = 512;
 static const uint32 GBenchmarkPrimitives = 200000;
@@ -410,15 +411,14 @@ void RendererGPUBenchmark(FRHICommandListImmediate& RHICmdList, FSynthBenchmarkR
 	MeasureLongGPUTaskExecutionTime(RHICmdList);
 	RHICmdList.ImmediateFlush(EImmediateFlushType::FlushRHIThread);	
 
-	TResourceArray<FBenchmarkVertex> Vertices;
+	TArray<FBenchmarkVertex> Vertices;
 	Vertices.Reserve(GBenchmarkVertices);
 	for (uint32 Index = 0; Index < GBenchmarkVertices; ++Index)
 	{
 		Vertices.Emplace(Index);
 	}
 
-	FRHIResourceCreateInfo CreateInfo(TEXT("RendererGPUBenchmark"), &Vertices);
-	FBufferRHIRef VertexBuffer = RHICmdList.CreateVertexBuffer(GBenchmarkVertices * sizeof(FBenchmarkVertex), BUF_Static, CreateInfo);
+	FBufferRHIRef VertexBuffer = UE::RHIResourceUtils::CreateVertexBufferFromArray(RHICmdList, TEXT("RendererGPUBenchmark"), EBufferUsageFlags::Static, MakeConstArrayView(Vertices));
 
 	// two RT to ping pong so we force the GPU to flush its pipeline
 	TRefCountPtr<IPooledRenderTarget> RTItems[2];
diff --git a/Engine/Source/Runtime/Renderer/Private/MeshDrawCommands.cpp b/Engine/Source/Runtime/Renderer/Private/MeshDrawCommands.cpp
--- a/Engine/Source/Runtime/Renderer/Private/MeshDrawCommands.cpp
+++ b/Engine/Source/Runtime/Renderer/Private/MeshDrawCommands.cpp
@@ -98,8 +98,13 @@ FPrimitiveIdVertexBufferPoolEntry FPrimitiveIdVertexBufferPool::Allocate(FRHICom
 		FPrimitiveIdVertexBufferPoolEntry NewEntry;
 		NewEntry.LastDiscardId = DiscardId;
 		NewEntry.BufferSize = BufferSize;
-		FRHIResourceCreateInfo CreateInfo(TEXT("FPrimitiveIdVertexBufferPool"));
-		NewEntry.BufferRHI = RHICmdList.CreateBuffer(NewEntry.BufferSize, BUF_VertexBuffer | BUF_Volatile | BUF_UniformBuffer | BUF_ShaderResource | BUF_ByteAddressBuffer, 0, ERHIAccess::VertexOrIndexBuffer, CreateInfo);
+
+		const FRHIBufferCreateDesc CreateDesc =
+			FRHIBufferCreateDesc::CreateVertex(TEXT("FPrimitiveIdVertexBufferPool"), BufferSize)
+			.AddUsage(EBufferUsageFlags::Volatile | EBufferUsageFlags::UniformBuffer | EBufferUsageFlags::ShaderResource | EBufferUsageFlags::ByteAddressBuffer)
+			.SetInitialState(ERHIAccess::VertexOrIndexBuffer);
+
+		NewEntry.BufferRHI = RHICmdList.CreateBuffer(CreateDesc);
 
 		return NewEntry;
 	}
diff --git a/Engine/Source/Runtime/Renderer/Private/OIT/OIT.cpp b/Engine/Source/Runtime/Renderer/Private/OIT/OIT.cpp
--- a/Engine/Source/Runtime/Renderer/Private/OIT/OIT.cpp
+++ b/Engine/Source/Runtime/Renderer/Private/OIT/OIT.cpp
@@ -287,8 +287,12 @@ public:
 		check(BytesPerElement == 2 || BytesPerElement == 4);
 		const EPixelFormat Format = BytesPerElement == 2 ? PF_R16_UINT : PF_R32_UINT;
 
-		FRHIResourceCreateInfo CreateInfo(DebugName);
-		IndexBufferRHI = RHICmdList.CreateBuffer(NumIndices * BytesPerElement, BUF_UnorderedAccess | BUF_ShaderResource | BUF_IndexBuffer, BytesPerElement /*Stride*/, ERHIAccess::VertexOrIndexBuffer, CreateInfo);
+		const FRHIBufferCreateDesc CreateDesc =
+			FRHIBufferCreateDesc::CreateIndex(DebugName, NumIndices * BytesPerElement, BytesPerElement)
+			.AddUsage(EBufferUsageFlags::UnorderedAccess | EBufferUsageFlags::ShaderResource)
+			.SetInitialState(ERHIAccess::VertexOrIndexBuffer);
+
+		IndexBufferRHI = RHICmdList.CreateBuffer(CreateDesc);
 		SortedIndexUAV = RHICmdList.CreateUnorderedAccessView(
 			IndexBufferRHI,
 			FRHIViewDesc::CreateBufferUAV()
diff --git a/Engine/Source/Runtime/Renderer/Private/PostProcess/SceneFilterRendering.cpp b/Engine/Source/Runtime/Renderer/Private/PostProcess/SceneFilterRendering.cpp
--- a/Engine/Source/Runtime/Renderer/Private/PostProcess/SceneFilterRendering.cpp
+++ b/Engine/Source/Runtime/Renderer/Private/PostProcess/SceneFilterRendering.cpp
@@ -237,11 +237,13 @@ void DrawTransformedRectangle(
 
 	// we don't do the triangle optimization as this case is rare for the DrawTransformedRectangle case
 
-	FRHIResourceCreateInfo CreateInfo(TEXT("DrawTransformedRectangle"));
-	FBufferRHIRef VertexBufferRHI = RHICmdList.CreateVertexBuffer(sizeof(FFilterVertex) * 4, BUF_Volatile, CreateInfo);
-	void* VoidPtr = RHICmdList.LockBuffer(VertexBufferRHI, 0, sizeof(FFilterVertex) * 4, RLM_WriteOnly);
+	const FRHIBufferCreateDesc CreateDesc =
+		FRHIBufferCreateDesc::CreateVertex<FFilterVertex>(TEXT("DrawTransformedRectangle"), 4)
+		.AddUsage(EBufferUsageFlags::Volatile)
+		.SetInitActionInitializer()
+		.DetermineInitialState();
 
-	FFilterVertex* Vertices = reinterpret_cast<FFilterVertex*>(VoidPtr);
+	TRHIBufferInitializer<FFilterVertex> Vertices = RHICmdList.CreateBufferInitializer(CreateDesc);
 
 	// LWC_TODO: precision loss
 	Vertices[0].Position = (FVector4f)PosTransform.TransformFVector4(FVector4(X,			Y,			ClipSpaceQuadZ,	1));
@@ -263,7 +265,8 @@ void DrawTransformedRectangle(
 		Vertices[VertexIndex].UV.Y = Vertices[VertexIndex].UV.Y / (float)TextureSize.Y;
 	}
 
-	RHICmdList.UnlockBuffer(VertexBufferRHI);
+	FBufferRHIRef VertexBufferRHI = Vertices.Finalize();
+
 	RHICmdList.SetStreamSource(0, VertexBufferRHI, 0);
 	RHICmdList.DrawIndexedPrimitive(GTwoTrianglesIndexBuffer.IndexBufferRHI, 0, 0, 4, 0, 2, 1);
 	VertexBufferRHI.SafeRelease();
diff --git a/Engine/Source/Runtime/Renderer/Private/RayTracing/RayTracingScene.cpp b/Engine/Source/Runtime/Renderer/Private/RayTracing/RayTracingScene.cpp
--- a/Engine/Source/Runtime/Renderer/Private/RayTracing/RayTracingScene.cpp
+++ b/Engine/Source/Runtime/Renderer/Private/RayTracing/RayTracingScene.cpp
@@ -172,8 +172,12 @@ void FRayTracingScene::Create(FRDGBuilder& GraphBuilder, const FViewInfo& View,
 				|| UploadBufferSize > Layer.InstanceGroupUploadBuffer->GetSize()
 				|| UploadBufferSize < Layer.InstanceGroupUploadBuffer->GetSize() / 2)
 			{
-				FRHIResourceCreateInfo CreateInfo(TEXT("FRayTracingScene::InstanceGroupUploadBuffer"));
-				Layer.InstanceGroupUploadBuffer = RHICmdList.CreateStructuredBuffer(sizeof(FRayTracingInstanceGroup), UploadBufferSize, BUF_ShaderResource | BUF_Volatile, CreateInfo);
+				const FRHIBufferCreateDesc CreateDesc =
+					FRHIBufferCreateDesc::CreateStructured(TEXT("FRayTracingScene::InstanceGroupUploadBuffer"), UploadBufferSize, sizeof(FRayTracingInstanceGroup))
+					.AddUsage(EBufferUsageFlags::ShaderResource | EBufferUsageFlags::Volatile)
+					.DetermineInitialState();
+
+				Layer.InstanceGroupUploadBuffer = RHICmdList.CreateBuffer(CreateDesc);
 				Layer.InstanceGroupUploadSRV = RHICmdList.CreateShaderResourceView(Layer.InstanceGroupUploadBuffer, FRHIViewDesc::CreateBufferSRV().SetTypeFromBuffer(Layer.InstanceGroupUploadBuffer));
 			}
 		}
@@ -186,8 +190,12 @@ void FRayTracingScene::Create(FRDGBuilder& GraphBuilder, const FViewInfo& View,
 				|| UploadBufferSize > Layer.InstanceUploadBuffer->GetSize()
 				|| UploadBufferSize < Layer.InstanceUploadBuffer->GetSize() / 2)
 			{
-				FRHIResourceCreateInfo CreateInfo(TEXT("FRayTracingScene::InstanceUploadBuffer"));
-				Layer.InstanceUploadBuffer = RHICmdList.CreateStructuredBuffer(sizeof(FRayTracingInstanceDescriptor), UploadBufferSize, BUF_ShaderResource | BUF_Volatile, CreateInfo);
+				const FRHIBufferCreateDesc CreateDesc =
+					FRHIBufferCreateDesc::CreateStructured(TEXT("FRayTracingScene::InstanceUploadBuffer"), UploadBufferSize, sizeof(FRayTracingInstanceDescriptor))
+					.AddUsage(EBufferUsageFlags::ShaderResource | EBufferUsageFlags::Volatile)
+					.DetermineInitialState();
+
+				Layer.InstanceUploadBuffer = RHICmdList.CreateBuffer(CreateDesc);
 				Layer.InstanceUploadSRV = RHICmdList.CreateShaderResourceView(Layer.InstanceUploadBuffer, FRHIViewDesc::CreateBufferSRV().SetTypeFromBuffer(Layer.InstanceUploadBuffer));
 			}
 		}
@@ -200,8 +208,12 @@ void FRayTracingScene::Create(FRDGBuilder& GraphBuilder, const FViewInfo& View,
 				|| UploadBufferSize > Layer.TransformUploadBuffer->GetSize()
 				|| UploadBufferSize < Layer.TransformUploadBuffer->GetSize() / 2)
 			{
-				FRHIResourceCreateInfo CreateInfo(TEXT("FRayTracingScene::TransformUploadBuffer"));
-				Layer.TransformUploadBuffer = RHICmdList.CreateStructuredBuffer(sizeof(FVector4f), UploadBufferSize, BUF_ShaderResource | BUF_Volatile, CreateInfo);
+				const FRHIBufferCreateDesc CreateDesc =
+					FRHIBufferCreateDesc::CreateStructured(TEXT("FRayTracingScene::TransformUploadBuffer"), UploadBufferSize, sizeof(FVector4f))
+					.AddUsage(EBufferUsageFlags::ShaderResource | EBufferUsageFlags::Volatile)
+					.DetermineInitialState();
+
+				Layer.TransformUploadBuffer = RHICmdList.CreateBuffer(CreateDesc);
 				Layer.TransformUploadSRV = RHICmdList.CreateShaderResourceView(Layer.TransformUploadBuffer, FRHIViewDesc::CreateBufferSRV().SetTypeFromBuffer(Layer.TransformUploadBuffer));
 			}
 		}
diff --git a/Engine/Source/Runtime/Renderer/Private/RayTracing/RaytracingSkylight.cpp b/Engine/Source/Runtime/Renderer/Private/RayTracing/RaytracingSkylight.cpp
--- a/Engine/Source/Runtime/Renderer/Private/RayTracing/RaytracingSkylight.cpp
+++ b/Engine/Source/Runtime/Renderer/Private/RayTracing/RaytracingSkylight.cpp
@@ -337,8 +337,6 @@ static void GenerateSkyLightVisibilityRays(
 	FIntVector& Dimensions
 )
 {
-	FRHIResourceCreateInfo CreateInfo(TEXT("SkyLightVisibilityRays"));
-
 	// Allocating mask of 256 x 256 rays
 	Dimensions = FIntVector(256, 256, 0);
 
diff --git a/Engine/Source/Runtime/Renderer/Private/RendererPrivateUtils.h b/Engine/Source/Runtime/Renderer/Private/RendererPrivateUtils.h
--- a/Engine/Source/Runtime/Renderer/Private/RendererPrivateUtils.h
+++ b/Engine/Source/Runtime/Renderer/Private/RendererPrivateUtils.h
@@ -23,10 +23,14 @@ public:
 
 	virtual void InitRHI(FRHICommandListBase& RHICmdList) override
 	{
-		const uint32 Size = sizeof(FVector2f) * 4 * NumTileQuadsInBuffer;
-		FRHIResourceCreateInfo CreateInfo(TEXT("FTileTexCoordVertexBuffer"));
-		VertexBufferRHI = RHICmdList.CreateBuffer(Size, BUF_Static | BUF_VertexBuffer, 0, ERHIAccess::VertexOrIndexBuffer, CreateInfo);
-		FVector2f* Vertices = (FVector2f*)RHICmdList.LockBuffer(VertexBufferRHI, 0, Size, RLM_WriteOnly);
+		const FRHIBufferCreateDesc CreateDesc =
+			FRHIBufferCreateDesc::CreateVertex<FVector2f>(TEXT("FTileTexCoordVertexBuffer"), 4 * NumTileQuadsInBuffer)
+			.AddUsage(EBufferUsageFlags::Static)
+			.SetInitialState(ERHIAccess::VertexOrIndexBuffer)
+			.SetInitActionInitializer();
+
+		TRHIBufferInitializer<FVector2f> Vertices = RHICmdList.CreateBufferInitializer(CreateDesc);
+
 		for (uint32 SpriteIndex = 0; SpriteIndex < NumTileQuadsInBuffer; ++SpriteIndex)
 		{
 			Vertices[SpriteIndex * 4 + 0] = FVector2f(0.0f, 0.0f);
@@ -34,7 +38,8 @@ public:
 			Vertices[SpriteIndex * 4 + 2] = FVector2f(1.0f, 1.0f);
 			Vertices[SpriteIndex * 4 + 3] = FVector2f(1.0f, 0.0f);
 		}
-		RHICmdList.UnlockBuffer(VertexBufferRHI);
+
+		VertexBufferRHI = Vertices.Finalize();
 	}
 
 	const uint32 NumTileQuadsInBuffer;
@@ -51,11 +56,14 @@ public:
 	/** Initialize the RHI for this rendering resource */
 	void InitRHI(FRHICommandListBase& RHICmdList) override
 	{
-		const uint32 Size = sizeof(uint16) * 6 * NumTileQuadsInBuffer;
-		const uint32 Stride = sizeof(uint16);
-		FRHIResourceCreateInfo CreateInfo(TEXT("FTileIndexBuffer"));
-		IndexBufferRHI = RHICmdList.CreateBuffer(Size, BUF_Static | BUF_IndexBuffer, Stride, ERHIAccess::VertexOrIndexBuffer, CreateInfo);
-		uint16* Indices = (uint16*)RHICmdList.LockBuffer(IndexBufferRHI, 0, Size, RLM_WriteOnly);
+		const FRHIBufferCreateDesc CreateDesc =
+			FRHIBufferCreateDesc::CreateIndex<uint16>(TEXT("FTileIndexBuffer"), 6 * NumTileQuadsInBuffer)
+			.AddUsage(EBufferUsageFlags::Static)
+			.SetInitialState(ERHIAccess::VertexOrIndexBuffer)
+			.SetInitActionInitializer();
+
+		TRHIBufferInitializer<uint16> Indices = RHICmdList.CreateBufferInitializer(CreateDesc);
+
 		for (uint32 SpriteIndex = 0; SpriteIndex < NumTileQuadsInBuffer; ++SpriteIndex)
 		{
 			Indices[SpriteIndex * 6 + 0] = SpriteIndex * 4 + 0;
@@ -65,7 +73,8 @@ public:
 			Indices[SpriteIndex * 6 + 4] = SpriteIndex * 4 + 2;
 			Indices[SpriteIndex * 6 + 5] = SpriteIndex * 4 + 3;
 		}
-		RHICmdList.UnlockBuffer(IndexBufferRHI);
+
+		IndexBufferRHI = Indices.Finalize();
 	}
 
 	const uint32 NumTileQuadsInBuffer;
diff --git a/Engine/Source/Runtime/Renderer/Private/SceneOcclusion.cpp b/Engine/Source/Runtime/Renderer/Private/SceneOcclusion.cpp
--- a/Engine/Source/Runtime/Renderer/Private/SceneOcclusion.cpp
+++ b/Engine/Source/Runtime/Renderer/Private/SceneOcclusion.cpp
@@ -341,18 +341,23 @@ public:
 		const uint32 Stride = sizeof(uint16);
 		const uint32 SizeInBytes = MaxBatchedPrimitives * NUM_CUBE_VERTICES * Stride;
 
-		FRHIResourceCreateInfo CreateInfo(TEXT("FOcclusionQueryIndexBuffer"));
-		IndexBufferRHI = RHICmdList.CreateBuffer(SizeInBytes, BUF_Static | BUF_IndexBuffer, Stride, ERHIAccess::VertexOrIndexBuffer, CreateInfo);
-		uint16* RESTRICT Indices = (uint16*) RHICmdList.LockBuffer(IndexBufferRHI, 0, SizeInBytes, RLM_WriteOnly);
+		const FRHIBufferCreateDesc CreateDesc =
+			FRHIBufferCreateDesc::CreateIndex(TEXT("FOcclusionQueryIndexBuffer"), SizeInBytes, Stride)
+			.AddUsage(EBufferUsageFlags::Static)
+			.SetInitialState(ERHIAccess::VertexOrIndexBuffer)
+			.SetInitActionInitializer();
 
-		for(uint32 PrimitiveIndex = 0;PrimitiveIndex < MaxBatchedPrimitives;PrimitiveIndex++)
+		TRHIBufferInitializer<uint16> Indices = RHICmdList.CreateBufferInitializer(CreateDesc);
+
+		for (uint32 PrimitiveIndex = 0; PrimitiveIndex < MaxBatchedPrimitives; PrimitiveIndex++)
 		{
 			for(int32 Index = 0;Index < NUM_CUBE_VERTICES;Index++)
 			{
 				Indices[PrimitiveIndex * NUM_CUBE_VERTICES + Index] = PrimitiveIndex * 8 + GCubeIndices[Index];
 			}
 		}
-		RHICmdList.UnlockBuffer(IndexBufferRHI);
+
+		IndexBufferRHI = Indices.Finalize();
 	}
 };
 TGlobalResource<FOcclusionQueryIndexBuffer> GOcclusionQueryIndexBuffer;
@@ -612,7 +617,7 @@ static void ExecutePointLightShadowOcclusionQuery(FRHICommandList& RHICmdList, F
 	RHICmdList.EndRenderQuery(ShadowOcclusionQuery);
 }
 
-static void PrepareDirectionalLightShadowOcclusionQuery(uint32& BaseVertexIndex, FVector3f* DestinationBuffer, const FViewInfo& View, const FProjectedShadowInfo& ProjectedShadowInfo)
+static void PrepareDirectionalLightShadowOcclusionQuery(uint32& BaseVertexIndex, TRHIBufferInitializer<FVector3f>& DestinationBuffer, const FViewInfo& View, const FProjectedShadowInfo& ProjectedShadowInfo)
 {
 	const FViewMatrices& ViewMatrices = View.GetShadowViewMatrices();
 
@@ -657,13 +662,12 @@ static void ExecuteDirectionalLightShadowOcclusionQuery(FRHICommandList& RHICmdL
 	RHICmdList.EndRenderQuery(ShadowOcclusionQuery);
 }
 
-static void PrepareProjectedShadowOcclusionQuery(uint32& BaseVertexIndex, FVector3f* DestinationBuffer, const FViewInfo& View, const FProjectedShadowInfo& ProjectedShadowInfo)
+static void PrepareProjectedShadowOcclusionQuery(uint32& BaseVertexIndex, TRHIBufferInitializer<FVector3f>& DestinationBuffer, const FViewInfo& View, const FProjectedShadowInfo& ProjectedShadowInfo)
 {
 	// The shadow transforms and view transforms are relative to different origins, so the world coordinates need to
 	// be translated.
 	const FVector4f PreShadowToPreViewTranslation(FVector3f(View.ViewMatrices.GetPreViewTranslation() - ProjectedShadowInfo.PreShadowTranslation), 0);
 
-	FVector3f* Vertices = &DestinationBuffer[BaseVertexIndex];
 	// Generate vertices for the shadow's frustum.
 	for (uint32 Z = 0; Z < 2; Z++)
 	{
@@ -679,7 +683,9 @@ static void PrepareProjectedShadowOcclusionQuery(uint32& BaseVertexIndex, FVecto
 						1.0f)
 				);
 				const FVector3f ProjectedVertex = UnprojectedVertex / UnprojectedVertex.W + PreShadowToPreViewTranslation;
-				Vertices[GetCubeVertexIndex(X, Y, Z)] = ProjectedVertex;
+				const uint16 CubeVertexIndex = GetCubeVertexIndex(X, Y, Z);
+
+				DestinationBuffer.WriteValueAtIndex(BaseVertexIndex + CubeVertexIndex, ProjectedVertex);
 			}
 		}
 	}
@@ -746,20 +752,19 @@ static bool AllocatePlanarReflectionOcclusionQuery(const FViewInfo& View, const
 	return bAllowBoundsTest;
 }
 
-static void PreparePlanarReflectionOcclusionQuery(uint32& BaseVertexIndex, FVector3f* DestinationBuffer, const FViewInfo& View, const FPlanarReflectionSceneProxy* SceneProxy)
+static void PreparePlanarReflectionOcclusionQuery(uint32& BaseVertexIndex, TRHIBufferInitializer<FVector3f>& DestinationBuffer, const FViewInfo& View, const FPlanarReflectionSceneProxy* SceneProxy)
 {
-	float* Vertices = (float*)(&DestinationBuffer[BaseVertexIndex]);
-
 	const FVector3f PrimitiveBoxMin = FVector3f(SceneProxy->WorldBounds.Min + View.ViewMatrices.GetPreViewTranslation());
 	const FVector3f PrimitiveBoxMax = FVector3f(SceneProxy->WorldBounds.Max + View.ViewMatrices.GetPreViewTranslation());
-	Vertices[0] = PrimitiveBoxMin.X; Vertices[1] = PrimitiveBoxMin.Y; Vertices[2] = PrimitiveBoxMin.Z;
-	Vertices[3] = PrimitiveBoxMin.X; Vertices[4] = PrimitiveBoxMin.Y; Vertices[5] = PrimitiveBoxMax.Z;
-	Vertices[6] = PrimitiveBoxMin.X; Vertices[7] = PrimitiveBoxMax.Y; Vertices[8] = PrimitiveBoxMin.Z;
-	Vertices[9] = PrimitiveBoxMin.X; Vertices[10] = PrimitiveBoxMax.Y; Vertices[11] = PrimitiveBoxMax.Z;
-	Vertices[12] = PrimitiveBoxMax.X; Vertices[13] = PrimitiveBoxMin.Y; Vertices[14] = PrimitiveBoxMin.Z;
-	Vertices[15] = PrimitiveBoxMax.X; Vertices[16] = PrimitiveBoxMin.Y; Vertices[17] = PrimitiveBoxMax.Z;
-	Vertices[18] = PrimitiveBoxMax.X; Vertices[19] = PrimitiveBoxMax.Y; Vertices[20] = PrimitiveBoxMin.Z;
-	Vertices[21] = PrimitiveBoxMax.X; Vertices[22] = PrimitiveBoxMax.Y; Vertices[23] = PrimitiveBoxMax.Z;
+
+	DestinationBuffer.WriteValueAtIndex(BaseVertexIndex + 0, FVector3f(PrimitiveBoxMin.X, PrimitiveBoxMin.Y, PrimitiveBoxMin.Z));
+	DestinationBuffer.WriteValueAtIndex(BaseVertexIndex + 1, FVector3f(PrimitiveBoxMin.X, PrimitiveBoxMin.Y, PrimitiveBoxMax.Z));
+	DestinationBuffer.WriteValueAtIndex(BaseVertexIndex + 2, FVector3f(PrimitiveBoxMin.X, PrimitiveBoxMax.Y, PrimitiveBoxMin.Z));
+	DestinationBuffer.WriteValueAtIndex(BaseVertexIndex + 3, FVector3f(PrimitiveBoxMin.X, PrimitiveBoxMax.Y, PrimitiveBoxMax.Z));
+	DestinationBuffer.WriteValueAtIndex(BaseVertexIndex + 4, FVector3f(PrimitiveBoxMax.X, PrimitiveBoxMin.Y, PrimitiveBoxMin.Z));
+	DestinationBuffer.WriteValueAtIndex(BaseVertexIndex + 5, FVector3f(PrimitiveBoxMax.X, PrimitiveBoxMin.Y, PrimitiveBoxMax.Z));
+	DestinationBuffer.WriteValueAtIndex(BaseVertexIndex + 6, FVector3f(PrimitiveBoxMax.X, PrimitiveBoxMax.Y, PrimitiveBoxMin.Z));
+	DestinationBuffer.WriteValueAtIndex(BaseVertexIndex + 7, FVector3f(PrimitiveBoxMax.X, PrimitiveBoxMax.Y, PrimitiveBoxMax.Z));
 
 	BaseVertexIndex += 8;
 }
@@ -1299,32 +1304,36 @@ static void BeginOcclusionTests(
 		if (NumVertices > 0)
 		{
 			uint32 BaseVertexOffset = 0;
-			FRHIResourceCreateInfo CreateInfo(TEXT("ViewOcclusionTests"));
-			FBufferRHIRef VertexBufferRHI = RHICmdList.CreateBuffer(sizeof(FVector3f) * NumVertices, BUF_Volatile | BUF_VertexBuffer, 0, ERHIAccess::VertexOrIndexBuffer, CreateInfo);
-			void* VoidPtr = RHICmdList.LockBuffer(VertexBufferRHI, 0, sizeof(FVector3f) * NumVertices, RLM_WriteOnly);
+
+			const FRHIBufferCreateDesc CreateDesc =
+				FRHIBufferCreateDesc::CreateVertex<FVector3f>(TEXT("ViewOcclusionTests"), NumVertices)
+				.AddUsage(EBufferUsageFlags::Volatile)
+				.SetInitialState(ERHIAccess::VertexOrIndexBuffer)
+				.SetInitActionInitializer();
+
+			TRHIBufferInitializer<FVector3f> Initializer = RHICmdList.CreateBufferInitializer(CreateDesc);
 
 			{
-				FVector3f* Vertices = reinterpret_cast<FVector3f*>(VoidPtr);
 				for (FProjectedShadowInfo const* Query : ViewQuery.CSMQueryInfos)
 				{
-					PrepareDirectionalLightShadowOcclusionQuery(BaseVertexOffset, Vertices, View, *Query);
+					PrepareDirectionalLightShadowOcclusionQuery(BaseVertexOffset, Initializer, View, *Query);
 					checkSlow(BaseVertexOffset <= NumVertices);
 				}
 
 				for (FProjectedShadowInfo const* Query : ViewQuery.ShadowQueryInfos)
 				{
-					PrepareProjectedShadowOcclusionQuery(BaseVertexOffset, Vertices, View, *Query);
+					PrepareProjectedShadowOcclusionQuery(BaseVertexOffset, Initializer, View, *Query);
 					checkSlow(BaseVertexOffset <= NumVertices);
 				}
 
 				for (FPlanarReflectionSceneProxy const* Query : ViewQuery.ReflectionQueryInfos)
 				{
-					PreparePlanarReflectionOcclusionQuery(BaseVertexOffset, Vertices, View, Query);
+					PreparePlanarReflectionOcclusionQuery(BaseVertexOffset, Initializer, View, Query);
 					checkSlow(BaseVertexOffset <= NumVertices);
 				}
 			}
 
-			RHICmdList.UnlockBuffer(VertexBufferRHI);
+			FBufferRHIRef VertexBufferRHI = Initializer.Finalize();
 
 			{
 				SCOPED_DRAW_EVENT(RHICmdList, ShadowFrustumQueries);
diff --git a/Engine/Source/Runtime/Renderer/Private/SceneRendering.cpp b/Engine/Source/Runtime/Renderer/Private/SceneRendering.cpp
--- a/Engine/Source/Runtime/Renderer/Private/SceneRendering.cpp
+++ b/Engine/Source/Runtime/Renderer/Private/SceneRendering.cpp
@@ -5710,13 +5710,13 @@ class FDummySceneColorResolveBuffer : public FVertexBuffer
 public:
 	virtual void InitRHI(FRHICommandListBase& RHICmdList) override
 	{
-		const int32 NumDummyVerts = 3;
-		const uint32 Size = sizeof(FVector4f) * NumDummyVerts;
-		FRHIResourceCreateInfo CreateInfo(TEXT("FDummySceneColorResolveBuffer"));
-		VertexBufferRHI = RHICmdList.CreateBuffer(Size, BUF_Static | BUF_VertexBuffer, 0, ERHIAccess::VertexOrIndexBuffer, CreateInfo);
-		void* BufferData = RHICmdList.LockBuffer(VertexBufferRHI, 0, Size, RLM_WriteOnly);
-		FMemory::Memset(BufferData, 0, Size);		
-		RHICmdList.UnlockBuffer(VertexBufferRHI);		
+		const FRHIBufferCreateDesc CreateDesc =
+			FRHIBufferCreateDesc::CreateVertex<FVector4f>(TEXT("FDummySceneColorResolveBuffer"), 3)
+			.AddUsage(EBufferUsageFlags::Static)
+			.SetInitialState(ERHIAccess::VertexOrIndexBuffer)
+			.SetInitActionZeroData();
+
+		VertexBufferRHI = RHICmdList.CreateBuffer(CreateDesc);
 	}
 };
 
diff --git a/Engine/Source/Runtime/Renderer/Private/ShaderPrint.cpp b/Engine/Source/Runtime/Renderer/Private/ShaderPrint.cpp
--- a/Engine/Source/Runtime/Renderer/Private/ShaderPrint.cpp
+++ b/Engine/Source/Runtime/Renderer/Private/ShaderPrint.cpp
@@ -16,6 +16,7 @@
 #include "ScreenPass.h"
 #include "SystemTextures.h"
 #include "RenderGraphUtils.h"
+#include "RHIResourceUtils.h"
 
 namespace ShaderPrint
 {
@@ -196,32 +197,26 @@ namespace ShaderPrint
 	//////////////////////////////////////////////////////////////////////////////////////////////////
 	// Global fixed index buffers
 
+	static const uint16 GLineIndexBufferData[2] = { 0, 1 };
+
 	class FLineIndexBuffer : public FIndexBuffer
 	{
 	public:
 		void InitRHI(FRHICommandListBase& RHICmdList) override
 		{
-			FRHIResourceCreateInfo CreateInfo(TEXT("FLineIndexBuffer"));
-			IndexBufferRHI = RHICmdList.CreateIndexBuffer(sizeof(uint16), sizeof(uint16) * 2, BUF_Static, CreateInfo);
-			void* VoidPtr = RHICmdList.LockBuffer(IndexBufferRHI, 0, sizeof(uint16) * 2, RLM_WriteOnly);
-			static const uint16 Indices[] = { 0, 1 };
-			FMemory::Memcpy(VoidPtr, Indices, 2 * sizeof(uint16));
-			RHICmdList.UnlockBuffer(IndexBufferRHI);
+			IndexBufferRHI = UE::RHIResourceUtils::CreateIndexBufferFromArray(RHICmdList, TEXT("FLineIndexBuffer"), EBufferUsageFlags::Static, MakeConstArrayView(GLineIndexBufferData));
 		}
 	};
 	TGlobalResource<FLineIndexBuffer, FRenderResource::EInitPhase::Pre> GLineIndexBuffer;
 
+	static const uint16 GTriangleIndexBufferData[3] = { 0, 1, 2 };
+
 	class FTriangleIndexBuffer : public FIndexBuffer
 	{
 	public:
 		void InitRHI(FRHICommandListBase& RHICmdList) override
 		{
-			FRHIResourceCreateInfo CreateInfo(TEXT("FTriangleIndexBuffer"));
-			IndexBufferRHI = RHICmdList.CreateIndexBuffer(sizeof(uint16), sizeof(uint16) * 3, BUF_Static, CreateInfo);
-			void* VoidPtr = RHICmdList.LockBuffer(IndexBufferRHI, 0, sizeof(uint16) * 3, RLM_WriteOnly);
-			static const uint16 Indices[] = { 0, 1, 2 };
-			FMemory::Memcpy(VoidPtr, Indices, 3 * sizeof(uint16));
-			RHICmdList.UnlockBuffer(IndexBufferRHI);
+			IndexBufferRHI = UE::RHIResourceUtils::CreateIndexBufferFromArray(RHICmdList, TEXT("FTriangleIndexBuffer"), EBufferUsageFlags::Static, MakeConstArrayView(GTriangleIndexBufferData));
 		}
 	};
 	TGlobalResource<FTriangleIndexBuffer, FRenderResource::EInitPhase::Pre> GTriangleIndexBuffer;
diff --git a/Engine/Source/Runtime/Renderer/Private/ShadowRendering.cpp b/Engine/Source/Runtime/Renderer/Private/ShadowRendering.cpp
--- a/Engine/Source/Runtime/Renderer/Private/ShadowRendering.cpp
+++ b/Engine/Source/Runtime/Renderer/Private/ShadowRendering.cpp
@@ -19,6 +19,7 @@
 #include "RenderCore.h"
 #include "TranslucentLighting.h"
 #include "MobileBasePassRendering.h"
+#include "RHIResourceUtils.h"
 
 using namespace LightFunctionAtlas;
 using namespace ShadowRendering;
@@ -198,6 +199,25 @@ static TAutoConsoleVariable<float> CVarShadowMaxSlopeScaleDepthBias(
 
 FForwardScreenSpaceShadowMaskTextureMobileOutputs GScreenSpaceShadowMaskTextureMobileOutputs;
 
+const FVector4f GDummyWholeSceneDirectionalShadowStencilVertexBufferData[] =
+{
+	// Far Plane
+	FVector4f( 1,  1,  1 /* StencilFar */),
+	FVector4f(-1,  1,  1),
+	FVector4f( 1, -1,  1),
+	FVector4f( 1, -1,  1),
+	FVector4f(-1,  1,  1),
+	FVector4f(-1, -1,  1),
+
+	// Near Plane
+	FVector4f(-1,  1, -1 /* StencilNear */),
+	FVector4f( 1,  1, -1),
+	FVector4f(-1, -1, -1),
+	FVector4f(-1, -1, -1),
+	FVector4f( 1,  1, -1),
+	FVector4f( 1, -1, -1),
+};
+
 /**
  * A dummy vertex buffer to bind when rendering whole scene shadows. This
  * prevents some D3D debug warnings about zero-element input layouts, but is not
@@ -206,30 +226,9 @@ FForwardScreenSpaceShadowMaskTextureMobileOutputs GScreenSpaceShadowMaskTextureM
 class FDummyWholeSceneDirectionalShadowStencilVertexBuffer : public FVertexBuffer
 {
 public:
-
 	virtual void InitRHI(FRHICommandListBase& RHICmdList) override
 	{
-		FRHIResourceCreateInfo CreateInfo(TEXT("FDummyWholeSceneDirectionalShadowStencilVertexBuffer"));
-		VertexBufferRHI = RHICmdList.CreateBuffer(sizeof(FVector4f) * 12, BUF_Static | BUF_VertexBuffer, 0, ERHIAccess::VertexOrIndexBuffer, CreateInfo);
-		FVector4f* DummyContents = (FVector4f*)RHICmdList.LockBuffer(VertexBufferRHI, 0, sizeof(FVector4f) * 12, RLM_WriteOnly);
-
-		// Far Plane
-		DummyContents[0] = FVector4f( 1,  1,  1 /* StencilFar */);
-		DummyContents[1] = FVector4f(-1,  1,  1);
-		DummyContents[2] = FVector4f( 1, -1,  1);
-		DummyContents[3] = FVector4f( 1, -1,  1);
-		DummyContents[4] = FVector4f(-1,  1,  1);
-		DummyContents[5] = FVector4f(-1, -1,  1);
-
-		// Near Plane
-		DummyContents[6]  = FVector4f(-1,  1, -1 /* StencilNear */);
-		DummyContents[7]  = FVector4f( 1,  1, -1);
-		DummyContents[8]  = FVector4f(-1, -1, -1);
-		DummyContents[9]  = FVector4f(-1, -1, -1);
-		DummyContents[10] = FVector4f( 1,  1, -1);
-		DummyContents[11] = FVector4f( 1, -1, -1);
-
-		RHICmdList.UnlockBuffer(VertexBufferRHI);
+		VertexBufferRHI = UE::RHIResourceUtils::CreateVertexBufferFromArray(RHICmdList, TEXT("FDummyWholeSceneDirectionalShadowStencilVertexBuffer"), EBufferUsageFlags::Static, MakeConstArrayView(GDummyWholeSceneDirectionalShadowStencilVertexBufferData));
 	}
 };
 
@@ -787,10 +786,14 @@ class FFrustumVertexBuffer : public FVertexBuffer
 public:
 	virtual void InitRHI(FRHICommandListBase& RHICmdList) override
 	{
-		FRHIResourceCreateInfo CreateInfo(TEXT("FProjectedShadowInfoStencilFrustum"));
-		VertexBufferRHI = RHICmdList.CreateVertexBuffer(sizeof(FVector4f) * 8, BUF_Static, CreateInfo);
-		FVector4f* OutFrustumVertices = reinterpret_cast<FVector4f*>(RHICmdList.LockBuffer(VertexBufferRHI, 0, sizeof(FVector4f) * 8, RLM_WriteOnly));
-		
+		const FRHIBufferCreateDesc CreateDesc =
+			FRHIBufferCreateDesc::CreateVertex<FVector4f>(TEXT("FProjectedShadowInfoStencilFrustum"), 8)
+			.AddUsage(EBufferUsageFlags::Static)
+			.SetInitActionInitializer()
+			.DetermineInitialState();
+
+		TRHIBufferInitializer<FVector4f> OutFrustumVertices = RHICmdList.CreateBufferInitializer(CreateDesc);
+
 		for(uint32 vZ = 0;vZ < 2;vZ++)
 		{
 			for(uint32 vY = 0;vY < 2;vY++)
@@ -806,7 +809,7 @@ public:
 			}
 		}
 
-		RHICmdList.UnlockBuffer(VertexBufferRHI);
+		VertexBufferRHI = OutFrustumVertices.Finalize();
 	}
 };
 TGlobalResource<FFrustumVertexBuffer> GFrustumVertexBuffer;
diff --git a/Engine/Source/Runtime/Renderer/Private/SkyAtmosphereRendering.cpp b/Engine/Source/Runtime/Renderer/Private/SkyAtmosphereRendering.cpp
--- a/Engine/Source/Runtime/Renderer/Private/SkyAtmosphereRendering.cpp
+++ b/Engine/Source/Runtime/Renderer/Private/SkyAtmosphereRendering.cpp
@@ -1099,30 +1099,31 @@ public:
 		const uint32 GroupSize = GetSampletCount();
 		const float GroupSizeInv = 1.0f / float(GroupSize);
 
-		UniformSphereSamplesBuffer.Initialize(RHICmdList, TEXT("UniformSphereSamplesBuffer"), sizeof(FVector4f), GroupSize * GroupSize, EPixelFormat::PF_A32B32G32R32F, BUF_Static);
-		FVector4f* Dest = (FVector4f*)RHICmdList.LockBuffer(UniformSphereSamplesBuffer.Buffer, 0, sizeof(FVector4f)*GroupSize*GroupSize, RLM_WriteOnly);
-
-		FRandomStream RandomStream(0xDE4DC0DE);
-		for (uint32 i = 0; i < GroupSize; ++i)
-		{
-			for (uint32 j = 0; j < GroupSize; ++j)
+		UniformSphereSamplesBuffer.InitializeWithData(RHICmdList, TEXT("UniformSphereSamplesBuffer"), sizeof(FVector4f), GroupSize * GroupSize, PF_A32B32G32R32F, EBufferUsageFlags::Static,
+			[GroupSize, GroupSizeInv](FRHIBufferInitializer& Initializer)
 			{
-				const float u0 = (float(i) + RandomStream.GetFraction()) * GroupSizeInv;
-				const float u1 = (float(j) + RandomStream.GetFraction()) * GroupSizeInv;
-		
-				const float a = 1.0f - 2.0f * u0;
-				const float b = FMath::Sqrt(1.0f - a*a);
-				const float phi = 2 * PI * u1;
-		
-				uint32 idx = j * GroupSize + i;
-				Dest[idx].X = b * FMath::Cos(phi);
-				Dest[idx].Y = b * FMath::Sin(phi);
-				Dest[idx].Z = a;
-				Dest[idx].W = 0.0f;
-			}
-		}
+				TArrayView<FVector4f> Dest = Initializer.GetWriteView<FVector4f>();
 
-		RHICmdList.UnlockBuffer(UniformSphereSamplesBuffer.Buffer);
+				FRandomStream RandomStream(0xDE4DC0DE);
+				for (uint32 i = 0; i < GroupSize; ++i)
+				{
+					for (uint32 j = 0; j < GroupSize; ++j)
+					{
+						const float u0 = (float(i) + RandomStream.GetFraction()) * GroupSizeInv;
+						const float u1 = (float(j) + RandomStream.GetFraction()) * GroupSizeInv;
+
+						const float a = 1.0f - 2.0f * u0;
+						const float b = FMath::Sqrt(1.0f - a * a);
+						const float phi = 2 * PI * u1;
+
+						uint32 idx = j * GroupSize + i;
+						Dest[idx].X = b * FMath::Cos(phi);
+						Dest[idx].Y = b * FMath::Sin(phi);
+						Dest[idx].Z = a;
+						Dest[idx].W = 0.0f;
+					}
+				}
+			});
 	}
 
 	virtual void ReleaseRHI()
diff --git a/Engine/Source/Runtime/Renderer/Private/SpriteIndexBuffer.h b/Engine/Source/Runtime/Renderer/Private/SpriteIndexBuffer.h
--- a/Engine/Source/Runtime/Renderer/Private/SpriteIndexBuffer.h
+++ b/Engine/Source/Runtime/Renderer/Private/SpriteIndexBuffer.h
@@ -12,11 +12,14 @@ class FSpriteIndexBuffer : public FIndexBuffer
 public:
 	virtual void InitRHI(FRHICommandListBase& RHICmdList) override
 	{
-		const uint32 Size = sizeof(uint16) * 6 * NumSprites;
-		const uint32 Stride = sizeof(uint16);
-		FRHIResourceCreateInfo CreateInfo(TEXT("FSpriteIndexBuffer"));
-		IndexBufferRHI = RHICmdList.CreateBuffer( Size, BUF_Static | BUF_IndexBuffer, Stride, ERHIAccess::VertexOrIndexBuffer, CreateInfo );
-		uint16* Indices = (uint16*)RHICmdList.LockBuffer( IndexBufferRHI, 0, Size, RLM_WriteOnly );
+		const FRHIBufferCreateDesc CreateDesc =
+			FRHIBufferCreateDesc::CreateIndex<uint16>(TEXT("FSpriteIndexBuffer"), 6 * NumSprites)
+			.AddUsage(EBufferUsageFlags::Static)
+			.SetInitialState(ERHIAccess::VertexOrIndexBuffer)
+			.SetInitActionInitializer();
+
+		TRHIBufferInitializer<uint16> Indices = RHICmdList.CreateBufferInitializer(CreateDesc);
+
 		for (uint32 SpriteIndex = 0; SpriteIndex < NumSprites; ++SpriteIndex)
 		{
 			Indices[SpriteIndex*6 + 0] = SpriteIndex*4 + 0;
@@ -26,6 +29,7 @@ public:
 			Indices[SpriteIndex*6 + 4] = SpriteIndex*4 + 1;
 			Indices[SpriteIndex*6 + 5] = SpriteIndex*4 + 3;
 		}
-		RHICmdList.UnlockBuffer( IndexBufferRHI );
+
+		IndexBufferRHI = Indices.Finalize();
 	}
 };
\ No newline at end of file
diff --git a/Engine/Source/Runtime/Renderer/Private/SystemTextures.cpp b/Engine/Source/Runtime/Renderer/Private/SystemTextures.cpp
--- a/Engine/Source/Runtime/Renderer/Private/SystemTextures.cpp
+++ b/Engine/Source/Runtime/Renderer/Private/SystemTextures.cpp
@@ -1678,8 +1678,10 @@ FRDGBufferRef GetInternalDefaultBuffer(
 
 	const uint32 BufferSize = NumBytePerElement * NumElements;
 
+	const TCHAR* BufferName = TEXT("");
+
+
 	FRDGBufferDesc BufferDesc;
-	FRHIResourceCreateInfo CreateInfo(TEXT(""));
 	FRHICommandListBase& RHICmdList = GraphBuilder.RHICmdList;
 
 	// Adding new buffer if there is no fit (slow path)
@@ -1687,25 +1689,46 @@ FRDGBufferRef GetInternalDefaultBuffer(
 	switch (BufferType)
 	{
 	case EDefaultBufferType::VertexBuffer:
-		CreateInfo.DebugName = TEXT("DefaultBuffer");
+	{
+		BufferName = TEXT("DefaultBuffer");
 		BufferDesc = FRDGBufferDesc::CreateUploadDesc(NumBytePerElement, NumElements);
-		RHIBuffer = RHICmdList.CreateVertexBuffer(BufferSize, BUF_Static | BUF_ShaderResource, CreateInfo);
+
+		const FRHIBufferCreateDesc BufferCreateDesc =
+			FRHIBufferCreateDesc::CreateVertex(BufferName, BufferSize)
+			.AddUsage(EBufferUsageFlags::Static | EBufferUsageFlags::ShaderResource)
+			.DetermineInitialState();
+		RHIBuffer = RHICmdList.CreateBuffer(BufferCreateDesc);
+	}
 		break;
 
 	case EDefaultBufferType::StructuredBuffer:
-		CreateInfo.DebugName = TEXT("DefaultStructuredBuffer");
+	{
+		BufferName = TEXT("DefaultStructuredBuffer");
 		BufferDesc = FRDGBufferDesc::CreateStructuredDesc(NumBytePerElement, NumElements);
 		// Remove the UAV flag, as default resources are supposed to be read-only.
 		EnumRemoveFlags(BufferDesc.Usage, EBufferUsageFlags::UnorderedAccess);
-		RHIBuffer = RHICmdList.CreateStructuredBuffer(NumBytePerElement, BufferSize, BufferDesc.Usage, CreateInfo);
+		
+		const FRHIBufferCreateDesc BufferCreateDesc =
+			FRHIBufferCreateDesc::CreateStructured(BufferName, BufferSize, NumBytePerElement)
+			.AddUsage(BufferDesc.Usage)
+			.DetermineInitialState();
+		RHIBuffer = RHICmdList.CreateBuffer(BufferCreateDesc);
+	}
 		break;
 
 	case EDefaultBufferType::ByteAddressBuffer:
-		CreateInfo.DebugName = TEXT("DefaultByteAddressBuffer");
+	{
+		BufferName = TEXT("DefaultByteAddressBuffer");
 		BufferDesc = FRDGBufferDesc::CreateByteAddressDesc(BufferSize);
 		// Same as above.
 		EnumRemoveFlags(BufferDesc.Usage, EBufferUsageFlags::UnorderedAccess);
-		RHIBuffer = RHICmdList.CreateStructuredBuffer(NumBytePerElement, BufferSize, BufferDesc.Usage, CreateInfo);
+
+		const FRHIBufferCreateDesc BufferCreateDesc =
+			FRHIBufferCreateDesc::CreateStructured(BufferName, BufferSize, NumBytePerElement)
+			.AddUsage(BufferDesc.Usage)
+			.DetermineInitialState();
+		RHIBuffer = RHICmdList.CreateBuffer(BufferCreateDesc);
+	}
 		break;
 
 	}
@@ -1734,7 +1757,7 @@ FRDGBufferRef GetInternalDefaultBuffer(
 	FDefaultBuffer Entry;
 	Entry.Key = Key;
 	Entry.Hash = Hash;
-	Entry.Buffer = new FRDGPooledBuffer(RHICmdList, RHIBuffer, BufferDesc, NumElements, CreateInfo.DebugName);
+	Entry.Buffer = new FRDGPooledBuffer(RHICmdList, RHIBuffer, BufferDesc, NumElements, BufferName);
 
 	Index = DefaultBuffers.Add(Entry);
 	HashDefaultBuffers.Add(Hash, Index);
diff --git a/Engine/Source/Runtime/Renderer/Private/Tests/RayTracingTestbed.cpp b/Engine/Source/Runtime/Renderer/Private/Tests/RayTracingTestbed.cpp
--- a/Engine/Source/Runtime/Renderer/Private/Tests/RayTracingTestbed.cpp
+++ b/Engine/Source/Runtime/Renderer/Private/Tests/RayTracingTestbed.cpp
@@ -121,13 +121,12 @@ bool RunRayTracingTestbed_RenderThread(const FString& Parameters)
 	FUnorderedAccessViewRHIRef OcclusionResultBufferView;
 
 	{
-		FRHIResourceCreateInfo CreateInfo(TEXT("OcclusionResultBuffer"));
-		OcclusionResultBuffer = RHICmdList.CreateBuffer(sizeof(uint32) * NumRays, 
-			BUF_Static | BUF_UnorderedAccess | BUF_StructuredBuffer, 
-			sizeof(uint32), 
-			ERHIAccess::UAVMask, 
-			CreateInfo
-		);
+		const FRHIBufferCreateDesc CreateDesc =
+			FRHIBufferCreateDesc::CreateStructured<uint32>(TEXT("OcclusionResultBuffer"), NumRays)
+			.AddUsage(EBufferUsageFlags::Static | EBufferUsageFlags::UnorderedAccess)
+			.SetInitialState(ERHIAccess::UAVMask);
+
+		OcclusionResultBuffer = RHICmdList.CreateBuffer(CreateDesc);
 		OcclusionResultBufferView = RHICmdList.CreateUnorderedAccessView(OcclusionResultBuffer, 
 			FRHIViewDesc::CreateBufferUAV()
 			.SetType(FRHIViewDesc::EBufferType::Structured)
@@ -140,13 +139,12 @@ bool RunRayTracingTestbed_RenderThread(const FString& Parameters)
 	FUnorderedAccessViewRHIRef IntersectionResultBufferView;
 
 	{
-		FRHIResourceCreateInfo CreateInfo(TEXT("IntersectionResultBuffer"));
-		IntersectionResultBuffer = RHICmdList.CreateBuffer(sizeof(FBasicRayTracingIntersectionResult) * NumRays, 
-			BUF_Static | BUF_UnorderedAccess | BUF_StructuredBuffer, 
-			sizeof(FBasicRayTracingIntersectionResult), 
-			ERHIAccess::UAVMask, 
-			CreateInfo
-		);
+		const FRHIBufferCreateDesc CreateDesc =
+			FRHIBufferCreateDesc::CreateStructured<FBasicRayTracingIntersectionResult>(TEXT("IntersectionResultBuffer"), NumRays)
+			.AddUsage(EBufferUsageFlags::Static | EBufferUsageFlags::UnorderedAccess)
+			.SetInitialState(ERHIAccess::UAVMask);
+
+		IntersectionResultBuffer = RHICmdList.CreateBuffer(CreateDesc);
 		IntersectionResultBufferView = RHICmdList.CreateUnorderedAccessView(IntersectionResultBuffer, 
 			FRHIViewDesc::CreateBufferUAV()
 			.SetType(FRHIViewDesc::EBufferType::Structured)
@@ -194,17 +192,16 @@ bool RunRayTracingTestbed_RenderThread(const FString& Parameters)
 	const FRayTracingSceneInitializer& SceneInitializer = RayTracingSceneRHI->GetInitializer();
 
 	FRayTracingAccelerationStructureSize SceneSizeInfo = RHICalcRayTracingSceneSize(SceneInitializer);
-	FRHIResourceCreateInfo SceneBufferCreateInfo(TEXT("RayTracingTestBedSceneBuffer"));
-	FBufferRHIRef SceneBuffer = RHICmdList.CreateBuffer(
-		uint32(SceneSizeInfo.ResultSize),
-		BUF_AccelerationStructure, 0 /*Stride*/, ERHIAccess::BVHWrite,
-		SceneBufferCreateInfo);
-
-	FRHIResourceCreateInfo ScratchBufferCreateInfo(TEXT("RayTracingTestBedScratchBuffer"));
-	FBufferRHIRef ScratchBuffer = RHICmdList.CreateBuffer(
-		uint32(SceneSizeInfo.BuildScratchSize),
-		BUF_UnorderedAccess, GRHIRayTracingScratchBufferAlignment, ERHIAccess::UAVCompute,
-		ScratchBufferCreateInfo);
+
+	const FRHIBufferCreateDesc SceneBufferCreateDesc =
+		FRHIBufferCreateDesc::Create(TEXT("RayTracingTestBedSceneBuffer"), uint32(SceneSizeInfo.ResultSize), 0, EBufferUsageFlags::AccelerationStructure)
+		.SetInitialState(ERHIAccess::BVHWrite);
+	FBufferRHIRef SceneBuffer = RHICmdList.CreateBuffer(SceneBufferCreateDesc);
+
+	const FRHIBufferCreateDesc ScratchBufferCreateDesc =
+		FRHIBufferCreateDesc::Create(TEXT("RayTracingTestBedScratchBuffer"), uint32(SceneSizeInfo.BuildScratchSize), GRHIRayTracingScratchBufferAlignment, EBufferUsageFlags::UnorderedAccess)
+		.SetInitialState(ERHIAccess::UAVCompute);
+	FBufferRHIRef ScratchBuffer = RHICmdList.CreateBuffer(ScratchBufferCreateDesc);
 
 	FRWBufferStructured InstanceBuffer;
 	InstanceBuffer.Initialize(RHICmdList, TEXT("RayTracingTestBedInstanceBuffer"), GRHIRayTracingInstanceDescriptorSize, SceneInitializer.MaxNumInstances);
@@ -213,38 +210,29 @@ bool RunRayTracingTestbed_RenderThread(const FString& Parameters)
 	AccelerationStructureAddressesBuffer.Initialize(RHICmdList, TEXT("RayTracingTestBedAccelerationStructureAddressesBuffer"), sizeof(FRayTracingAccelerationStructureAddress), BUF_Volatile | BUF_MultiGPUAllocate);
 
 	const uint32 NumInstanceGroups = RayTracingSceneInitializationData.NumGPUInstanceGroups + RayTracingSceneInitializationData.NumCPUInstanceGroups;
-	const uint32 InstanceGroupUploadBytes = NumInstanceGroups * sizeof(FRayTracingInstanceGroup);
-	FBufferRHIRef InstanceGroupUploadBuffer;
-	FShaderResourceViewRHIRef InstanceGroupUploadSRV;
-	{
-		FRHIResourceCreateInfo CreateInfo(TEXT("RayTracingTestBedInstanceGroupUploadBuffer"));
-		InstanceGroupUploadBuffer = RHICmdList.CreateStructuredBuffer(sizeof(FRayTracingInstanceGroup), InstanceGroupUploadBytes, BUF_ShaderResource | BUF_Volatile, CreateInfo);
-		InstanceGroupUploadSRV = RHICmdList.CreateShaderResourceView(InstanceGroupUploadBuffer, FRHIViewDesc::CreateBufferSRV().SetTypeFromBuffer(InstanceGroupUploadBuffer));
-	}
+	const FRHIBufferCreateDesc InstanceGroupUploadCreateDesc =
+		FRHIBufferCreateDesc::CreateStructured<FRayTracingInstanceGroup>(TEXT("RayTracingTestBedInstanceGroupUploadBuffer"), NumInstanceGroups)
+		.AddUsage(EBufferUsageFlags::ShaderResource | EBufferUsageFlags::Volatile)
+		.SetInitActionInitializer()
+		.DetermineInitialState();
+	TRHIBufferInitializer<FRayTracingInstanceGroup> InstanceGroupUploadData = RHICmdList.CreateBufferInitializer(InstanceGroupUploadCreateDesc);
 
 	const uint32 NumInstanceDescriptors = RayTracingSceneInitializationData.NumGPUInstanceDescriptors + RayTracingSceneInitializationData.NumCPUInstanceDescriptors;
-	const uint32 InstanceUploadBufferSize = NumInstanceDescriptors * sizeof(FRayTracingInstanceDescriptor);
-	FBufferRHIRef InstanceUploadBuffer;
-	FShaderResourceViewRHIRef InstanceUploadSRV;
-	{
-		FRHIResourceCreateInfo CreateInfo(TEXT("RayTracingTestBedInstanceUploadBuffer"));
-		InstanceUploadBuffer = RHICmdList.CreateStructuredBuffer(sizeof(FRayTracingInstanceDescriptor), InstanceUploadBufferSize, BUF_ShaderResource | BUF_Volatile, CreateInfo);
-		InstanceUploadSRV = RHICmdList.CreateShaderResourceView(InstanceUploadBuffer, FRHIViewDesc::CreateBufferSRV().SetTypeFromBuffer(InstanceUploadBuffer));
-	}
-
-	const uint32 TransformUploadBufferSize = RayTracingSceneInitializationData.NumNativeCPUInstances * 3 * sizeof(FVector4f);
-	FBufferRHIRef TransformUploadBuffer;
-	FShaderResourceViewRHIRef TransformUploadSRV;
-	{
-		FRHIResourceCreateInfo CreateInfo(TEXT("RayTracingTestBedInstanceUploadBuffer"));
-		TransformUploadBuffer = RHICmdList.CreateStructuredBuffer(sizeof(FVector4f), TransformUploadBufferSize, BUF_ShaderResource | BUF_Volatile, CreateInfo);
-		TransformUploadSRV = RHICmdList.CreateShaderResourceView(TransformUploadBuffer, FRHIViewDesc::CreateBufferSRV().SetTypeFromBuffer(TransformUploadBuffer));
-	}
+	const FRHIBufferCreateDesc InstanceUploadCreateDesc =
+		FRHIBufferCreateDesc::CreateStructured<FRayTracingInstanceDescriptor>(TEXT("RayTracingTestBedInstanceUploadBuffer"), NumInstanceDescriptors)
+		.AddUsage(EBufferUsageFlags::ShaderResource | EBufferUsageFlags::Volatile)
+		.SetInitActionInitializer()
+		.DetermineInitialState();
+	TRHIBufferInitializer<FRayTracingInstanceDescriptor> InstanceUploadData = RHICmdList.CreateBufferInitializer(InstanceUploadCreateDesc);
+
+	const FRHIBufferCreateDesc TransformUploadCreateDesc =
+		FRHIBufferCreateDesc::CreateStructured<FVector4f>(TEXT("RayTracingTestBedTransformUploadBuffer"), RayTracingSceneInitializationData.NumNativeCPUInstances * 3)
+		.AddUsage(EBufferUsageFlags::ShaderResource | EBufferUsageFlags::Volatile)
+		.SetInitActionInitializer()
+		.DetermineInitialState();
+	TRHIBufferInitializer<FVector4f> TransformUploadData = RHICmdList.CreateBufferInitializer(TransformUploadCreateDesc);
 
 	{
-		FRayTracingInstanceGroup* InstanceGroupUploadData = (FRayTracingInstanceGroup*)RHICmdList.LockBuffer(InstanceGroupUploadBuffer, 0, InstanceGroupUploadBytes, RLM_WriteOnly);
-		FRayTracingInstanceDescriptor* InstanceUploadData = (FRayTracingInstanceDescriptor*)RHICmdList.LockBuffer(InstanceUploadBuffer, 0, InstanceUploadBufferSize, RLM_WriteOnly);
-		FVector4f* TransformUploadData = (FVector4f*)RHICmdList.LockBuffer(TransformUploadBuffer, 0, TransformUploadBufferSize, RLM_WriteOnly);
 		FillRayTracingInstanceUploadBuffer(
 			FVector::ZeroVector,
 			Instances,
@@ -256,14 +244,21 @@ bool RunRayTracingTestbed_RenderThread(const FString& Parameters)
 			RayTracingSceneInitializationData.NumCPUInstanceGroups,
 			RayTracingSceneInitializationData.NumGPUInstanceDescriptors,
 			RayTracingSceneInitializationData.NumCPUInstanceDescriptors,
-			MakeArrayView(InstanceGroupUploadData, NumInstanceGroups),
-			MakeArrayView(InstanceUploadData, NumInstanceDescriptors),
-			MakeArrayView(TransformUploadData, RayTracingSceneInitializationData.NumNativeCPUInstances * 3));
-		RHICmdList.UnlockBuffer(TransformUploadBuffer);
-		RHICmdList.UnlockBuffer(InstanceUploadBuffer);
-		RHICmdList.UnlockBuffer(InstanceGroupUploadBuffer);
+			InstanceGroupUploadData.GetWriteView(),
+			InstanceUploadData.GetWriteView(),
+			TransformUploadData.GetWriteView()
+		);
 	}
 
+	FBufferRHIRef InstanceGroupUploadBuffer = InstanceGroupUploadData.Finalize();
+	FShaderResourceViewRHIRef InstanceGroupUploadSRV = RHICmdList.CreateShaderResourceView(InstanceGroupUploadBuffer, FRHIViewDesc::CreateBufferSRV().SetTypeFromBuffer(InstanceGroupUploadBuffer));
+
+	FBufferRHIRef InstanceUploadBuffer = InstanceUploadData.Finalize();
+	FShaderResourceViewRHIRef InstanceUploadSRV = RHICmdList.CreateShaderResourceView(InstanceUploadBuffer, FRHIViewDesc::CreateBufferSRV().SetTypeFromBuffer(InstanceUploadBuffer));
+
+	FBufferRHIRef TransformUploadBuffer = TransformUploadData.Finalize();
+	FShaderResourceViewRHIRef TransformUploadSRV = RHICmdList.CreateShaderResourceView(TransformUploadBuffer, FRHIViewDesc::CreateBufferSRV().SetTypeFromBuffer(TransformUploadBuffer));
+
 	for (uint32 GPUIndex : RHICmdList.GetGPUMask())
 	{
 		FRayTracingAccelerationStructureAddress* AddressesPtr = (FRayTracingAccelerationStructureAddress*)RHICmdList.LockBufferMGPU(
diff --git a/Engine/Source/Runtime/Renderer/Private/VT/VirtualTextureSpace.cpp b/Engine/Source/Runtime/Renderer/Private/VT/VirtualTextureSpace.cpp
--- a/Engine/Source/Runtime/Renderer/Private/VT/VirtualTextureSpace.cpp
+++ b/Engine/Source/Runtime/Renderer/Private/VT/VirtualTextureSpace.cpp
@@ -444,8 +444,12 @@ void FVirtualTextureSpace::ApplyUpdates(FVirtualTextureSystem* System, FRDGBuild
 			NewBufferSize = FMath::Max(NewBufferSize, UpdateBuffer->GetSize() * 2u);
 		}
 
-		FRHIResourceCreateInfo CreateInfo(TEXT("VirtualTexture_PageTableUpdateBuffer"));
-		UpdateBuffer = RHICmdList.CreateVertexBuffer(NewBufferSize, BUF_ShaderResource | BUF_Volatile, CreateInfo);
+		const FRHIBufferCreateDesc CreateDesc =
+			FRHIBufferCreateDesc::CreateVertex(TEXT("VirtualTexture_PageTableUpdateBuffer"), NewBufferSize)
+			.AddUsage(EBufferUsageFlags::ShaderResource | EBufferUsageFlags::Volatile)
+			.DetermineInitialState();
+
+		UpdateBuffer = RHICmdList.CreateBuffer(CreateDesc);
 		UpdateBufferSRV = RHICmdList.CreateShaderResourceView(
 			UpdateBuffer, 
 			FRHIViewDesc::CreateBufferSRV()
diff --git a/Engine/Source/Runtime/Renderer/Private/VolumetricFog.cpp b/Engine/Source/Runtime/Renderer/Private/VolumetricFog.cpp
--- a/Engine/Source/Runtime/Renderer/Private/VolumetricFog.cpp
+++ b/Engine/Source/Runtime/Renderer/Private/VolumetricFog.cpp
@@ -673,11 +673,13 @@ public:
 
 	virtual void InitRHI(FRHICommandListBase& RHICmdList) override
 	{
-		const int32 NumTriangles = NumVertices - 2;
-		const uint32 Size = NumVertices * sizeof(FScreenVertex);
-		FRHIResourceCreateInfo CreateInfo(TEXT("FCircleRasterizeVertexBuffer"));
-		VertexBufferRHI = RHICmdList.CreateBuffer(Size, BUF_Static | BUF_VertexBuffer, 0, ERHIAccess::VertexOrIndexBuffer, CreateInfo);
-		FScreenVertex* DestVertex = (FScreenVertex*)RHICmdList.LockBuffer(VertexBufferRHI, 0, Size, RLM_WriteOnly);
+		const FRHIBufferCreateDesc CreateDesc =
+			FRHIBufferCreateDesc::CreateVertex<FScreenVertex>(TEXT("FCircleRasterizeVertexBuffer"), NumVertices)
+			.AddUsage(EBufferUsageFlags::Static)
+			.SetInitialState(ERHIAccess::VertexOrIndexBuffer)
+			.SetInitActionInitializer();
+
+		TRHIBufferInitializer<FScreenVertex> DestVertex = RHICmdList.CreateBufferInitializer(CreateDesc);
 
 		const int32 NumSegments = NumVertices - 1;
 		const float RadiansPerRingSegment = PI / (float)NumSegments;
@@ -693,7 +695,7 @@ public:
 			DestVertex[VertexIndex].UV = FVector2f(RadiusScale * FMath::Cos(Angle) * .5f + .5f, RadiusScale * FMath::Sin(Angle) * .5f + .5f);
 		}
 
-		RHICmdList.UnlockBuffer(VertexBufferRHI);
+		VertexBufferRHI = DestVertex.Finalize();
 	}
 
 	static int32 NumVertices;
diff --git a/Engine/Source/Runtime/SlateRHIRenderer/Private/SlateElementVertexBuffer.h b/Engine/Source/Runtime/SlateRHIRenderer/Private/SlateElementVertexBuffer.h
--- a/Engine/Source/Runtime/SlateRHIRenderer/Private/SlateElementVertexBuffer.h
+++ b/Engine/Source/Runtime/SlateRHIRenderer/Private/SlateElementVertexBuffer.h
@@ -64,8 +64,12 @@ public:
 	
 			SetBufferSize(MinBufferSize);
 
-			FRHIResourceCreateInfo CreateInfo(TEXT("SlateElementVertices"));
-			VertexBufferRHI = RHICmdList.CreateVertexBuffer( MinBufferSize, BUF_Dynamic, CreateInfo );
+			const FRHIBufferCreateDesc CreateDesc =
+				FRHIBufferCreateDesc::CreateVertex(TEXT("SlateElementVertices"), MinBufferSize)
+				.AddUsage(EBufferUsageFlags::Dynamic)
+				.DetermineInitialState();
+
+			VertexBufferRHI = RHICmdList.CreateBuffer(CreateDesc);
 
 			// Ensure the vertex buffer could be created
 			check(IsValidRef(VertexBufferRHI));
@@ -128,8 +132,12 @@ private:
 		{
 			VertexBufferRHI.SafeRelease();
 
-			FRHIResourceCreateInfo CreateInfo(TEXT("SlateElementVertices"));
-			VertexBufferRHI = RHICmdList.CreateVertexBuffer(FinalSize, BUF_Dynamic, CreateInfo);
+			const FRHIBufferCreateDesc CreateDesc =
+				FRHIBufferCreateDesc::CreateVertex(TEXT("SlateElementVertices"), FinalSize)
+				.AddUsage(EBufferUsageFlags::Dynamic)
+				.DetermineInitialState();
+
+			VertexBufferRHI = RHICmdList.CreateBuffer(CreateDesc);
 
 			check(IsValidRef(VertexBufferRHI));
 
diff --git a/Engine/Source/Runtime/VulkanRHI/Private/VulkanBuffer.cpp b/Engine/Source/Runtime/VulkanRHI/Private/VulkanBuffer.cpp
--- a/Engine/Source/Runtime/VulkanRHI/Private/VulkanBuffer.cpp
+++ b/Engine/Source/Runtime/VulkanRHI/Private/VulkanBuffer.cpp
@@ -11,6 +11,7 @@
 #include "VulkanLLM.h"
 #include "VulkanRayTracing.h"
 #include "VulkanTransientResourceAllocator.h"
+#include "RHICoreBufferInitializer.h"
 #include "RHICoreStats.h"
 
 struct FVulkanPendingBufferLock
@@ -113,28 +114,21 @@ VkBufferUsageFlags FVulkanBuffer::UEToVKBufferUsageFlags(FVulkanDevice* InDevice
 	return OutVkUsage;
 }
 
-FVulkanBuffer::FVulkanBuffer(FVulkanDevice* InDevice, FRHIBufferDesc const& InBufferDesc, FRHIResourceCreateInfo& CreateInfo, FRHICommandListBase* InRHICmdList, const FRHITransientHeapAllocation* InTransientHeapAllocation)
-	: FRHIBuffer(InBufferDesc)
+FVulkanBuffer::FVulkanBuffer(FVulkanDevice* InDevice, const FRHIBufferCreateDesc& CreateDesc, const FRHITransientHeapAllocation* InTransientHeapAllocation)
+	: FRHIBuffer(CreateDesc)
 	, VulkanRHI::FDeviceChild(InDevice)
 {
 	VULKAN_TRACK_OBJECT_CREATE(FVulkanBuffer, this);
 
-	const bool bZeroSize = (InBufferDesc.Size == 0);
-	BufferUsageFlags = UEToVKBufferUsageFlags(InDevice, InBufferDesc.Usage, bZeroSize);
-	
-	#if !(UE_BUILD_SHIPPING || UE_BUILD_TEST)
-	if (CreateInfo.DebugName)
-	{
-		SetName(CreateInfo.DebugName);
-	}
-	#endif
+	const bool bZeroSize = (CreateDesc.Size == 0);
+	BufferUsageFlags = UEToVKBufferUsageFlags(InDevice, CreateDesc.Usage, bZeroSize);
 
 	if (!bZeroSize)
 	{
 		check(InDevice);
 
 		const bool bUnifiedMem = InDevice->HasUnifiedMemory();
-		const uint32 BufferAlignment = VulkanRHI::FMemoryManager::CalculateBufferAlignment(*InDevice, InBufferDesc.Usage, bZeroSize);
+		const uint32 BufferAlignment = VulkanRHI::FMemoryManager::CalculateBufferAlignment(*InDevice, CreateDesc.Usage, bZeroSize);
 
 		if (InTransientHeapAllocation != nullptr)
 		{
@@ -142,33 +136,12 @@ FVulkanBuffer::FVulkanBuffer(FVulkanDevice* InDevice, FRHIBufferDesc const& InBu
 			CurrentBufferAlloc.HostPtr = bUnifiedMem ? CurrentBufferAlloc.Alloc.GetMappedPointer(Device) : nullptr;
 			CurrentBufferAlloc.DeviceAddress = GetBufferDeviceAddress(InDevice, CurrentBufferAlloc.Alloc.GetBufferHandle()) + CurrentBufferAlloc.Alloc.Offset;
 			check(CurrentBufferAlloc.Alloc.Offset % BufferAlignment == 0);
-			check(CurrentBufferAlloc.Alloc.Size >= InBufferDesc.Size);
+			check(CurrentBufferAlloc.Alloc.Size >= CreateDesc.Size);
 		}
 		else
 		{
 			AllocateMemory(CurrentBufferAlloc);
 		}
-
-		if (CreateInfo.ResourceArray)
-		{
-			const uint32 CopyDataSize = FMath::Min(InBufferDesc.Size, CreateInfo.ResourceArray->GetResourceDataSize());
-
-			// We know this buffer is not in use by GPU atm. If we do have a direct access initialize it without extra copies
-			if (CurrentBufferAlloc.HostPtr)
-			{
-				FMemory::Memcpy(CurrentBufferAlloc.HostPtr, CreateInfo.ResourceArray->GetResourceData(), CopyDataSize);
-				++LockCounter;
-			}
-			else
-			{
-				check(InRHICmdList);
-				void* Data = Lock(*InRHICmdList, RLM_WriteOnly, CopyDataSize, 0);
-				FMemory::Memcpy(Data, CreateInfo.ResourceArray->GetResourceData(), CopyDataSize);
-				Unlock(*InRHICmdList);
-			}
-
-			CreateInfo.ResourceArray->Discard();
-		}
 	}
 }
 
@@ -438,14 +411,115 @@ void FVulkanBuffer::ReleaseOwnership()
 	FRHIBuffer::ReleaseOwnership();
 }
 
-FBufferRHIRef FVulkanDynamicRHI::RHICreateBuffer(FRHICommandListBase& RHICmdList, FRHIBufferDesc const& Desc, ERHIAccess ResourceState, FRHIResourceCreateInfo& CreateInfo)
+FRHIBufferInitializer FVulkanDynamicRHI::RHICreateBufferInitializer(FRHICommandListBase& RHICmdList, const FRHIBufferCreateDesc& CreateDesc)
 {
 #if VULKAN_USE_LLM
 	LLM_SCOPE_VULKAN(ELLMTagVulkan::VulkanBuffers);
 #else
-	LLM_SCOPE(EnumHasAnyFlags(Desc.Usage, EBufferUsageFlags::VertexBuffer | EBufferUsageFlags::IndexBuffer) ? ELLMTag::Meshes : ELLMTag::RHIMisc);
+	LLM_SCOPE(EnumHasAnyFlags(CreateDesc.Usage, EBufferUsageFlags::VertexBuffer | EBufferUsageFlags::IndexBuffer) ? ELLMTag::Meshes : ELLMTag::RHIMisc);
 #endif
-	return new FVulkanBuffer(Device, Desc, CreateInfo, &RHICmdList);
+
+	FVulkanBuffer* Buffer = new FVulkanBuffer(Device, CreateDesc);
+
+	if (CreateDesc.IsNull() || CreateDesc.InitAction == ERHIBufferInitAction::Default)
+	{
+		return UE::RHICore::FDefaultBufferInitializer(RHICmdList, Buffer);
+	}
+
+	if (CreateDesc.InitAction == ERHIBufferInitAction::Zeroed)
+	{
+		if (void* HostPointer = Buffer->GetCurrentHostPointer())
+		{
+			FMemory::Memzero(HostPointer, CreateDesc.Size);
+
+			return UE::RHICore::FDefaultBufferInitializer(RHICmdList, Buffer);
+		}
+
+		return UE::RHICore::FCustomBufferInitializer(RHICmdList, Buffer, nullptr, CreateDesc.Size,
+			[Buffer = TRefCountPtr<FVulkanBuffer>(Buffer)](FRHICommandListBase& RHICmdList) mutable
+			{
+				FVulkanUploadContext& UploadContext = FVulkanUploadContext::Get(RHICmdList);
+
+				FVulkanCommandBuffer& CommandBuffer = UploadContext.GetCommandBuffer();
+				const VulkanRHI::FVulkanAllocation& Allocation = Buffer->GetCurrentAllocation();
+
+				const auto Offset = Buffer->GetCurrentAllocation().Offset;
+				
+				VulkanRHI::vkCmdFillBuffer(
+					CommandBuffer.GetHandle(),
+					Allocation.GetBufferHandle(),
+					Allocation.Offset,
+					Allocation.Size,
+					0
+				);
+
+				// High level code expects the data in Buffer to be ready to read
+				const VkMemoryBarrier BarrierAfter =
+				{
+					VK_STRUCTURE_TYPE_MEMORY_BARRIER, nullptr, VK_ACCESS_TRANSFER_WRITE_BIT, VK_ACCESS_MEMORY_READ_BIT | VK_ACCESS_MEMORY_WRITE_BIT
+				};
+				VulkanRHI::vkCmdPipelineBarrier(
+					CommandBuffer.GetHandle(),
+					VK_PIPELINE_STAGE_TRANSFER_BIT,
+					VK_PIPELINE_STAGE_ALL_COMMANDS_BIT,
+					0,
+					1,
+					&BarrierAfter,
+					0,
+					nullptr,
+					0,
+					nullptr
+				);
+
+				return TRefCountPtr<FRHIBuffer>(MoveTemp(Buffer));
+			});
+	}
+
+	void* WritableData = nullptr;
+	bool bUsingLock = false;
+
+	if (void* HostPointer = Buffer->GetCurrentHostPointer())
+	{
+		WritableData = HostPointer;
+		Buffer->IncrementLockCounter();
+	}
+	else
+	{
+		WritableData = Buffer->Lock(RHICmdList, RLM_WriteOnly, CreateDesc.Size, 0);
+		bUsingLock = true;
+	}
+
+	if (CreateDesc.InitAction == ERHIBufferInitAction::ResourceArray)
+	{
+		check(CreateDesc.InitialData);
+
+		FMemory::Memcpy(WritableData, CreateDesc.InitialData->GetResourceData(), CreateDesc.InitialData->GetResourceDataSize());
+
+		// Discard the resource array's contents.
+		CreateDesc.InitialData->Discard();
+
+		if (bUsingLock)
+		{
+			Buffer->Unlock(RHICmdList);
+		}
+
+		return UE::RHICore::FDefaultBufferInitializer(RHICmdList, Buffer);
+	}
+
+	if (CreateDesc.InitAction == ERHIBufferInitAction::Initializer)
+	{
+		return UE::RHICore::FCustomBufferInitializer(RHICmdList, Buffer, WritableData, CreateDesc.Size,
+			[Buffer = TRefCountPtr<FVulkanBuffer>(Buffer), bUsingLock](FRHICommandListBase& RHICmdList) mutable
+			{
+				if (bUsingLock)
+				{
+					Buffer->Unlock(RHICmdList);
+				}
+				return TRefCountPtr<FRHIBuffer>(MoveTemp(Buffer));
+			});
+	}
+
+	return UE::RHICore::HandleUnknownBufferInitializerInitAction(RHICmdList, CreateDesc);
 }
 
 void* FVulkanDynamicRHI::LockBuffer_BottomOfPipe(FRHICommandListBase& RHICmdList, FRHIBuffer* BufferRHI, uint32 Offset, uint32 Size, EResourceLockMode LockMode)
diff --git a/Engine/Source/Runtime/VulkanRHI/Private/VulkanRayTracing.cpp b/Engine/Source/Runtime/VulkanRHI/Private/VulkanRayTracing.cpp
--- a/Engine/Source/Runtime/VulkanRHI/Private/VulkanRayTracing.cpp
+++ b/Engine/Source/Runtime/VulkanRHI/Private/VulkanRayTracing.cpp
@@ -408,8 +408,12 @@ FVulkanRayTracingGeometry::FVulkanRayTracingGeometry(FRHICommandListBase& RHICmd
 	}
 
 	FString DebugNameString = Initializer.DebugName.ToString();
-	FRHIResourceCreateInfo BlasBufferCreateInfo(*DebugNameString);
-	AccelerationStructureBuffer = ResourceCast(RHICmdList.CreateBuffer(SizeInfo.ResultSize, BUF_AccelerationStructure, 0, ERHIAccess::BVHWrite, BlasBufferCreateInfo).GetReference());
+
+	const FRHIBufferCreateDesc CreateDesc =
+		FRHIBufferCreateDesc::Create(*DebugNameString, SizeInfo.ResultSize, 0, EBufferUsageFlags::AccelerationStructure)
+		.SetInitialState(ERHIAccess::BVHWrite);
+
+	AccelerationStructureBuffer = ResourceCast(RHICmdList.CreateBuffer(CreateDesc).GetReference());
 
 	VkDevice NativeDevice = Device->GetInstanceHandle();
 
@@ -511,10 +515,11 @@ void FVulkanRayTracingGeometry::CompactAccelerationStructure(FVulkanCommandBuffe
 	TRefCountPtr<FVulkanBuffer> OldAccelerationStructure = AccelerationStructureBuffer;
 	VkAccelerationStructureKHR OldHandle = Handle;
 
-	FString DebugNameString = Initializer.DebugName.ToString();
-	FRHIResourceCreateInfo BlasBufferCreateInfo(*DebugNameString);
-	AccelerationStructureBuffer = new FVulkanBuffer(Device, FRHIBufferDesc(InSizeAfterCompaction, 0, BUF_AccelerationStructure), BlasBufferCreateInfo);
-	
+	const FString DebugNameString = Initializer.DebugName.ToString();
+	const FRHIBufferCreateDesc CreateDesc =
+		FRHIBufferCreateDesc::Create(*DebugNameString, InSizeAfterCompaction, 0, BUF_AccelerationStructure);
+
+	AccelerationStructureBuffer = new FVulkanBuffer(Device, CreateDesc);
 
 	INC_MEMORY_STAT_BY(STAT_VulkanRayTracingUsedVideoMemory, AccelerationStructureBuffer->GetSize());
 	INC_MEMORY_STAT_BY(STAT_VulkanRayTracingBLASMemory, AccelerationStructureBuffer->GetSize());
@@ -794,8 +799,12 @@ void BuildAccelerationStructure(
 		{
 			const uint64 ScratchBufferSize = bIsUpdate ? Scene.SizeInfo.UpdateScratchSize : Scene.SizeInfo.BuildScratchSize;
 
-			FRHIResourceCreateInfo ScratchBufferCreateInfo(TEXT("BuildScratchTLAS"));
-			ScratchBuffer = RHICmdList.CreateBuffer(ScratchBufferSize, BUF_StructuredBuffer | BUF_RayTracingScratch, 0, ERHIAccess::UAVCompute, ScratchBufferCreateInfo);
+			const FRHIBufferCreateDesc CreateDesc =
+				FRHIBufferCreateDesc::CreateStructured(TEXT("BuildScratchTLAS"), ScratchBufferSize, 0)
+				.AddUsage(EBufferUsageFlags::RayTracingScratch)
+				.SetInitialState(ERHIAccess::UAVCompute);
+
+			ScratchBuffer = RHICmdList.CreateBuffer(CreateDesc);
 			InScratchBuffer = ResourceCast(ScratchBuffer.GetReference());
 			InScratchOffset = 0;
 		}
@@ -895,9 +904,13 @@ FVulkanRayTracingShaderTable::FVulkanRayTracingShaderTable(FRHICommandListBase&
 
 		const uint32 ParameterBufferSize = Initializer.NumGeometrySegments * sizeof(FVulkanRayTracingGeometryParameters);
 		InlineGeometryParameterData.SetNumUninitialized(ParameterBufferSize);
+
+		const FRHIBufferCreateDesc CreateDesc =
+			FRHIBufferCreateDesc::CreateStructured<FVulkanRayTracingGeometryParameters>(TEXT("RayTracingSceneMetadata"), Initializer.NumGeometrySegments)
+			.AddUsage(EBufferUsageFlags::Dynamic | EBufferUsageFlags::ShaderResource)
+			.SetInitialState(ERHIAccess::SRVMask);
 		
-		FRHIResourceCreateInfo ParameterBufferCreateInfo(TEXT("RayTracingSceneMetadata"));
-		InlineGeometryParameterBuffer = ResourceCast(RHICmdList.CreateBuffer(ParameterBufferSize, BUF_Dynamic | BUF_StructuredBuffer | BUF_ShaderResource, sizeof(FVulkanRayTracingGeometryParameters), ERHIAccess::SRVMask, ParameterBufferCreateInfo).GetReference());
+		InlineGeometryParameterBuffer = ResourceCast(RHICmdList.CreateBuffer(CreateDesc).GetReference());
 	}
 }
 
diff --git a/Engine/Source/Runtime/VulkanRHI/Private/VulkanResources.h b/Engine/Source/Runtime/VulkanRHI/Private/VulkanResources.h
--- a/Engine/Source/Runtime/VulkanRHI/Private/VulkanResources.h
+++ b/Engine/Source/Runtime/VulkanRHI/Private/VulkanResources.h
@@ -969,7 +969,7 @@ public:
 class FVulkanBuffer : public FRHIBuffer, public VulkanRHI::FDeviceChild, public FVulkanViewableResource
 {
 public:
-	FVulkanBuffer(FVulkanDevice* InDevice, FRHIBufferDesc const& InBufferDesc, FRHIResourceCreateInfo& CreateInfo, class FRHICommandListBase* InRHICmdList = nullptr, const FRHITransientHeapAllocation* InTransientHeapAllocation = nullptr);
+	FVulkanBuffer(FVulkanDevice* InDevice, const FRHIBufferCreateDesc& CreateDesc, const FRHITransientHeapAllocation* InTransientHeapAllocation = nullptr);
 	virtual ~FVulkanBuffer();
 
 	inline const VulkanRHI::FVulkanAllocation& GetCurrentAllocation() const
@@ -1035,6 +1035,16 @@ public:
 		VkDeviceAddress DeviceAddress = 0;
 	};
 
+	void* GetCurrentHostPointer()
+	{
+		return CurrentBufferAlloc.HostPtr;
+	}
+
+	void IncrementLockCounter()
+	{
+		LockCounter++;
+	}
+
 protected:
 
 	// Will return a new allocation that can be used for this buffer (proper memory type and size)
diff --git a/Engine/Source/Runtime/VulkanRHI/Private/VulkanTransientResourceAllocator.cpp b/Engine/Source/Runtime/VulkanRHI/Private/VulkanTransientResourceAllocator.cpp
--- a/Engine/Source/Runtime/VulkanRHI/Private/VulkanTransientResourceAllocator.cpp
+++ b/Engine/Source/Runtime/VulkanRHI/Private/VulkanTransientResourceAllocator.cpp
@@ -108,19 +108,21 @@ FRHITransientTexture* FVulkanTransientResourceAllocator::CreateTexture(const FRH
 
 FRHITransientBuffer* FVulkanTransientResourceAllocator::CreateBuffer(const FRHIBufferCreateInfo& InCreateInfo, const TCHAR* InDebugName, const FRHITransientAllocationFences& Fences)
 {
-	checkf(!EnumHasAnyFlags(InCreateInfo.Usage, BUF_AccelerationStructure), TEXT("AccelerationStructure not yet supported as TransientResource."));
-	checkf(!EnumHasAnyFlags(InCreateInfo.Usage, BUF_Volatile), TEXT("The volatile flag is not supported for transient resources."));
+	const FRHIBufferCreateDesc CreateDesc =
+		FRHIBufferCreateDesc::Create(InDebugName, InCreateInfo)
+		.SetInitialState(ERHIAccess::Discard);
 
-	const bool bZeroSize = (InCreateInfo.Size == 0);
-	const uint32 Alignment = VulkanRHI::FMemoryManager::CalculateBufferAlignment(*Device, InCreateInfo.Usage, bZeroSize);
-	uint64 Size = Align(InCreateInfo.Size, Alignment);
+	checkf(!EnumHasAnyFlags(CreateDesc.Usage, BUF_AccelerationStructure), TEXT("AccelerationStructure not yet supported as TransientResource."));
+	checkf(!EnumHasAnyFlags(CreateDesc.Usage, BUF_Volatile), TEXT("The volatile flag is not supported for transient resources."));
 
-	return CreateBufferInternal(InCreateInfo, InDebugName, Fences, Size, Alignment,
-		[Device=Device, InCreateInfo, InDebugName, Size](const FRHITransientHeap::FResourceInitializer& Initializer)
+	const bool bZeroSize = (CreateDesc.Size == 0);
+	const uint32 Alignment = VulkanRHI::FMemoryManager::CalculateBufferAlignment(*Device, CreateDesc.Usage, bZeroSize);
+	uint64 Size = Align(CreateDesc.Size, Alignment);
+
+	return CreateBufferInternal(CreateDesc, InDebugName, Fences, Size, Alignment,
+	[Device=Device, CreateDesc, InDebugName, Size](const FRHITransientHeap::FResourceInitializer& Initializer)
 	{
-		FRHIResourceCreateInfo ResourceCreateInfo(InDebugName);
-		FRHIBuffer* Buffer = new FVulkanBuffer(Device, FRHIBufferDesc(InCreateInfo.Size, InCreateInfo.Stride, InCreateInfo.Usage), ResourceCreateInfo, nullptr, &Initializer.Allocation);
-		Buffer->SetTrackedAccess_Unsafe(ERHIAccess::Discard);
-		return new FRHITransientBuffer(Buffer, 0/*GpuVirtualAddress*/, Initializer.Hash, Size, ERHITransientAllocationType::Heap, InCreateInfo);
+		FRHIBuffer* Buffer = new FVulkanBuffer(Device, CreateDesc, &Initializer.Allocation);
+		return new FRHITransientBuffer(Buffer, 0/*GpuVirtualAddress*/, Initializer.Hash, Size, ERHITransientAllocationType::Heap, CreateDesc);
 	});
 }
diff --git a/Engine/Source/Runtime/VulkanRHI/Public/VulkanDynamicRHI.h b/Engine/Source/Runtime/VulkanRHI/Public/VulkanDynamicRHI.h
--- a/Engine/Source/Runtime/VulkanRHI/Public/VulkanDynamicRHI.h
+++ b/Engine/Source/Runtime/VulkanRHI/Public/VulkanDynamicRHI.h
@@ -131,7 +131,9 @@ public:
 	virtual FComputePipelineStateRHIRef RHICreateComputePipelineState(FRHIComputeShader* ComputeShader) final override;
 	virtual FUniformBufferRHIRef RHICreateUniformBuffer(const void* Contents, const FRHIUniformBufferLayout* Layout, EUniformBufferUsage Usage, EUniformBufferValidation Validation) final override;
 	virtual void RHIUpdateUniformBuffer(FRHICommandListBase& RHICmdList, FRHIUniformBuffer* UniformBufferRHI, const void* Contents) final override;
-	virtual FBufferRHIRef RHICreateBuffer(FRHICommandListBase& RHICmdList, FRHIBufferDesc const& Desc, ERHIAccess ResourceState, FRHIResourceCreateInfo& CreateInfo) final override;
+
+	[[nodiscard]] virtual FRHIBufferInitializer RHICreateBufferInitializer(FRHICommandListBase& RHICmdList, const FRHIBufferCreateDesc& CreateDesc) final override;
+
 	virtual void RHIReplaceResources(FRHICommandListBase& RHICmdList, TArray<FRHIResourceReplaceInfo>&& ReplaceInfos) final override;
 	virtual void* LockBuffer_BottomOfPipe(FRHICommandListBase& RHICmdList, FRHIBuffer* Buffer, uint32 Offset, uint32 Size, EResourceLockMode LockMode) final override;
 	virtual void UnlockBuffer_BottomOfPipe(FRHICommandListBase& RHICmdList, FRHIBuffer* Buffer) final override;
