Build fix

#rnx

diff --git a/Engine/Source/Runtime/RenderCore/Public/VirtualTextureEnum.h b/Engine/Source/Runtime/RenderCore/Public/VirtualTextureEnum.h
--- a/Engine/Source/Runtime/RenderCore/Public/VirtualTextureEnum.h
+++ b/Engine/Source/Runtime/RenderCore/Public/VirtualTextureEnum.h
@@ -20,7 +20,6 @@ enum class EVTProducerPriority : uint8
 
 	Count UMETA(Hidden)
 };
-static_assert(static_cast<uint8>(EVTProducerPriority::Count) == static_cast<uint8>(EVTProducerPriority::Count), "EVTProducerPriority is the UENUM counterpart to EVTProducerPriority");
 ENUM_RANGE_BY_COUNT(EVTProducerPriority, EVTProducerPriority::Count);
 
 
@@ -33,5 +32,4 @@ enum class EVTInvalidatePriority : uint8
 
 	Count UMETA(Hidden)
 };
-static_assert(static_cast<uint8>(EVTInvalidatePriority::Count) == static_cast<uint8>(EVTInvalidatePriority::Count), "EVTInvalidatePriority is the UENUM counterpart to EVTInvalidatePriority");
 ENUM_RANGE_BY_COUNT(EVTInvalidatePriority, EVTInvalidatePriority::Count);
