Added delegate for when render command pipes are synced.

diff --git a/Engine/Source/Runtime/RenderCore/Private/RenderingThread.cpp b/Engine/Source/Runtime/RenderCore/Private/RenderingThread.cpp
--- a/Engine/Source/Runtime/RenderCore/Private/RenderingThread.cpp
+++ b/Engine/Source/Runtime/RenderCore/Private/RenderingThread.cpp
@@ -1837,6 +1837,7 @@ inline bool HasBitsSet(const FRenderCommandPipeBitArray& Bits)
 namespace UE::RenderCommandPipe
 {
 	static thread_local FRenderCommandPipe* ReplayingPipe = nullptr;
+	static FStopRecordingDelegate StopRecordingDelegate;
 
 	void Initialize()
 	{
@@ -1870,12 +1871,16 @@ namespace UE::RenderCommandPipe
 
 	FRenderCommandPipeBitArray StopRecording()
 	{
-		return GRenderCommandPipeRegistry.StopRecording();
+		FRenderCommandPipeBitArray PipeBits = GRenderCommandPipeRegistry.StopRecording();
+		GetStopRecordingDelegate().Broadcast(PipeBits);
+		return PipeBits;
 	}
 
 	FRenderCommandPipeBitArray StopRecording(TConstArrayView<FRenderCommandPipe*> Pipes)
 	{
-		return GRenderCommandPipeRegistry.StopRecording(Pipes);
+		FRenderCommandPipeBitArray PipeBits = GRenderCommandPipeRegistry.StopRecording(Pipes);
+		GetStopRecordingDelegate().Broadcast(PipeBits);
+		return PipeBits;
 	}
 
 	TConstArrayView<FRenderCommandPipe*> GetPipes()
@@ -1883,6 +1888,11 @@ namespace UE::RenderCommandPipe
 		return GRenderCommandPipeRegistry.GetPipes();
 	}
 
+	FStopRecordingDelegate& GetStopRecordingDelegate()
+	{
+		return StopRecordingDelegate;
+	}
+
 	FSyncScope::FSyncScope()
 	{
 		PipeBits = StopRecording();
diff --git a/Engine/Source/Runtime/RenderCore/Public/RenderingThread.h b/Engine/Source/Runtime/RenderCore/Public/RenderingThread.h
--- a/Engine/Source/Runtime/RenderCore/Public/RenderingThread.h
+++ b/Engine/Source/Runtime/RenderCore/Public/RenderingThread.h
@@ -308,6 +308,10 @@ namespace UE::RenderCommandPipe
 	// Returns the list of all registered pipes.
 	extern TConstArrayView<FRenderCommandPipe*> GetPipes();
 
+	// A delegate to receive events at sync points when recording is stopped. Input is a bit array of pipes that were stopped.
+	DECLARE_MULTICAST_DELEGATE_OneParam(FStopRecordingDelegate, const FRenderCommandPipeBitArray&);
+	extern RENDERCORE_API FStopRecordingDelegate& GetStopRecordingDelegate();
+
 	// [Game Thread] Stops render command pipe recording during the duration of the scope and restarts recording once the scope is complete.
 	class RENDERCORE_API FSyncScope
 	{
