Fix threading issues with render command fences
- ESyncDepth::RHIThread only synced as far as RHI command list translation. It did not include the submission tasks, which caused crashes in some platform RHIs.
- Reworked FRenderCommandFence::BeginFence() so that syncing with the RHI thread includes all RHI tasks.

diff --git a/Engine/Source/Runtime/RHI/Private/RHICommandList.cpp b/Engine/Source/Runtime/RHI/Private/RHICommandList.cpp
--- a/Engine/Source/Runtime/RHI/Private/RHICommandList.cpp
+++ b/Engine/Source/Runtime/RHI/Private/RHICommandList.cpp
@@ -1106,7 +1106,7 @@ void FRHICommandListExecutor::FSubmitState::Submit(const FSubmitArgs& Args)
 	delete this;
 }
 
-RHI_API void FRHICommandListExecutor::Submit(TConstArrayView<FRHICommandListBase*> AdditionalCommandLists, ERHISubmitFlags SubmitFlags)
+RHI_API FGraphEventRef FRHICommandListExecutor::Submit(TConstArrayView<FRHICommandListBase*> AdditionalCommandLists, ERHISubmitFlags SubmitFlags)
 {
 	check(IsInRenderingThread());
 	SCOPED_NAMED_EVENT(RHICmdList_Submit, FColor::White);
@@ -1322,6 +1322,7 @@ RHI_API void FRHICommandListExecutor::Submit(TConstArrayView<FRHICommandListBase
 #endif
 
 	CommandListImmediate.InitializeImmediateContexts();
+	return CompletionEvent;
 }
 
 void FRHICommandListImmediate::InitializeImmediateContexts()
diff --git a/Engine/Source/Runtime/RHI/Public/RHICommandList.h b/Engine/Source/Runtime/RHI/Public/RHICommandList.h
--- a/Engine/Source/Runtime/RHI/Public/RHICommandList.h
+++ b/Engine/Source/Runtime/RHI/Public/RHICommandList.h
@@ -5633,7 +5633,7 @@ public:
 	static inline FRHICommandListImmediate& GetImmediateCommandList();
 	RHI_API void LatchBypass();
 
-	RHI_API void Submit(TConstArrayView<FRHICommandListBase*> AdditionalCommandLists, ERHISubmitFlags SubmitFlags);
+	RHI_API FGraphEventRef Submit(TConstArrayView<FRHICommandListBase*> AdditionalCommandLists, ERHISubmitFlags SubmitFlags);
 
 	RHI_API static void WaitOnRHIThreadFence(FGraphEventRef& Fence);
 
diff --git a/Engine/Source/Runtime/RenderCore/Private/RenderingThread.cpp b/Engine/Source/Runtime/RenderCore/Private/RenderingThread.cpp
--- a/Engine/Source/Runtime/RenderCore/Private/RenderingThread.cpp
+++ b/Engine/Source/Runtime/RenderCore/Private/RenderingThread.cpp
@@ -1082,30 +1082,25 @@ void FRenderCommandFence::BeginFence(ESyncDepth SyncDepth)
 
 	ENQUEUE_RENDER_COMMAND(BeginFence)([Event, SyncDepth](FRHICommandListImmediate& RHICmdList) mutable
 	{
-		if (SyncDepth == ESyncDepth::RenderThread)
-		{
-			TRACE_CPUPROFILER_EVENT_SCOPE(SyncTrigger_RenderThread);
-			Event.Trigger();
-		}
-		else
+		if (SyncDepth == ESyncDepth::Swapchain)
 		{
-			RHICmdList.EnqueueLambda([SyncDepth, Event](FRHICommandListImmediate&) mutable
+			UE::Tasks::FTaskEvent SwapchainEvent{ UE_SOURCE_LOCATION };
+			Event.AddPrerequisites(SwapchainEvent);
+
+			RHICmdList.EnqueueLambda([SyncDepth, SwapchainEvent](FRHICommandListImmediate&) mutable
 			{
-				if (SyncDepth == ESyncDepth::RHIThread)
-				{
-					// Sync the Game Thread with the RHI Thread
-					TRACE_CPUPROFILER_EVENT_SCOPE(SyncTrigger_RHIThread);
-					Event.Trigger();
-				}
-				else
-				{
-					// This command runs *after* a present has happened, so the counter has already been incremented.
-					// Subtracting 1 gives us the index of the frame that has *just* been presented.
-					RHITriggerTaskEventOnFlip(GRHIPresentCounter - 1, Event);
-				}
+				// This command runs *after* a present has happened, so the counter has already been incremented.
+				// Subtracting 1 gives us the index of the frame that has *just* been presented.
+				RHITriggerTaskEventOnFlip(GRHIPresentCounter - 1, SwapchainEvent);
 			});
-			RHICmdList.ImmediateFlush(EImmediateFlushType::DispatchToRHIThread);
 		}
+		else if (SyncDepth == ESyncDepth::RHIThread)
+		{
+			Event.AddPrerequisites(GRHICommandList.Submit({}, ERHISubmitFlags::SubmitToGPU));
+		}
+
+		TRACE_CPUPROFILER_EVENT_SCOPE(SyncTrigger_RenderThread);
+		Event.Trigger();
 	});
 
 	CompletionTask = MoveTemp(Event);
