GetRayTracing*GroupLibrary function needs current shader platform to make sure the raytracing shader libraries of different preview platforms don't mix ray tracing shaders which might not be compatible.

diff --git a/Engine/Plugins/Experimental/GPULightmass/Source/GPULightmass/Private/LightmapRenderer.cpp b/Engine/Plugins/Experimental/GPULightmass/Source/GPULightmass/Private/LightmapRenderer.cpp
--- a/Engine/Plugins/Experimental/GPULightmass/Source/GPULightmass/Private/LightmapRenderer.cpp
+++ b/Engine/Plugins/Experimental/GPULightmass/Source/GPULightmass/Private/LightmapRenderer.cpp
@@ -1346,9 +1346,11 @@ bool FSceneRenderState::SetupRayTracingScene(FRDGBuilder& GraphBuilder, FSceneUn
 			}
 			PSOInitializer.SetRayGenShaderTable(RayGenShaderTable);
 
+			EShaderPlatform ShaderPlatform = View.GetShaderPlatform();
+
 			auto DefaultClosestHitShader = GetGPULightmassDefaultOpaqueHitShader(GlobalShaderMap);
 			TArray<FRHIRayTracingShader*> RayTracingHitGroupLibrary;
-			FShaderMapResource::GetRayTracingHitGroupLibrary(RayTracingHitGroupLibrary, DefaultClosestHitShader);
+			FShaderMapResource::GetRayTracingHitGroupLibrary(ShaderPlatform, RayTracingHitGroupLibrary, DefaultClosestHitShader);
 
 			FRHIRayTracingShader* HiddenMaterialShader = GetGPULightmassDefaultHiddenHitShader(View.ShaderMap);
 			RayTracingHitGroupLibrary.Add(HiddenMaterialShader);
diff --git a/Engine/Source/Runtime/RenderCore/Private/ShaderResource.cpp b/Engine/Source/Runtime/RenderCore/Private/ShaderResource.cpp
--- a/Engine/Source/Runtime/RenderCore/Private/ShaderResource.cpp
+++ b/Engine/Source/Runtime/RenderCore/Private/ShaderResource.cpp
@@ -103,54 +103,66 @@ bool FShaderMapResource::ArePlatformsCompatible(EShaderPlatform CurrentPlatform,
 class FRayTracingShaderLibrary
 {
 public:
-	uint32 AddShader(FRHIRayTracingShader* Shader)
+	uint32 AddShader(EShaderPlatform ShaderPlatform, FRHIRayTracingShader* Shader)
 	{
 		const int32 PayloadIndex = FMath::CountTrailingZeros(Shader->RayTracingPayloadType);
 		FScopeLock Lock(&CS);
 
-		if (UnusedIndicies[PayloadIndex].Num() != 0)
+		FShaderLibrary& ShaderLibrary = ShadersPerPlatform.FindOrAdd(ShaderPlatform);
+		if (ShaderLibrary.UnusedIndicies[PayloadIndex].Num() != 0)
 		{
-			uint32 Index = UnusedIndicies[PayloadIndex].Pop(EAllowShrinking::No);
-			checkSlow(Shaders[PayloadIndex][Index] == nullptr);
-			Shaders[PayloadIndex][Index] = Shader;
+			uint32 Index = ShaderLibrary.UnusedIndicies[PayloadIndex].Pop(EAllowShrinking::No);
+			checkSlow(ShaderLibrary.Shaders[PayloadIndex][Index] == nullptr);
+			ShaderLibrary.Shaders[PayloadIndex][Index] = Shader;
 			return Index;
 		}
 		else
 		{
-			return Shaders[PayloadIndex].Add(Shader);
+			return ShaderLibrary.Shaders[PayloadIndex].Add(Shader);
 		}
 	}
 
-	void RemoveShader(uint32 Index, FRHIRayTracingShader* Shader)
+	void RemoveShader(EShaderPlatform ShaderPlatform, uint32 Index, FRHIRayTracingShader* Shader)
 	{
 		if (Index != ~0u)
 		{
 			const int32 PayloadIndex = FMath::CountTrailingZeros(Shader->RayTracingPayloadType);
+
 			FScopeLock Lock(&CS);
-			checkSlow(Shaders[PayloadIndex][Index] == Shader);
-			UnusedIndicies[PayloadIndex].Push(Index);
-			Shaders[PayloadIndex][Index] = nullptr;
+			FShaderLibrary* ShaderLibrary = ShadersPerPlatform.Find(ShaderPlatform);
+			check(ShaderLibrary);
+			checkSlow(ShaderLibrary->Shaders[PayloadIndex][Index] == Shader);
+			ShaderLibrary->UnusedIndicies[PayloadIndex].Push(Index);
+			ShaderLibrary->Shaders[PayloadIndex][Index] = nullptr;
 		}
 	}
 
-	void GetShaders(TArray<FRHIRayTracingShader*>& OutShaders, FRHIRayTracingShader* DefaultShader)
+	void GetShaders(EShaderPlatform ShaderPlatform, TArray<FRHIRayTracingShader*>& OutShaders, FRHIRayTracingShader* DefaultShader)
 	{
 		const int32 PayloadIndex = FMath::CountTrailingZeros(DefaultShader->RayTracingPayloadType);
 		const int32 BaseOutIndex = OutShaders.Num();
 
 		FScopeLock Lock(&CS);
 
-		OutShaders.Append(Shaders[PayloadIndex]);
-
-		for (uint32 Index : UnusedIndicies[PayloadIndex])
+		FShaderLibrary* ShaderLibrary = ShadersPerPlatform.Find(ShaderPlatform);
+		if (ShaderLibrary)
 		{
-			OutShaders[BaseOutIndex + Index] = DefaultShader;
+			OutShaders.Append(ShaderLibrary->Shaders[PayloadIndex]);
+
+			for (uint32 Index : ShaderLibrary->UnusedIndicies[PayloadIndex])
+			{
+				OutShaders[BaseOutIndex + Index] = DefaultShader;
+			}
 		}
 	}
 
 private:
-	TArray<uint32> UnusedIndicies[32];
-	TArray<FRHIRayTracingShader*> Shaders[32];
+	struct FShaderLibrary
+	{
+		TArray<uint32> UnusedIndicies[32];
+		TArray<FRHIRayTracingShader*> Shaders[32];
+	};
+	TMap<EShaderPlatform, FShaderLibrary> ShadersPerPlatform;
 	FCriticalSection CS;
 };
 
@@ -158,19 +170,19 @@ static FRayTracingShaderLibrary GlobalRayTracingHitGroupLibrary;
 static FRayTracingShaderLibrary GlobalRayTracingCallableShaderLibrary;
 static FRayTracingShaderLibrary GlobalRayTracingMissShaderLibrary;
 
-void FShaderMapResource::GetRayTracingHitGroupLibrary(TArray<FRHIRayTracingShader*>& RayTracingShaders, FRHIRayTracingShader* DefaultShader)
+void FShaderMapResource::GetRayTracingHitGroupLibrary(EShaderPlatform ShaderPlatform, TArray<FRHIRayTracingShader*>& RayTracingShaders, FRHIRayTracingShader* DefaultShader)
 {
-	GlobalRayTracingHitGroupLibrary.GetShaders(RayTracingShaders, DefaultShader);
+	GlobalRayTracingHitGroupLibrary.GetShaders(ShaderPlatform, RayTracingShaders, DefaultShader);
 }
 
-void FShaderMapResource::GetRayTracingCallableShaderLibrary(TArray<FRHIRayTracingShader*>& RayTracingCallableShaders, FRHIRayTracingShader* DefaultShader)
+void FShaderMapResource::GetRayTracingCallableShaderLibrary(EShaderPlatform ShaderPlatform, TArray<FRHIRayTracingShader*>& RayTracingCallableShaders, FRHIRayTracingShader* DefaultShader)
 {
-	GlobalRayTracingCallableShaderLibrary.GetShaders(RayTracingCallableShaders, DefaultShader);
+	GlobalRayTracingCallableShaderLibrary.GetShaders(ShaderPlatform, RayTracingCallableShaders, DefaultShader);
 }
 
-void FShaderMapResource::GetRayTracingMissShaderLibrary(TArray<FRHIRayTracingShader*>& RayTracingMissShaders, FRHIRayTracingShader* DefaultShader)
+void FShaderMapResource::GetRayTracingMissShaderLibrary(EShaderPlatform ShaderPlatform, TArray<FRHIRayTracingShader*>& RayTracingMissShaders, FRHIRayTracingShader* DefaultShader)
 {
-	GlobalRayTracingMissShaderLibrary.GetShaders(RayTracingMissShaders, DefaultShader);
+	GlobalRayTracingMissShaderLibrary.GetShaders(ShaderPlatform, RayTracingMissShaders, DefaultShader);
 }
 #endif // RHI_RAYTRACING
 
@@ -547,13 +559,13 @@ void FShaderMapResource::ReleaseRHI()
 				switch (Shader->GetFrequency())
 				{
 				case SF_RayHitGroup:
-					GlobalRayTracingHitGroupLibrary.RemoveShader(IndexInLibrary, static_cast<FRHIRayTracingShader*>(Shader));
+					GlobalRayTracingHitGroupLibrary.RemoveShader(Platform, IndexInLibrary, static_cast<FRHIRayTracingShader*>(Shader));
 					break;
 				case SF_RayCallable:
-					GlobalRayTracingCallableShaderLibrary.RemoveShader(IndexInLibrary, static_cast<FRHIRayTracingShader*>(Shader));
+					GlobalRayTracingCallableShaderLibrary.RemoveShader(Platform, IndexInLibrary, static_cast<FRHIRayTracingShader*>(Shader));
 					break;
 				case SF_RayMiss:
-					GlobalRayTracingMissShaderLibrary.RemoveShader(IndexInLibrary, static_cast<FRHIRayTracingShader*>(Shader));
+					GlobalRayTracingMissShaderLibrary.RemoveShader(Platform, IndexInLibrary, static_cast<FRHIRayTracingShader*>(Shader));
 					break;
 				default:
 					break;
@@ -625,13 +637,13 @@ FRHIShader* FShaderMapResource::CreateShaderOrCrash(int32 ShaderIndex, bool bReq
 				switch (Frequency)
 				{
 					case SF_RayHitGroup:
-						RayTracingLibraryIndices[ShaderIndex] = GlobalRayTracingHitGroupLibrary.AddShader(static_cast<FRHIRayTracingShader*>(Shader));
+						RayTracingLibraryIndices[ShaderIndex] = GlobalRayTracingHitGroupLibrary.AddShader(Platform, static_cast<FRHIRayTracingShader*>(Shader));
 						break;
 					case SF_RayCallable:
-						RayTracingLibraryIndices[ShaderIndex] = GlobalRayTracingCallableShaderLibrary.AddShader(static_cast<FRHIRayTracingShader*>(Shader));
+						RayTracingLibraryIndices[ShaderIndex] = GlobalRayTracingCallableShaderLibrary.AddShader(Platform, static_cast<FRHIRayTracingShader*>(Shader));
 						break;
 					case SF_RayMiss:
-						RayTracingLibraryIndices[ShaderIndex] = GlobalRayTracingMissShaderLibrary.AddShader(static_cast<FRHIRayTracingShader*>(Shader));
+						RayTracingLibraryIndices[ShaderIndex] = GlobalRayTracingMissShaderLibrary.AddShader(Platform, static_cast<FRHIRayTracingShader*>(Shader));
 						break;
 					case SF_RayGen:
 						// NOTE: we do not maintain a library for raygen shaders since the list of rayshaders we care about is usually small and consistent
diff --git a/Engine/Source/Runtime/RenderCore/Public/Shader.h b/Engine/Source/Runtime/RenderCore/Public/Shader.h
--- a/Engine/Source/Runtime/RenderCore/Public/Shader.h
+++ b/Engine/Source/Runtime/RenderCore/Public/Shader.h
@@ -413,21 +413,25 @@ public:
 
 #if RHI_RAYTRACING
 
-	UE_DEPRECATED(5.1, "GetRayTracingMaterialLibrary is deprecated. Use GetRayTracingHitGroupLibrary instead.")
-	static void GetRayTracingMaterialLibrary(TArray<FRHIRayTracingShader*>& RayTracingMaterials, FRHIRayTracingShader* DefaultShader)
+	UE_DEPRECATED(5.6, "GetRayTracingHitGroupLibrary needs current EShaderPlatform")
+	static void GetRayTracingHitGroupLibrary(TArray<FRHIRayTracingShader*>& RayTracingHitGroupShaders, FRHIRayTracingShader* DefaultShader)
 	{
-		GetRayTracingHitGroupLibrary(RayTracingMaterials, DefaultShader);
+		GetRayTracingHitGroupLibrary(GMaxRHIShaderPlatform, RayTracingHitGroupShaders, DefaultShader);
 	}
-
-	UE_DEPRECATED(5.1, "GetRayTracingMaterialLibraryIndex is deprecated. Use GetRayTracingHitGroupLibraryIndex instead.")
-	inline uint32 GetRayTracingMaterialLibraryIndex(int32 ShaderIndex)
+	UE_DEPRECATED(5.6, "GetRayTracingCallableShaderLibrary needs current EShaderPlatform")
+	static void GetRayTracingCallableShaderLibrary(TArray<FRHIRayTracingShader*>& RayTracingCallableShaders, FRHIRayTracingShader* DefaultShader)
+	{
+		GetRayTracingCallableShaderLibrary(GMaxRHIShaderPlatform, RayTracingCallableShaders, DefaultShader);
+	}
+	UE_DEPRECATED(5.6, "GetRayTracingMissShaderLibrary needs current EShaderPlatform")
+	static void GetRayTracingMissShaderLibrary(TArray<FRHIRayTracingShader*>& RayTracingMissShaders, FRHIRayTracingShader* DefaultShader)
 	{
-		return GetRayTracingHitGroupLibraryIndex(ShaderIndex);
+		GetRayTracingMissShaderLibrary(GMaxRHIShaderPlatform, RayTracingMissShaders, DefaultShader);
 	}
 
-	static RENDERCORE_API void GetRayTracingHitGroupLibrary(TArray<FRHIRayTracingShader*>& RayTracingHitGroupShaders, FRHIRayTracingShader* DefaultShader);
-	static RENDERCORE_API void GetRayTracingCallableShaderLibrary(TArray<FRHIRayTracingShader*>& RayTracingCallableShaders, FRHIRayTracingShader* DefaultShader);
-	static RENDERCORE_API void GetRayTracingMissShaderLibrary(TArray<FRHIRayTracingShader*>& RayTracingMissShaders, FRHIRayTracingShader* DefaultShader);
+	static RENDERCORE_API void GetRayTracingHitGroupLibrary(EShaderPlatform ShaderPlatform, TArray<FRHIRayTracingShader*>& RayTracingHitGroupShaders, FRHIRayTracingShader* DefaultShader);
+	static RENDERCORE_API void GetRayTracingCallableShaderLibrary(EShaderPlatform ShaderPlatform, TArray<FRHIRayTracingShader*>& RayTracingCallableShaders, FRHIRayTracingShader* DefaultShader);
+	static RENDERCORE_API void GetRayTracingMissShaderLibrary(EShaderPlatform ShaderPlatform, TArray<FRHIRayTracingShader*>& RayTracingMissShaders, FRHIRayTracingShader* DefaultShader);
 
 	inline uint32 GetRayTracingHitGroupLibraryIndex(int32 ShaderIndex)
 	{
diff --git a/Engine/Source/Runtime/Renderer/Private/RayTracing/RayTracingMaterialHitShaders.cpp b/Engine/Source/Runtime/Renderer/Private/RayTracing/RayTracingMaterialHitShaders.cpp
--- a/Engine/Source/Runtime/Renderer/Private/RayTracing/RayTracingMaterialHitShaders.cpp
+++ b/Engine/Source/Runtime/Renderer/Private/RayTracing/RayTracingMaterialHitShaders.cpp
@@ -683,7 +683,7 @@ void FDeferredShadingSceneRenderer::CreateMaterialRayTracingMaterialPipeline(
 	FRayTracingPipelineStateInitializer Initializer;
 	Initializer.MaxPayloadSizeInBytes = GetRayTracingPayloadTypeMaxSize(PayloadType);
 
-	const FShaderBindingLayout* ShaderBindingLayout = RayTracing::GetShaderBindingLayout(View.GetShaderPlatform());
+	const FShaderBindingLayout* ShaderBindingLayout = RayTracing::GetShaderBindingLayout(ShaderPlatform);
 	if (ShaderBindingLayout)
 	{
 		Initializer.ShaderBindingLayout = &ShaderBindingLayout->RHILayout;
@@ -692,7 +692,7 @@ void FDeferredShadingSceneRenderer::CreateMaterialRayTracingMaterialPipeline(
 	FRHIRayTracingShader* DefaultMissShader = bIsPathTracing ? GetPathTracingDefaultMissShader(View.ShaderMap) : GetRayTracingDefaultMissShader(View.ShaderMap);
 
 	TArray<FRHIRayTracingShader*> RayTracingMissShaderLibrary;
-	FShaderMapResource::GetRayTracingMissShaderLibrary(RayTracingMissShaderLibrary, DefaultMissShader);
+	FShaderMapResource::GetRayTracingMissShaderLibrary(ShaderPlatform, RayTracingMissShaderLibrary, DefaultMissShader);
 
 	// make sure we have at least one miss shader present
 	check(RayTracingMissShaderLibrary.Num() > 0);
@@ -715,11 +715,11 @@ void FDeferredShadingSceneRenderer::CreateMaterialRayTracingMaterialPipeline(
 	TArray<FRHIRayTracingShader*> RayTracingHitGroupLibrary;
 	if (bEnableMaterials)
 	{
-		FShaderMapResource::GetRayTracingHitGroupLibrary(RayTracingHitGroupLibrary, OpaqueShadowShader);
+		FShaderMapResource::GetRayTracingHitGroupLibrary(ShaderPlatform, RayTracingHitGroupLibrary, OpaqueShadowShader);
 
 		if (bSupportMeshDecals)
 		{
-			FShaderMapResource::GetRayTracingHitGroupLibrary(RayTracingHitGroupLibrary, OpaqueMeshDecalHitShader);
+			FShaderMapResource::GetRayTracingHitGroupLibrary(ShaderPlatform, RayTracingHitGroupLibrary, OpaqueMeshDecalHitShader);
 		}
 	}
 
@@ -748,7 +748,7 @@ void FDeferredShadingSceneRenderer::CreateMaterialRayTracingMaterialPipeline(
 
 		if (bEnableMaterials)
 		{
-			FShaderMapResource::GetRayTracingCallableShaderLibrary(RayTracingCallableShaderLibrary, DefaultCallableShader);
+			FShaderMapResource::GetRayTracingCallableShaderLibrary(ShaderPlatform, RayTracingCallableShaderLibrary, DefaultCallableShader);
 		}
 		else
 		{
