[Shaders] Dump artifacts on failed serialization of shader compile job outputs.

Dumps the input and output of an FShaderCompileJob when the initial serialization of the output failed due to mismatch in input hash and shader target.
This also adds a new commandlet "DebugShaderCompileJob" to load up these artifacts and analyze the content with the debugger.

FShaderCompilingManager::CreateShaderDebugInfoPath() has been moved to FShaderCompilerInput::GetOrCreateShaderDebugInfoPath(),
so that this function can be used in RenderCore without module dependency changes.

#rnx

diff --git a/Engine/Source/Editor/UnrealEd/Classes/Commandlets/DebugShaderCompileJobCommandlet.h b/Engine/Source/Editor/UnrealEd/Classes/Commandlets/DebugShaderCompileJobCommandlet.h
new file mode 100644
--- /dev/null
+++ b/Engine/Source/Editor/UnrealEd/Classes/Commandlets/DebugShaderCompileJobCommandlet.h
@@ -0,0 +1,17 @@
+// Copyright Epic Games, Inc. All Rights Reserved.
+
+#pragma once
+
+#include "CoreMinimal.h"
+#include "Commandlets/Commandlet.h"
+#include "DebugShaderCompileJobCommandlet.generated.h"
+
+UCLASS(config=Editor)
+class UDebugShaderCompileJobCommandlet : public UCommandlet
+{
+	GENERATED_UCLASS_BODY()
+
+	//~ Begin UCommandlet Interface
+	virtual int32 Main(const FString& Params) override;
+	//~ End UCommandlet Interface
+};
diff --git a/Engine/Source/Editor/UnrealEd/Private/Commandlets/DebugShaderCompileJobCommandlet.cpp b/Engine/Source/Editor/UnrealEd/Private/Commandlets/DebugShaderCompileJobCommandlet.cpp
new file mode 100644
--- /dev/null
+++ b/Engine/Source/Editor/UnrealEd/Private/Commandlets/DebugShaderCompileJobCommandlet.cpp
@@ -0,0 +1,94 @@
+// Copyright Epic Games, Inc. All Rights Reserved.
+
+#include "Commandlets/DebugShaderCompileJobCommandlet.h"
+#include "ShaderCompiler.h"
+
+DEFINE_LOG_CATEGORY_STATIC(LogDebugShaderCompileJobCommandlet, Log, All);
+
+UDebugShaderCompileJobCommandlet::UDebugShaderCompileJobCommandlet(const FObjectInitializer& ObjectInitializer)
+	: Super(ObjectInitializer)
+{
+}
+
+// This function is intended for debugging, so disable optimizations here
+UE_DISABLE_OPTIMIZATION_SHIP
+
+int32 UDebugShaderCompileJobCommandlet::Main(const FString& Params)
+{
+	TRACE_CPUPROFILER_EVENT_SCOPE(UDebugShaderCompileJobCommandlet::Main);
+
+	StaticExec(nullptr, TEXT("log LogShaders Verbose"));
+	StaticExec(nullptr, TEXT("log LogShaderCompilers Verbose"));
+
+	TArray<FString> Tokens;
+	TArray<FString> Switches;
+	TMap<FString, FString> ParamVals;
+	UCommandlet::ParseCommandLine(*Params, Tokens, Switches, ParamVals);
+
+	// Display help
+	if (Switches.Contains("help"))
+	{
+		UE_LOG(LogDebugShaderCompileJobCommandlet, Display, TEXT("DebugShaderCompileJob FILE*"));
+		UE_LOG(LogDebugShaderCompileJobCommandlet, Display, TEXT("This commandlet loads the specified shader compile jobs that were serialized and uploaded as a cook artifact and then breaks the debugger."));
+		return 0;
+	}
+
+	// Validate input parameters
+	if (Tokens.IsEmpty())
+	{
+		UE_LOG(LogDebugShaderCompileJobCommandlet, Error, TEXT("Missing input filenames"));
+		return 1;
+	}
+
+	// Wait for a debugger to attach
+	{
+		constexpr double kNotificationWaitPeriod = 1.0;
+		constexpr double kMaxWaitTime = 10.0;
+
+		const double StartTime = FPlatformTime::Seconds();
+		double LastWaitTime = 0.0;
+
+		while (!FPlatformMisc::IsDebuggerPresent())
+		{
+			const double CurrentTime = FPlatformTime::Seconds();
+			if (CurrentTime > LastWaitTime + kNotificationWaitPeriod)
+			{
+				if (CurrentTime > StartTime + kMaxWaitTime)
+				{
+					UE_LOG(LogDebugShaderCompileJobCommandlet, Error, TEXT("No debugger attached after %.2fs => Exiting now"), CurrentTime - StartTime);
+					return 1;
+				}
+				else
+				{
+					LastWaitTime = CurrentTime;
+					UE_LOG(LogDebugShaderCompileJobCommandlet, Display, TEXT("Wait for debugger to attach ..."));
+				}
+			}
+			return 0;
+		}
+	}
+
+	// Load compile job
+	for (const FString& InputFilename : Tokens)
+	{
+		if (TUniquePtr<FArchive> InputFile = TUniquePtr<FArchive>(IFileManager::Get().CreateFileReader(*InputFilename)))
+		{
+			FShaderCompileJob Job;
+			Job.SerializeArtifact(*InputFile);
+
+			// Break debugger to investigate compile job
+			UE_DEBUG_BREAK();
+		}
+		else
+		{
+			UE_LOG(LogDebugShaderCompileJobCommandlet, Warning, TEXT("Failed to load input file: %s"), *InputFilename);
+		}
+	}
+
+	UE_LOG(LogDebugShaderCompileJobCommandlet, Display, TEXT("End debugging shader compile job"));
+
+	return 0;
+}
+
+// If this is part of a unity build file, turn optimizations back on
+UE_ENABLE_OPTIMIZATION_SHIP
diff --git a/Engine/Source/Runtime/Engine/Private/ShaderCompiler/ShaderCompiler.cpp b/Engine/Source/Runtime/Engine/Private/ShaderCompiler/ShaderCompiler.cpp
--- a/Engine/Source/Runtime/Engine/Private/ShaderCompiler/ShaderCompiler.cpp
+++ b/Engine/Source/Runtime/Engine/Private/ShaderCompiler/ShaderCompiler.cpp
@@ -1354,31 +1354,7 @@ EShaderDebugInfoFlags FShaderCompilingManager::GetDumpShaderDebugInfoFlags() con
 
 FString FShaderCompilingManager::CreateShaderDebugInfoPath(const FShaderCompilerInput& ShaderCompilerInput) const
 {
-	TRACE_CPUPROFILER_EVENT_SCOPE(FShaderCompilingManager::CreateShaderDebugInfoPath);
-
-	FString DumpDebugInfoPath = ShaderCompilerInput.DumpDebugInfoRootPath / ShaderCompilerInput.DebugGroupName + ShaderCompilerInput.DebugExtension;
-
-	// Sanitize the name to be used as a path
-	// List mostly comes from set of characters not allowed by windows in a path.  Just try to rename a file and type one of these for the list.
-	DumpDebugInfoPath.ReplaceInline(TEXT("<"), TEXT("("));
-	DumpDebugInfoPath.ReplaceInline(TEXT(">"), TEXT(")"));
-	DumpDebugInfoPath.ReplaceInline(TEXT("::"), TEXT("=="));
-	DumpDebugInfoPath.ReplaceInline(TEXT("|"), TEXT("_"));
-	DumpDebugInfoPath.ReplaceInline(TEXT("*"), TEXT("-"));
-	DumpDebugInfoPath.ReplaceInline(TEXT("?"), TEXT("!"));
-	DumpDebugInfoPath.ReplaceInline(TEXT("\""), TEXT("\'"));
-
-	if (!IFileManager::Get().DirectoryExists(*DumpDebugInfoPath))
-	{
-		if (!IFileManager::Get().MakeDirectory(*DumpDebugInfoPath, true))
-		{
-			const uint32 ErrorCode = FPlatformMisc::GetLastError();
-			UE_LOG(LogShaderCompilers, Warning, TEXT("Last Error %u: Failed to create directory for shader debug info '%s'. Try enabling large file paths or r.DumpShaderDebugShortNames."), ErrorCode, *DumpDebugInfoPath);
-			return FString();
-		}
-	}
-
-	return DumpDebugInfoPath;
+	return ShaderCompilerInput.GetOrCreateShaderDebugInfoPath();
 }
 
 bool FShaderCompilingManager::ShouldRecompileToDumpShaderDebugInfo(const FShaderCompileJob& Job) const
diff --git a/Engine/Source/Runtime/RenderCore/Private/ShaderCore.cpp b/Engine/Source/Runtime/RenderCore/Private/ShaderCore.cpp
--- a/Engine/Source/Runtime/RenderCore/Private/ShaderCore.cpp
+++ b/Engine/Source/Runtime/RenderCore/Private/ShaderCore.cpp
@@ -3598,6 +3598,35 @@ FString MakeInjectedShaderCodeBlock(const TCHAR* BlockName, const FString& CodeT
 	return FString::Printf(TEXT("#line 1 \"%s\"\n%s"), BlockName, *CodeToInject);
 }
 
+FString FShaderCompilerInput::GetOrCreateShaderDebugInfoPath() const
+{
+	TRACE_CPUPROFILER_EVENT_SCOPE(FShaderCompilerInput::GetOrCreateShaderDebugInfoPath);
+
+	FString OutDumpDebugInfoPath = FPaths::Combine(DumpDebugInfoRootPath, DebugGroupName + DebugExtension);
+
+	// Sanitize the name to be used as a path
+	// List mostly comes from set of characters not allowed by windows in a path.  Just try to rename a file and type one of these for the list.
+	OutDumpDebugInfoPath.ReplaceInline(TEXT("<"), TEXT("("));
+	OutDumpDebugInfoPath.ReplaceInline(TEXT(">"), TEXT(")"));
+	OutDumpDebugInfoPath.ReplaceInline(TEXT("::"), TEXT("=="));
+	OutDumpDebugInfoPath.ReplaceInline(TEXT("|"), TEXT("_"));
+	OutDumpDebugInfoPath.ReplaceInline(TEXT("*"), TEXT("-"));
+	OutDumpDebugInfoPath.ReplaceInline(TEXT("?"), TEXT("!"));
+	OutDumpDebugInfoPath.ReplaceInline(TEXT("\""), TEXT("\'"));
+
+	if (!IFileManager::Get().DirectoryExists(*OutDumpDebugInfoPath))
+	{
+		if (!IFileManager::Get().MakeDirectory(*OutDumpDebugInfoPath, true))
+		{
+			const uint32 ErrorCode = FPlatformMisc::GetLastError();
+			UE_LOG(LogShaders, Warning, TEXT("Last Error %u: Failed to create directory for shader debug info '%s'. Try enabling large file paths or r.DumpShaderDebugShortNames."), ErrorCode, *OutDumpDebugInfoPath);
+			return FString();
+		}
+	}
+
+	return OutDumpDebugInfoPath;
+}
+
 FString FShaderCompilerError::GetErrorStringWithSourceLocation() const
 {
 	if (!ErrorVirtualFilePath.IsEmpty() && !ErrorLineString.IsEmpty())
@@ -4210,14 +4239,37 @@ void FShaderCompileJob::AppendDebugName(FStringBuilderBase& OutName) const
 	OutName << (Input.DumpDebugInfoPath.IsEmpty() ? Input.DebugGroupName : Input.DumpDebugInfoPath);
 }
 
+static void DumpShaderCompileJobArtifact(FShaderCompileJob& Job, const TCHAR* InFilename)
+{
+	// Serialize compile job to binary file. Can be loaded and analyzed in debugger with commandlet FLoadShaderCompileJobCommandlet.
+	const FString DebugInfoPath = Job.Input.GetOrCreateShaderDebugInfoPath();
+	const FString JobDumpFilename = FPaths::Combine(DebugInfoPath, InFilename);
+	if (TUniquePtr<FArchive> JobDumpFile = TUniquePtr<FArchive>(IFileManager::Get().CreateFileWriter(*JobDumpFilename)))
+	{
+		Job.SerializeArtifact(*JobDumpFile);
+	}
+}
+
 void FShaderCompileJob::SerializeWorkerOutput(FArchive& Ar)
+{
+	SerializeWorkerOutputInner(Ar);
+}
+
+void FShaderCompileJob::SerializeWorkerOutputInner(FArchive& Ar, bool bSerializeForArtifact)
 {
 	Ar << Output;
 
 	bool bSecondaryOutput = SecondaryOutput.IsValid();
 	Ar << bSecondaryOutput;
 
-	checkf(Output.ValidateInputHash == Input.Hash, TEXT("Output.ValidateInputHash does not match Input.Hash; incorrect results associated with job?"));
+	if (!bSerializeForArtifact)
+	{
+		if (!ensureMsgf(Output.ValidateInputHash == Input.Hash, TEXT("Output.ValidateInputHash does not match Input.Hash; incorrect results associated with job?")))
+		{
+			DumpShaderCompileJobArtifact(*this, TEXT("ShaderCompileJob-InputHashMismatch"));
+			UE_LOG(LogShaders, Fatal, TEXT("Cannot continue serializing shader compile jobs after input hash mismatch in primary output"));
+		}
+	}
 	// empty validation hash after running validation to avoid impacting output deduplication
 	Output.ValidateInputHash = FShaderCompilerInputHash();
 
@@ -4228,7 +4280,14 @@ void FShaderCompileJob::SerializeWorkerOutput(FArchive& Ar)
 			SecondaryOutput = MakeUnique<FShaderCompilerOutput>();
 		}
 		Ar << *SecondaryOutput;
-		checkf(SecondaryOutput->ValidateInputHash == Input.Hash, TEXT("SecondaryOutput.ValidateInputHash does not match Input.Hash; incorrect results associated with job?"));
+		if (!bSerializeForArtifact)
+		{
+			if (!ensureMsgf(SecondaryOutput->ValidateInputHash == Input.Hash, TEXT("SecondaryOutput.ValidateInputHash does not match Input.Hash; incorrect results associated with job?")))
+			{
+				DumpShaderCompileJobArtifact(*this, TEXT("ShaderCompileJob-SecondaryOutput-InputHashMismatch"));
+				UE_LOG(LogShaders, Fatal, TEXT("Cannot continue serializing shader compile jobs after input hash mismatch in secondary output"));
+			}
+		}
 		// empty validation hash after running validation to avoid impacting output deduplication
 		SecondaryOutput->ValidateInputHash = FShaderCompilerInputHash();
 	}
@@ -4237,10 +4296,21 @@ void FShaderCompileJob::SerializeWorkerOutput(FArchive& Ar)
 	Ar << bSucceededTemp;
 	bSucceeded = bSucceededTemp;
 
-	checkf(Output.Target == Input.Target, TEXT("Output FShaderTarget does not match the input struct; incorrect results associated with job?"));
-	if (bSecondaryOutput)
+	if (!bSerializeForArtifact)
 	{
-		checkf(SecondaryOutput->Target == Input.Target, TEXT("SecondaryOutput FShaderTarget does not match the input struct; incorrect results associated with job?"));
+		if (!ensureMsgf(Output.Target == Input.Target, TEXT("Output FShaderTarget does not match the input struct; incorrect results associated with job?")))
+		{
+			DumpShaderCompileJobArtifact(*this, TEXT("ShaderCompileJob-ShaderTargetMismatch"));
+			UE_LOG(LogShaders, Fatal, TEXT("Cannot continue serializing shader compile jobs after target mismatch in primary output"));
+		}
+		if (bSecondaryOutput)
+		{
+			if (!ensureMsgf(SecondaryOutput->Target == Input.Target, TEXT("SecondaryOutput FShaderTarget does not match the input struct; incorrect results associated with job?")))
+			{
+				DumpShaderCompileJobArtifact(*this, TEXT("ShaderCompileJob-SecondaryOutput-ShaderTargetMismatch"));
+				UE_LOG(LogShaders, Fatal, TEXT("Cannot continue serializing shader compile jobs after target mismatch in secondary output"));
+			}
+		}
 	}
 }
 
@@ -4263,6 +4333,14 @@ void FShaderCompileJob::SerializeWorkerInput(FArchive& Ar)
 	}
 }
 
+void FShaderCompileJob::SerializeArtifact(FArchive& Ar)
+{
+	SerializeWorkerInput(Ar);
+
+	constexpr bool bSerializeForArtifact = true;
+	SerializeWorkerOutputInner(Ar, bSerializeForArtifact);
+}
+
 FStringView FShaderCompileJob::GetFinalSourceView() const
 {
 	 // any modifications to the source done as part of the compile step will be written to the "ModifiedShaderSource" field
diff --git a/Engine/Source/Runtime/RenderCore/Public/ShaderCompilerCore.h b/Engine/Source/Runtime/RenderCore/Public/ShaderCompilerCore.h
--- a/Engine/Source/Runtime/RenderCore/Public/ShaderCompilerCore.h
+++ b/Engine/Source/Runtime/RenderCore/Public/ShaderCompilerCore.h
@@ -469,6 +469,9 @@ struct FShaderCompilerInput
 
 		return RootParametersStructure != nullptr;
 	}
+
+	/** Returns the shader debug info path for this shader compiler input and create the directory if it doesn't exist yet. */
+	RENDERCORE_API FString GetOrCreateShaderDebugInfoPath() const;
 };
 
 /** A shader compiler error or warning. */
diff --git a/Engine/Source/Runtime/RenderCore/Public/ShaderCompilerJobTypes.h b/Engine/Source/Runtime/RenderCore/Public/ShaderCompilerJobTypes.h
--- a/Engine/Source/Runtime/RenderCore/Public/ShaderCompilerJobTypes.h
+++ b/Engine/Source/Runtime/RenderCore/Public/ShaderCompilerJobTypes.h
@@ -370,6 +370,9 @@ public:
 	// Serializes only the subset of data written by ShaderCompiler and read from SCW when using worker processes.
 	RENDERCORE_API void SerializeWorkerInput(FArchive& Ar);
 
+	// Serializes the compile job for a cook artifact for later analysis in a commandlet.
+	RENDERCORE_API void SerializeArtifact(FArchive& Ar);
+
 	RENDERCORE_API FStringView GetFinalSourceView() const;
 
 	virtual RENDERCORE_API void AppendDiagnostics(FString& OutDiagnostics, int32 InJobIndex, int32 InNumJobs, const TCHAR* Indentation = nullptr) const override final;
@@ -381,6 +384,10 @@ public:
 		FShaderCommonCompileJob(Type, InHash, InId, InPriroity),
 		Key(InKey)
 	{}
+
+private:
+	void SerializeWorkerOutputInner(FArchive& Ar, bool bSerializeForArtifact = false);
+
 };
 
 struct FShaderPipelineCompileJobKey
