Mip generation fix for PF_FloatRGB media textures.

Even if sRGB was requested upon texture creation (TexCreate_SRGB), there are no sRGB-encoded floating point texture platform formats.

diff --git a/Engine/Source/Runtime/RenderCore/Private/GenerateMips.cpp b/Engine/Source/Runtime/RenderCore/Private/GenerateMips.cpp
--- a/Engine/Source/Runtime/RenderCore/Private/GenerateMips.cpp
+++ b/Engine/Source/Runtime/RenderCore/Private/GenerateMips.cpp
@@ -207,7 +207,8 @@ void FGenerateMips::ExecuteCompute(FRDGBuilder& GraphBuilder, ERHIFeatureLevel::
 	const FRDGTextureDesc& TextureDesc = Texture->Desc;
 
 	// Select compute shader variant (normal vs. sRGB etc.)
-	bool bMipsSRGB = EnumHasAnyFlags(TextureDesc.Flags, TexCreate_SRGB);
+	// Note: Floating-point platform textures cannot be sRGB-encoded (invalid condition). We exclude them here for safety.
+	bool bMipsSRGB = EnumHasAnyFlags(TextureDesc.Flags, TexCreate_SRGB) && !IsFloatFormat(TextureDesc.Format);
 	const bool bMipsSwizzle = false; 
 
 	FGenerateMipsCS::FPermutationDomain PermutationVector;
