Fixed assert when in immediate mode for culled reserved commits.

diff --git a/Engine/Source/Runtime/RenderCore/Private/RenderGraphBuilder.cpp b/Engine/Source/Runtime/RenderCore/Private/RenderGraphBuilder.cpp
--- a/Engine/Source/Runtime/RenderCore/Private/RenderGraphBuilder.cpp
+++ b/Engine/Source/Runtime/RenderCore/Private/RenderGraphBuilder.cpp
@@ -4141,7 +4141,7 @@ void FRDGBuilder::AddLastBufferTransition(FRDGBuffer* Buffer)
 void FRDGBuilder::AddCulledReservedCommitTransition(FRDGBufferRef Buffer)
 {
 	check(Buffer->HasRHI() && Buffer->bExternal && Buffer->PendingCommitSize > 0);
-	check(Buffer->ReferenceCount == 0);
+	check(Buffer->ReferenceCount == 0 || IsImmediateMode());
 
 	const FRDGPassHandle ProloguePassHandle = GetProloguePassHandle();
 	const FRDGPassHandle EpiloguePassHandle = GetEpiloguePassHandle();
