[shaders] add a cvar that can be set to break in the debugger when preprocessing a specific named job (minor debug QoL improvement when iterating on preprocessor code)

diff --git a/Engine/Source/Runtime/RenderCore/Private/ShaderCore.cpp b/Engine/Source/Runtime/RenderCore/Private/ShaderCore.cpp
--- a/Engine/Source/Runtime/RenderCore/Private/ShaderCore.cpp
+++ b/Engine/Source/Runtime/RenderCore/Private/ShaderCore.cpp
@@ -112,6 +112,12 @@ static TAutoConsoleVariable<bool> CVarDumpDebugInfoForCacheHits(
 	TEXT("If true, debug info (via IShaderFormat::OutputDebugData) will be output for all jobs including duplicates and cache/DDC hits. If false, only jobs that actually executed compilation will dump debug info."),
 	ECVF_Default);
 
+static FString GBreakOnPreprocessJob;
+static FAutoConsoleVariableRef CVarBreakOnPreprocessJob(
+	TEXT("r.ShaderCompiler.BreakOnPreprocessJob"),
+	GBreakOnPreprocessJob,
+	TEXT("If specified, triggers a breakpoint when preprocessing a job whose name matches the given string (case-insensitive, substrings are supported)"));
+
 static TAutoConsoleVariable<FString> CVarShaderOverrideDebugDir(
 	TEXT("r.OverrideShaderDebugDir"),
 	"",
@@ -1586,6 +1592,13 @@ public:
 		FShaderCompileJob& Job, 
 		const FShaderCompilerEnvironment& Environment)
 	{
+		if (!GBreakOnPreprocessJob.IsEmpty())
+		{
+			if (Job.Input.DebugGroupName.Contains(GBreakOnPreprocessJob, ESearchCase::IgnoreCase))
+			{
+				UE_DEBUG_BREAK();
+			}
+		}
 		Job.PreprocessOutput.bSucceeded = Backend->PreprocessShader(Job.Input, Environment, Job.PreprocessOutput);
 		if (Job.PreprocessOutput.bSucceeded && Backend->RequiresSecondaryCompile(Job.Input, Environment, Job.PreprocessOutput))
 		{
