[shaders] add option to have shaderformat-specific configuration of symbol related cvars

diff --git a/Engine/Source/Runtime/RenderCore/Private/ShaderCore.cpp b/Engine/Source/Runtime/RenderCore/Private/ShaderCore.cpp
--- a/Engine/Source/Runtime/RenderCore/Private/ShaderCore.cpp
+++ b/Engine/Source/Runtime/RenderCore/Private/ShaderCore.cpp
@@ -546,9 +546,15 @@ struct FShaderSettingHelper
 		// Then check the per platform settings.
 		else if (FConfigCacheIni* PlatformConfig = GetPlatformConfigForFormat(ShaderFormat))
 		{
-			// Override with a build machine specific setting, if present.
 			bool bFoundConfig = false;
-			if (GIsBuildMachine && SettingSectionBuildMachine)
+			TStringBuilder<128> ShaderFormatStr;
+			ShaderFormat.ToString(ShaderFormatStr);
+			// first check for a shaderformat-specific value for the setting
+			bFoundConfig = PlatformConfig->GetValue(ShaderFormatStr.ToString(), SettingName, bEnabled, GEngineIni);
+
+			// if not found fall back to the per-platform value for compatibility with existing configs
+			// (from either of the configured settings sections, i.e. [ShaderCompiler] or [ShaderCompiler_BuildMachine])
+			if (!bFoundConfig && GIsBuildMachine && SettingSectionBuildMachine)
 			{
 				bFoundConfig = PlatformConfig->GetValue(SettingSectionBuildMachine, SettingName, bEnabled, GEngineIni);
 			}
