Clamp FGlobalDynamicReadBuffer size to MaxViewDimensionForTypedBuffer

diff --git a/Engine/Source/Runtime/RenderCore/Private/DynamicBufferAllocator.cpp b/Engine/Source/Runtime/RenderCore/Private/DynamicBufferAllocator.cpp
--- a/Engine/Source/Runtime/RenderCore/Private/DynamicBufferAllocator.cpp
+++ b/Engine/Source/Runtime/RenderCore/Private/DynamicBufferAllocator.cpp
@@ -150,7 +150,14 @@ FGlobalDynamicReadBuffer::FAllocation FGlobalDynamicReadBuffer::AllocateInternal
 		if (Buffer == nullptr)
 		{
 			const uint32 AlignedNum = FMath::DivideAndRoundUp(Num, (uint32)GAlignReadBufferRenderingBufferSize) * GAlignReadBufferRenderingBufferSize;
-			const uint32 NewBufferSize = FMath::Max(AlignedNum, (uint32)GMinReadBufferRenderingBufferSize);
+			const uint64 MaxViewDimensionForTypedBuffer = GRHIGlobals.MaxViewDimensionForTypedBuffer;
+			// Cannot have the MinReadBufferRenderingBufferSize be greater than MaxViewDimensionForTypedBuffer for that platform
+			const uint32 AdjustedMinReadBufferRenderingBufferSize = FMath::Min<uint32>(GMinReadBufferRenderingBufferSize, MaxViewDimensionForTypedBuffer);
+			const uint32 NewBufferSize = FMath::Max(AlignedNum, AdjustedMinReadBufferRenderingBufferSize);
+			if (UNLIKELY(NewBufferSize > MaxViewDimensionForTypedBuffer))
+			{
+				UE_LOG(LogRendererCore, Warning, TEXT("FGlobalDynamicReadBuffer::AllocateInternal, trying to allocate a Buffer of size (in elements) %d that is greater than the platform max of %llu"), NewBufferSize, MaxViewDimensionForTypedBuffer);
+			}
 			Buffer = new FDynamicAllocReadBuffer();
 			BufferPool->Buffers.Add(Buffer);
 			// Use static in D3D11 because we're running out of zeroed pages, and that causes the upload of this buffer to take a very long time. Static uses more memory and might be slightly less efficient
