Improve material dependency evaluation to avoid overly stressing TSet functions when iterating on shader type dependencies with many permutations of the same shader type.

diff --git a/Engine/Source/Runtime/RenderCore/Private/ShaderCore.cpp b/Engine/Source/Runtime/RenderCore/Private/ShaderCore.cpp
--- a/Engine/Source/Runtime/RenderCore/Private/ShaderCore.cpp
+++ b/Engine/Source/Runtime/RenderCore/Private/ShaderCore.cpp
@@ -3499,6 +3499,7 @@ void AppendShaderDependencies(
 	TSet<const FShaderParametersMetadata*, DefaultKeyFuncs<const FShaderParametersMetadata*>, FMemStackSetAllocator> ReferencedUniformBuffers;
 	ReferencedUniformBuffers.Reserve(128);
 
+	const FShaderType* LastShaderType = nullptr;
 	for (const FShaderTypeDependency& ShaderTypeDependency : ShaderTypeDependencies)
 	{
 		const FShaderType* ShaderType = FindShaderTypeByName(ShaderTypeDependency.ShaderTypeName);
@@ -3528,9 +3529,14 @@ void AppendShaderDependencies(
 
 		KeyGen.Append(GetShaderTypeLayoutHash(ShaderType->GetLayout(), LayoutParams));
 
-		for (const FShaderParametersMetadata* UniformBuffer : ShaderType->GetReferencedUniformBuffers())
+		if (ShaderType != LastShaderType)
 		{
-			ReferencedUniformBuffers.Add(UniformBuffer);
+			for (const FShaderParametersMetadata* UniformBuffer : ShaderType->GetReferencedUniformBuffers())
+			{
+				ReferencedUniformBuffers.Add(UniformBuffer);
+			}
+
+			LastShaderType = ShaderType;
 		}
 	}
 
