Added AddExistingCommandList to dispatch pass to support tasks that record work earlier in the frame. Also added support for querying whether async compute is enabled.

diff --git a/Engine/Source/Runtime/RenderCore/Private/RenderGraphPass.cpp b/Engine/Source/Runtime/RenderCore/Private/RenderGraphPass.cpp
--- a/Engine/Source/Runtime/RenderCore/Private/RenderGraphPass.cpp
+++ b/Engine/Source/Runtime/RenderCore/Private/RenderGraphPass.cpp
@@ -138,6 +138,25 @@ FRHICommandList* FRDGDispatchPassBuilder::CreateCommandList()
 	return RHICmdList;
 }
 
+void FRDGDispatchPassBuilder::AddExistingCommandList(FRHICommandList* RHICmdList)
+{
+#if RDG_ENABLE_DEBUG
+	for (FRHICommandListImmediate::FQueuedCommandList& QueuedCmdList : Pass->CommandLists)
+	{
+		checkf(QueuedCmdList.CmdList != RHICmdList, TEXT("Command list was already added!"));
+	}
+#endif
+
+	// When parallel executing we have to create a prologue command list to absorb the barriers and scopes.
+	if (Pass->bParallelExecute && Pass->CommandLists.IsEmpty())
+	{
+		FRHICommandList* RHICmdListPrologue = CreateCommandList();
+		RHICmdListPrologue->FinishRecording();
+	}
+
+	Pass->CommandLists.Emplace(RHICmdList);
+}
+
 void FRDGDispatchPassBuilder::Finish()
 {
 	// With serial execution the pass commands are embedded in the immediate command list instead.
@@ -416,4 +435,4 @@ const TCHAR* FRDGPass::GetName() const
 		return Name.GetTCHAR();
 	}
 }
-#endif
+#endif
\ No newline at end of file
diff --git a/Engine/Source/Runtime/RenderCore/Public/RenderGraphBuilder.h b/Engine/Source/Runtime/RenderCore/Public/RenderGraphBuilder.h
--- a/Engine/Source/Runtime/RenderCore/Public/RenderGraphBuilder.h
+++ b/Engine/Source/Runtime/RenderCore/Public/RenderGraphBuilder.h
@@ -308,6 +308,11 @@ public:
 		return ParallelSetup.bEnabled;
 	}
 
+	inline bool IsAsyncComputeEnabled() const
+	{
+		return bSupportsAsyncCompute;
+	}
+
 	/** Tells the builder to delete unused RHI resources. The behavior of this method depends on whether RDG immediate mode is enabled:
 	 *   Deferred:  RHI resource flushes are performed prior to execution.
 	 *   Immediate: RHI resource flushes are performed immediately.
diff --git a/Engine/Source/Runtime/RenderCore/Public/RenderGraphPass.h b/Engine/Source/Runtime/RenderCore/Public/RenderGraphPass.h
--- a/Engine/Source/Runtime/RenderCore/Public/RenderGraphPass.h
+++ b/Engine/Source/Runtime/RenderCore/Public/RenderGraphPass.h
@@ -749,8 +749,15 @@ private:
 class FRDGDispatchPassBuilder
 {
 public:
+	/** Create a new command list to record into and inserts it. Call FinishRecording() on the task when done. */
 	RENDERCORE_API FRHICommandList* CreateCommandList();
 
+	/** Insert a user-allocated command list into the dispatch pass. This version is useful if the command list
+     *  is recorded in a task that was launched prior to inserting the dispatch pass. Prefer CreateCommandList
+     *  if possible for performance. Call FinishRecording() on the task when done.
+     */
+	RENDERCORE_API void AddExistingCommandList(FRHICommandList* RHICmdList);
+
 private:
 	FRDGDispatchPassBuilder(FRDGDispatchPass* InPass)
 		: Pass(InPass)
