Fixed crash due to async clean-up task not releasing data during the task body but rather when it is destructed, which can happen after shutdown of the allocators since the task is statically initialized.

diff --git a/Engine/Source/Runtime/RenderCore/Private/RenderGraphBuilder.cpp b/Engine/Source/Runtime/RenderCore/Private/RenderGraphBuilder.cpp
--- a/Engine/Source/Runtime/RenderCore/Private/RenderGraphBuilder.cpp
+++ b/Engine/Source/Runtime/RenderCore/Private/RenderGraphBuilder.cpp
@@ -632,7 +632,13 @@ FRDGBuilder::FAsyncDeleter::~FAsyncDeleter()
 	if (Function)
 	{
 		// Launch the task with a prerequisite on any previously launched RDG async delete task.
-		LastTask = UE::Tasks::Launch(UE_SOURCE_LOCATION, [Function = MoveTemp(Function)]() mutable {}, MakeArrayView({ LastTask, Prerequisites }));
+		LastTask = UE::Tasks::Launch(UE_SOURCE_LOCATION, [Function = MoveTemp(Function)] () mutable
+		{
+			// Call and release the contents of the function inside the task lamda.
+			Function();
+			Function = {};
+
+		}, MakeArrayView({ LastTask, Prerequisites }));
 	}
 }
 
