Add more complete information about velocity output pass setting to shader map key. Previously we didn't distinguish between all values which may have led to DDC corruption.

diff --git a/Engine/Source/Runtime/RenderCore/Private/Shader.cpp b/Engine/Source/Runtime/RenderCore/Private/Shader.cpp
--- a/Engine/Source/Runtime/RenderCore/Private/Shader.cpp
+++ b/Engine/Source/Runtime/RenderCore/Private/Shader.cpp
@@ -1709,10 +1709,21 @@ void ShaderMapAppendKey(EShaderPlatform Platform, FShaderKeyGenerator& KeyGen)
 		}
 	}
 
-	if (IsUsingBasePassVelocity(Platform))
 	{
-		KeyGen.AppendSeparator();
-		KeyGen.Append(TEXT("GV"));
+		static const auto CVar = IConsoleManager::Get().FindTConsoleVariableDataInt(TEXT("r.VelocityOutputPass"));
+		const int32 VelocityOutputPassValue = CVar ? CVar->GetValueOnAnyThread() : 0;
+		if (VelocityOutputPassValue == 1)
+		{
+			KeyGen.AppendSeparator();
+			KeyGen.Append(TEXT("GV"));
+		}
+		else
+		{
+			KeyGen.AppendSeparator();
+			KeyGen.Append(TEXT("VOP"));
+			KeyGen.AppendSeparator();
+			KeyGen.Append(VelocityOutputPassValue);
+		}
 	}
 
 	{
