Fixed buffer size callbacks not being called in immediate mode.

diff --git a/Engine/Source/Runtime/RenderCore/Private/RenderGraphBuilder.cpp b/Engine/Source/Runtime/RenderCore/Private/RenderGraphBuilder.cpp
--- a/Engine/Source/Runtime/RenderCore/Private/RenderGraphBuilder.cpp
+++ b/Engine/Source/Runtime/RenderCore/Private/RenderGraphBuilder.cpp
@@ -1812,13 +1812,7 @@ void FRDGBuilder::Execute()
 
 		UE::Tasks::FTask BufferNumElementsCallbacksTask = AddSetupTask([this]
 		{
-			SCOPED_NAMED_EVENT_TEXT("FRDGBuilder::FinalizeDescs", FColor::Magenta);
-
-			for (FRDGBuffer* Buffer : NumElementsCallbackBuffers)
-			{
-				Buffer->FinalizeDesc();
-			}
-			NumElementsCallbackBuffers.Empty();
+			FinalizeDescs();
 
 		}, TaskPriority, bParallelCompileBuffers && !NumElementsCallbackBuffers.IsEmpty());
 
@@ -2577,6 +2571,7 @@ void FRDGBuilder::SetupAuxiliaryPasses(FRDGPass* Pass)
 		FCollectResourceContext Context;
 		SubmitBufferUploads(RHICmdList);
 		CompilePassOps(Pass);
+		FinalizeDescs();
 		CollectAllocations(Context, Pass);
 		AllocatePooledTextures(RHICmdList, Context.PooledTextures);
 		AllocatePooledBuffers(RHICmdList, Context.PooledBuffers);
@@ -3514,6 +3509,16 @@ void FRDGBuilder::ExecuteSerialPass(FRHIComputeCommandList& RHICmdListPass, FRDG
 
 ///////////////////////////////////////////////////////////////////////////////////////////////////
 
+void FRDGBuilder::FinalizeDescs()
+{
+	SCOPED_NAMED_EVENT_TEXT("FRDGBuilder::FinalizeDescs", FColor::Magenta);
+	for (FRDGBuffer* Buffer : NumElementsCallbackBuffers)
+	{
+		Buffer->FinalizeDesc();
+	}
+	NumElementsCallbackBuffers.Empty();
+}
+
 void FRDGBuilder::CollectAllocations(FCollectResourceContext& Context, FRDGPass* Pass)
 {
 	for (FRDGPass* PassToBegin : Pass->ResourcesToBegin)
diff --git a/Engine/Source/Runtime/RenderCore/Public/RenderGraphBuilder.h b/Engine/Source/Runtime/RenderCore/Public/RenderGraphBuilder.h
--- a/Engine/Source/Runtime/RenderCore/Public/RenderGraphBuilder.h
+++ b/Engine/Source/Runtime/RenderCore/Public/RenderGraphBuilder.h
@@ -689,6 +689,9 @@ private:
 	Experimental::TRobinHoodHashMap<FRDGPooledTexture*, FRDGTexture*, DefaultKeyFuncs<FRDGPooledTexture*>, FConcurrentLinearArrayAllocator> PooledTextureOwnershipMap;
 	Experimental::TRobinHoodHashMap<FRDGPooledBuffer*, FRDGBuffer*, DefaultKeyFuncs<FRDGPooledBuffer*>, FConcurrentLinearArrayAllocator> PooledBufferOwnershipMap;
 
+	/** Finalizes the resource descriptors by calling callbacks to gather resource sizes. */
+	void FinalizeDescs();
+
 	/** Collects new resource allocations for the pass into the provided context. */
 	void CollectAllocations(FCollectResourceContext& Context, FRDGPass* Pass);
 	void CollectAllocateTexture(FCollectResourceContext& Context, ERHIPipeline PassPipeline, FRDGPassHandle PassHandle, FRDGTexture* Texture);
