[Shaders] Improve assertion for missing entry in shader parameter list.

Instead of crashing with an assertion inside TMap::FindChecked(), add an extra assertion in AddUnboundShaderParameterError() to report the name of the missing parameter.

#rnx

diff --git a/Engine/Source/Runtime/RenderCore/Public/ShaderParameterParser.h b/Engine/Source/Runtime/RenderCore/Public/ShaderParameterParser.h
--- a/Engine/Source/Runtime/RenderCore/Public/ShaderParameterParser.h
+++ b/Engine/Source/Runtime/RenderCore/Public/ShaderParameterParser.h
@@ -146,7 +146,9 @@ public:
 	/** Gets parsing information from a parameter binding name. */
 	const FParsedShaderParameter& FindParameterInfos(const FString& ParameterName) const
 	{
-		return ParsedParameters.FindChecked(ParameterName);
+		const FParsedShaderParameter* ShaderParameter = ParsedParameters.Find(ParameterName);
+		checkf(ShaderParameter != nullptr, TEXT("Missing shader parameter '%s'"), *ParameterName);
+		return *ShaderParameter;
 	}
 
 	const FParsedShaderParameter* FindParameterInfosUnsafe(const FString& ParameterName) const
