[ShaderCore]
- Use RAII on pointer to fix potential leak

diff --git a/Engine/Source/Runtime/RenderCore/Private/ShaderCore.cpp b/Engine/Source/Runtime/RenderCore/Private/ShaderCore.cpp
--- a/Engine/Source/Runtime/RenderCore/Private/ShaderCore.cpp
+++ b/Engine/Source/Runtime/RenderCore/Private/ShaderCore.cpp
@@ -40,6 +40,7 @@
 #include "Stats/StatsMisc.h"
 #include "String/Find.h"
 #include "Tasks/Task.h"
+#include "Templates/UniquePtr.h"
 #include "VertexFactory.h"
 #if WITH_EDITOR
 #include "Misc/ConfigCacheIni.h"
@@ -2654,9 +2655,9 @@ static void StringCopyToAnsiCharArray(const TCHAR* Text, int32 TextLen, TArray<A
 }
 
 // Allocates structure and adds root file dependency
-static FShaderPreprocessDependencies* ShaderPreprocessDependenciesBegin(const TCHAR* VirtualFilePath)
+static TUniquePtr<FShaderPreprocessDependencies> ShaderPreprocessDependenciesBegin(const TCHAR* VirtualFilePath)
 {
-	FShaderPreprocessDependencies* PreprocessDependencies = new FShaderPreprocessDependencies();
+	TUniquePtr<FShaderPreprocessDependencies> PreprocessDependencies = MakeUnique<FShaderPreprocessDependencies>();
 
 	PreprocessDependencies->Dependencies.AddDefaulted();
 	StringCopyToAnsiCharArray(VirtualFilePath, FCString::Strlen(VirtualFilePath), PreprocessDependencies->Dependencies[0].ResultPath);
@@ -2666,7 +2667,7 @@ static FShaderPreprocessDependencies* ShaderPreprocessDependenciesBegin(const TC
 }
 
 // Adds finished dependencies to the cache
-static void ShaderPreprocessDependenciesEnd(const TCHAR* VirtualFilePath, FShaderPreprocessDependencies* PreprocessDependencies, EShaderPlatform Platform)
+static void ShaderPreprocessDependenciesEnd(const TCHAR* VirtualFilePath, TUniquePtr<FShaderPreprocessDependencies> PreprocessDependencies, EShaderPlatform Platform)
 {
 	uint32 CurrentHash = FCrc::Strihash_DEPRECATED(VirtualFilePath);
 	FShaderFileCache& ShaderFileCache = GShaderFileCache[CurrentHash % GSHADERFILECACHE_BUCKETS];
@@ -2678,11 +2679,7 @@ static void ShaderPreprocessDependenciesEnd(const TCHAR* VirtualFilePath, FShade
 			// Another thread could have finished the job...  If not, set the dependencies.
 			if (!CachedFile->Dependencies.IsValid())
 			{
-				CachedFile->Dependencies = MakeShareable(PreprocessDependencies);
-			}
-			else
-			{
-				delete PreprocessDependencies;
+				CachedFile->Dependencies = MakeShareable(PreprocessDependencies.Release());
 			}
 		}
 	}
@@ -2814,7 +2811,7 @@ static void InternalGetShaderIncludes(const TCHAR* EntryPointVirtualFilePath, co
 						if (SeenFilenameIndex == INDEX_NONE)
 						{
 							// Preprocess dependencies don't include platform files.
-							FShaderPreprocessDependencies* ExtractedIncludeDependencies = nullptr;
+							TUniquePtr<FShaderPreprocessDependencies> ExtractedIncludeDependencies;
 							if (OutDependencies && !bIsPlatformFile)
 							{
 								ExtractedIncludeDependencies = ShaderPreprocessDependenciesBegin(*ExtractedIncludeFilename);
@@ -2825,18 +2822,14 @@ static void InternalGetShaderIncludes(const TCHAR* EntryPointVirtualFilePath, co
 							LoadShaderSourceFile(*ExtractedIncludeFilename, ShaderPlatform, &IncludedFileContents, nullptr, ShaderPlatformName,
 								ExtractedIncludeDependencies ? &ExtractedIncludeDependencies->Dependencies[0].StrippedSource : nullptr);
 
-							InternalGetShaderIncludes(EntryPointVirtualFilePath, *ExtractedIncludeFilename, IncludedFileContents, IncludeVirtualFilePaths, ShaderPlatform, DepthLimit - 1, true, ShaderPlatformName, ExtractedIncludeDependencies);
+							InternalGetShaderIncludes(EntryPointVirtualFilePath, *ExtractedIncludeFilename, IncludedFileContents, IncludeVirtualFilePaths, ShaderPlatform, DepthLimit - 1, true, ShaderPlatformName, ExtractedIncludeDependencies.Get());
 
 							if (ExtractedIncludeDependencies)
 							{
 								// Some generated shaders are referenced as includes, and won't be found -- if so, just delete the dependencies
 								if (ExtractedIncludeDependencies->Dependencies[0].StrippedSource.IsValid())
 								{
-									ShaderPreprocessDependenciesEnd(*ExtractedIncludeFilename, ExtractedIncludeDependencies, ShaderPlatform);
-								}
-								else
-								{
-									delete ExtractedIncludeDependencies;
+									ShaderPreprocessDependenciesEnd(*ExtractedIncludeFilename, MoveTemp(ExtractedIncludeDependencies), ShaderPlatform);
 								}
 							}
 						}
@@ -2938,7 +2931,7 @@ bool GetShaderPreprocessDependencies(const TCHAR* VirtualFilePath, EShaderPlatfo
  */
 static void InternalGetShaderIncludes(const TCHAR* EntryPointVirtualFilePath, const TCHAR* VirtualFilePath, TArray<FString>& IncludeVirtualFilePaths, EShaderPlatform ShaderPlatform, uint32 DepthLimit, bool AddToIncludeFile, const FName* ShaderPlatformName, bool bPreprocessDependencies)
 {
-	FShaderPreprocessDependencies* PreprocessDependencies = nullptr;
+	TUniquePtr<FShaderPreprocessDependencies> PreprocessDependencies;
 	if (bPreprocessDependencies)
 	{
 		// Check if they've already been generated.  These are platform independent, so we only need to generate them once if multiple platforms are being cooked,
@@ -2955,12 +2948,12 @@ static void InternalGetShaderIncludes(const TCHAR* EntryPointVirtualFilePath, co
 	FString FileContents;
 	LoadShaderSourceFile(VirtualFilePath, ShaderPlatform, &FileContents, nullptr, ShaderPlatformName, PreprocessDependencies ? &PreprocessDependencies->Dependencies[0].StrippedSource : nullptr);
 
-	InternalGetShaderIncludes(EntryPointVirtualFilePath, VirtualFilePath, FileContents, IncludeVirtualFilePaths, ShaderPlatform, DepthLimit, AddToIncludeFile, ShaderPlatformName, PreprocessDependencies);
+	InternalGetShaderIncludes(EntryPointVirtualFilePath, VirtualFilePath, FileContents, IncludeVirtualFilePaths, ShaderPlatform, DepthLimit, AddToIncludeFile, ShaderPlatformName, PreprocessDependencies.Get());
 
 	if (PreprocessDependencies)
 	{
 		// Adds completed dependency structure to shader cache map entry
-		ShaderPreprocessDependenciesEnd(VirtualFilePath, PreprocessDependencies, ShaderPlatform);
+		ShaderPreprocessDependenciesEnd(VirtualFilePath, MoveTemp(PreprocessDependencies), ShaderPlatform);
 	}
 }
 
