ListShaders command

#tests Run ListShaders in Editor and on Client

diff --git a/Engine/Source/Runtime/Engine/Private/Materials/MaterialShader.cpp b/Engine/Source/Runtime/Engine/Private/Materials/MaterialShader.cpp
--- a/Engine/Source/Runtime/Engine/Private/Materials/MaterialShader.cpp
+++ b/Engine/Source/Runtime/Engine/Private/Materials/MaterialShader.cpp
@@ -3343,6 +3343,21 @@ const FMemoryImageString* FMaterialShaderMap::GetShaderSource(const FVertexFacto
 
 #endif // WITH_EDITOR
 
+#if !UE_BUILD_SHIPPING
+void FMaterialShaderMap::DumpShaderList(FShaderListReport& Out) const
+{
+	Super::DumpShaderList(Out);
+	
+	for (FMeshMaterialShaderMap* MeshShaderMap : GetContent()->OrderedMeshShaderMaps)
+	{
+		if (MeshShaderMap)
+		{
+			MeshShaderMap->DumpShaderList(*this, Out);
+		}
+	}
+}
+#endif // !UE_BUILD_SHIPPING
+
 void FMaterialShaderMap::GetShaderList(TMap<FShaderId, TShaderRef<FShader>>& OutShaders) const
 {
 	FSHAHash MaterialShaderMapHash;
diff --git a/Engine/Source/Runtime/Engine/Public/MaterialShared.h b/Engine/Source/Runtime/Engine/Public/MaterialShared.h
--- a/Engine/Source/Runtime/Engine/Public/MaterialShared.h
+++ b/Engine/Source/Runtime/Engine/Public/MaterialShared.h
@@ -1543,6 +1543,10 @@ public:
 
 	/** Builds a list of the shaders in a shader map. Key is FShaderType::TypeName */
 	ENGINE_API virtual void GetShaderList(TMap<FHashedName, TShaderRef<FShader>>& OutShaders) const override;
+	
+#if !UE_BUILD_SHIPPING
+	virtual void DumpShaderList(FShaderListReport& Out) const;
+#endif // !UE_BUILD_SHIPPING
 
 	/** Builds a list of the shader pipelines in a shader map. */
 	ENGINE_API virtual void GetShaderPipelineList(TArray<FShaderPipelineRef>& OutShaderPipelines) const override;
diff --git a/Engine/Source/Runtime/RHI/Public/RHIShaderLibrary.h b/Engine/Source/Runtime/RHI/Public/RHIShaderLibrary.h
--- a/Engine/Source/Runtime/RHI/Public/RHIShaderLibrary.h
+++ b/Engine/Source/Runtime/RHI/Public/RHIShaderLibrary.h
@@ -29,6 +29,7 @@ public:
 	virtual FSHAHash GetShaderHash(int32 ShaderMapIndex, int32 ShaderIndex) = 0;
 	virtual int32 FindShaderMapIndex(const FSHAHash& Hash) = 0;
 	virtual int32 FindShaderIndex(const FSHAHash& Hash) = 0;
+	virtual uint32 GetShaderSizeBytes(int32 ShaderIndex) const { return 0; }
 	virtual bool IsPreloading(int32 ShaderIndex, FGraphEventArray& OutCompletionEvents) { return false; }
 	virtual bool PreloadShader(int32 ShaderIndex, FGraphEventArray& OutCompletionEvents) { return false; }
 	virtual bool PreloadShaderMap(int32 ShaderMapIndex, FGraphEventArray& OutCompletionEvents) { return false; }
diff --git a/Engine/Source/Runtime/RenderCore/Private/ShaderCodeArchive.cpp b/Engine/Source/Runtime/RenderCore/Private/ShaderCodeArchive.cpp
--- a/Engine/Source/Runtime/RenderCore/Private/ShaderCodeArchive.cpp
+++ b/Engine/Source/Runtime/RenderCore/Private/ShaderCodeArchive.cpp
@@ -2087,6 +2087,11 @@ void FIoStoreShaderCodeArchive::MarkPreloadEntrySkipped(int32 ShaderGroupIndex
 	}
 }
 
+uint32 FIoStoreShaderCodeArchive::GetShaderSizeBytes(int32 ShaderIndex) const
+{
+	return Header.GetShaderUncompressedSize(ShaderIndex);
+}
+
 void FIoStoreShaderCodeArchive::AddRefPreloadedShaderGroup(int32 ShaderGroupIndex)
 {
 	FWriteScopeLock Lock(PreloadedShaderGroupsLock);
diff --git a/Engine/Source/Runtime/RenderCore/Private/ShaderCodeLibrary.cpp b/Engine/Source/Runtime/RenderCore/Private/ShaderCodeLibrary.cpp
--- a/Engine/Source/Runtime/RenderCore/Private/ShaderCodeLibrary.cpp
+++ b/Engine/Source/Runtime/RenderCore/Private/ShaderCodeLibrary.cpp
@@ -309,6 +309,7 @@ public:
 	virtual int32 GetGroupIndexForShader(int32 ShaderIndex) const override;
 	virtual int32 GetLibraryId() const override;
 	virtual int32 GetLibraryShaderIndex(int32 ShaderIndex) const override;
+	virtual uint32 GetShaderSizeBytes(int32 ShaderIndex) const override;
 
 	class FShaderLibraryInstance* LibraryInstance;
 	int32 ShaderMapIndex;
@@ -1387,6 +1388,12 @@ FRHIShader* FShaderMapResource_SharedCode::CreateRHIShaderOrCrash(int32 ShaderIn
 	return CreatedShader;
 }
 
+uint32 FShaderMapResource_SharedCode::GetShaderSizeBytes(int32 ShaderIndex) const
+{
+	const int32 LibraryShaderIndex = LibraryInstance->Library->GetShaderIndex(ShaderMapIndex, ShaderIndex);
+	return LibraryInstance->Library->GetShaderSizeBytes(LibraryShaderIndex);
+}
+
 void FShaderMapResource_SharedCode::ReleasePreloadedShaderCode(int32 ShaderIndex)
 {
 	SCOPED_LOADTIMER(FShaderMapResource_SharedCode_InitRHI);	// part of shader initialization in a way
diff --git a/Engine/Source/Runtime/RenderCore/Private/ShaderMap.cpp b/Engine/Source/Runtime/RenderCore/Private/ShaderMap.cpp
--- a/Engine/Source/Runtime/RenderCore/Private/ShaderMap.cpp
+++ b/Engine/Source/Runtime/RenderCore/Private/ShaderMap.cpp
@@ -42,6 +42,42 @@ TAutoConsoleVariable<bool> CVarEnableShaderMapTracking(
 	TEXT("Enables the tracking of every shadermap instantiated. Required to run ListShaderMaps command."),
 	ECVF_ReadOnly);
 
+FAutoConsoleCommandWithArgsAndOutputDevice GListShaderCmd(
+	TEXT("ListShaders"),
+	TEXT("Spits out a csv table containing stats of all shaders"),
+	FConsoleCommandWithArgsAndOutputDeviceDelegate::CreateStatic(
+		[](const TArray<FString>& Params, FOutputDevice& Out)
+		{
+#if !UE_BUILD_SHIPPING
+			if (!CVarEnableShaderMapTracking.GetValueOnGameThread())
+			{
+				UE_LOG(LogShaders, Warning, TEXT("Enable r.TrackShaderMaps in ini config to enable the functionality"));
+				return;
+			}
+			Out.Logf(TEXT("FriendlyName,OwnerName,VertexFactoryName,ShaderTypeName,ShaderPermutation,UncompressedSizeKb,Hash,UsedForRenderingRHIPtr"));
+			FScopeLock SMAccess(&GAllShaderMapsGuard);
+			for (FShaderMapBase* ShaderMap : GAllShaderMaps)
+			{
+				FShaderListReport Report;
+				Report.BindLambda(
+					[&Out](
+						const TCHAR* FriendlyName,
+						const TCHAR* OwnerName,
+						const TCHAR* VertexFactory,
+						const TCHAR* ShaderType,
+						uint32 Permutation,
+						uint32 Bytes,
+						const TCHAR* Hash,
+						void* Ptr)
+					{
+						Out.Logf(TEXT("%s,%s,%s,%s,%d,%.2f,%s,%p"), FriendlyName, OwnerName, VertexFactory, ShaderType, Permutation, Bytes / 1024.f, Hash, Ptr);
+					}
+				);
+				ShaderMap->DumpShaderList(Report);
+			}
+#endif // !UE_BUILD_SHIPPING
+		}));
+
 FAutoConsoleCommandWithArgsAndOutputDevice GListShaderMapsCmd(
 	TEXT("ListShaderMaps"),
 	TEXT("Spits out a csv table containing stats of all shadermaps"),
@@ -618,6 +654,48 @@ void FShaderMapContent::RemoveShaderPipelineType(const FShaderPipelineType* Shad
 	}
 }
 
+#if !UE_BUILD_SHIPPING
+void FShaderMapContent::DumpShaderList(const FShaderMapBase& InShaderMap, FShaderListReport& Out) const
+{
+	FString FriendlyName;
+	FString OwnerName;
+	FShaderMapResource* Resource = InShaderMap.GetResource();
+	if (Resource)
+	{
+		FriendlyName = Resource->GetFriendlyName();
+		OwnerName = Resource->GetOwnerName().ToString();
+	}
+	for (int32 ShaderIndex = 0; ShaderIndex < Shaders.Num(); ++ShaderIndex)
+	{
+		FShader* Shader = Shaders[ShaderIndex].GetChecked();
+		uint32 ShaderCodeSizeBytes = 0;
+		FRHIShader* RHIShader = nullptr;
+		FString Hash;
+		if (Resource)
+		{
+			if (Resource->HasShader(Shader->GetResourceIndex()))
+			{
+				RHIShader = Resource->GetShader(Shader->GetResourceIndex());
+			}
+			const uint32 Index = Shader->GetResourceIndex();
+			ShaderCodeSizeBytes = Resource->GetShaderSizeBytes(Index);
+			Hash = Resource->GetShaderHash(Index).ToString();
+		}
+
+		Out.ExecuteIfBound(
+			*FriendlyName, 
+			*OwnerName,
+			Shader->GetVertexFactoryType(InShaderMap.GetPointerTable()) ? Shader->GetVertexFactoryType(InShaderMap.GetPointerTable())->GetName() : TEXT(""),
+			Shader->GetType(InShaderMap.GetPointerTable()) ? Shader->GetType(InShaderMap.GetPointerTable())->GetName() : TEXT(""),
+			ShaderPermutations[ShaderIndex],
+			ShaderCodeSizeBytes,
+			*Hash,
+			RHIShader
+		);
+	}	
+}
+#endif // !UE_BUILD_SHIPPING
+
 void FShaderMapContent::GetShaderList(const FShaderMapBase& InShaderMap, const FSHAHash& InMaterialShaderMapHash, TMap<FShaderId, TShaderRef<FShader>>& OutShaders) const
 {
 	for (int32 ShaderIndex = 0; ShaderIndex < Shaders.Num(); ++ShaderIndex)
diff --git a/Engine/Source/Runtime/RenderCore/Private/ShaderResource.cpp b/Engine/Source/Runtime/RenderCore/Private/ShaderResource.cpp
--- a/Engine/Source/Runtime/RenderCore/Private/ShaderResource.cpp
+++ b/Engine/Source/Runtime/RenderCore/Private/ShaderResource.cpp
@@ -487,6 +487,12 @@ void FShaderMapResource::AddRef()
 	NumRefs.fetch_add(1, std::memory_order_relaxed);
 }
 
+uint32 FShaderMapResource::GetShaderSizeBytes(int32 ShaderIndex) const
+{
+	// default FShaderMapResource does not contain inline shader bytecode
+	return 0;
+}
+
 void FShaderMapResource::Release()
 {
 	check(NumRefs.load(std::memory_order_relaxed) > 0);
@@ -768,3 +774,14 @@ uint32 FShaderMapResource_InlineCode::GetSizeBytes() const
 
 	return TotalSize;
 }
+
+uint32 FShaderMapResource_InlineCode::GetShaderSizeBytes(int32 ShaderIndex) const
+{
+	if (Code)
+	{
+		return Code->ShaderCodeResources[ShaderIndex].GetUncompressedSize();
+	}
+	return 0u;
+}
+
+
diff --git a/Engine/Source/Runtime/RenderCore/Public/Shader.h b/Engine/Source/Runtime/RenderCore/Public/Shader.h
--- a/Engine/Source/Runtime/RenderCore/Public/Shader.h
+++ b/Engine/Source/Runtime/RenderCore/Public/Shader.h
@@ -359,6 +359,8 @@ public:
 		return RHIShaders[ShaderIndex].load(std::memory_order_acquire) != nullptr;
 	}
 
+	virtual uint32 GetShaderSizeBytes(int32 ShaderIndex) const;
+
 	virtual void PreloadShader(int32 ShaderIndex, FGraphEventArray& OutCompletionEvents)
 	{
 		/* no-op when not using shader library */
@@ -624,6 +626,8 @@ public:
 	RENDERCORE_API virtual FString GetFriendlyName() const override { return TEXT("InlineCode"); }
 	virtual uint32 GetSizeBytes() const override;
 
+	virtual uint32 GetShaderSizeBytes(int32 ShaderIndex) const override;
+
 	TRefCountPtr<FShaderMapResourceCode> Code;
 };
 
@@ -2191,6 +2195,8 @@ private:
 	const FShaderMapBase* ShaderMap;
 };
 
+DECLARE_DELEGATE_EightParams(FShaderListReport, const TCHAR*, const TCHAR*, const TCHAR*, const TCHAR*, uint32, uint32, const TCHAR*, void*);
+
 /** A collection of shaders of different types */
 class FShaderMapContent
 {
@@ -2279,6 +2285,10 @@ public:
 
 	RENDERCORE_API void RemoveShaderPipelineType(const FShaderPipelineType* ShaderPipelineType);
 
+#if !UE_BUILD_SHIPPING
+	RENDERCORE_API void DumpShaderList(const FShaderMapBase& InShaderMap, FShaderListReport& Out) const;
+#endif // !UE_BUILD_SHIPPING
+
 	/** Builds a list of the shaders in a shader map. */
 	RENDERCORE_API void GetShaderList(const FShaderMapBase& InShaderMap, const FSHAHash& InMaterialShaderMapHash, TMap<FShaderId, TShaderRef<FShader>>& OutShaders) const;
 
@@ -2378,6 +2388,13 @@ public:
 	/** Builds a list of the shaders in a shader map. Key is FShaderType::TypeName */
 	RENDERCORE_API virtual void GetShaderList(TMap<FHashedName, TShaderRef<FShader>>& OutShaders) const = 0;
 	RENDERCORE_API virtual void GetShaderPipelineList(TArray<FShaderPipelineRef>& OutShaderPipelines) const = 0;
+	
+#if !UE_BUILD_SHIPPING
+	virtual void DumpShaderList(FShaderListReport& Out) const
+	{
+		GetContent()->DumpShaderList(*this, Out);
+	}
+#endif // !UE_BUILD_SHIPPING
 
 	inline EShaderPlatform GetShaderPlatform() const { return Content.Object ? Content.Object->GetShaderPlatform() : SP_NumPlatforms; }
 	inline uint32 GetFrozenContentSize() const { return Content.FrozenSize; }
diff --git a/Engine/Source/Runtime/RenderCore/Public/ShaderCodeArchive.h b/Engine/Source/Runtime/RenderCore/Public/ShaderCodeArchive.h
--- a/Engine/Source/Runtime/RenderCore/Public/ShaderCodeArchive.h
+++ b/Engine/Source/Runtime/RenderCore/Public/ShaderCodeArchive.h
@@ -581,7 +581,7 @@ public:
 
 	friend RENDERCORE_API FArchive& operator <<(FArchive& Ar, FIoStoreShaderCodeArchiveHeader& Ref);
 
-	inline uint64 GetShaderUncompressedSize(int ShaderIndex)
+	inline uint64 GetShaderUncompressedSize(int ShaderIndex) const
 	{
 		const FIoStoreShaderCodeEntry& ThisShaderEntry = ShaderEntries[ShaderIndex];
 		const FIoStoreShaderGroupEntry& GroupEntry = ShaderGroupEntries[ThisShaderEntry.ShaderGroupIndex];
@@ -642,6 +642,8 @@ public:
 			PreloadedShaderGroups.GetAllocatedSize();
 	}
 
+	virtual uint32 GetShaderSizeBytes(int32 ShaderIndex) const override;
+
 	virtual int32 GetNumShaders() const override { return Header.ShaderEntries.Num(); }
 	virtual int32 GetNumShaderMaps() const override { return Header.ShaderMapEntries.Num(); }
 	virtual int32 GetNumShadersForShaderMap(int32 ShaderMapIndex) const override { return Header.ShaderMapEntries[ShaderMapIndex].NumShaders; }
