Disable async compute on RDG for small builders.

diff --git a/Engine/Source/Runtime/RenderCore/Private/RenderGraphBuilder.cpp b/Engine/Source/Runtime/RenderCore/Private/RenderGraphBuilder.cpp
--- a/Engine/Source/Runtime/RenderCore/Private/RenderGraphBuilder.cpp
+++ b/Engine/Source/Runtime/RenderCore/Private/RenderGraphBuilder.cpp
@@ -582,7 +582,8 @@ FRDGBuilder::FRDGBuilder(FRHICommandListImmediate& InRHICmdList, FRDGEventName I
 {
 	ProloguePass = SetupEmptyPass(Passes.Allocate<FRDGSentinelPass>(Allocators.Root, RDG_EVENT_NAME("Graph Prologue (Graphics)")));
 
-	bSupportsAsyncCompute = IsAsyncComputeSupported(ShaderPlatform);
+	// Only enable async compute for large builders that are set to run in parallel.
+	bSupportsAsyncCompute = IsAsyncComputeSupported(ShaderPlatform) && EnumHasAnyFlags(InFlags, ERDGBuilderFlags::ParallelExecute);
 	bSupportsRenderPassMerge = IsRenderPassMergeEnabled(ShaderPlatform);
 
 	const bool bParallelExecuteFlag = EnumHasAnyFlags(InFlags, ERDGBuilderFlags::ParallelExecute);
