Added FRDGBuilder::SkipInitialAsyncComputeFence which allows GPUSkinCache async compute tasks to overlap with RDG work.

diff --git a/Engine/Source/Runtime/Engine/Private/GPUSkinCache.cpp b/Engine/Source/Runtime/Engine/Private/GPUSkinCache.cpp
--- a/Engine/Source/Runtime/Engine/Private/GPUSkinCache.cpp
+++ b/Engine/Source/Runtime/Engine/Private/GPUSkinCache.cpp
@@ -1407,6 +1407,9 @@ UE::Tasks::FTask FGPUSkinCache::Dispatch(FRDGBuilder& GraphBuilder, const UE::Ta
 
 	if (TaskData->Pipeline == ERHIPipeline::AsyncCompute)
 	{
+		// Tell the builder that we will manually sync async compute work back to graphics.
+		GraphBuilder.SkipInitialAsyncComputeFence();
+
 		GraphBuilder.AddPostExecuteCallback([TaskData, &RHICmdList = GraphBuilder.RHICmdList]
 		{
 			checkf(TaskData->bWaitPassAdded, TEXT("FGPUSkinCache::AddAsyncComputeWait was never called!"));
diff --git a/Engine/Source/Runtime/RenderCore/Private/RenderGraphBuilder.cpp b/Engine/Source/Runtime/RenderCore/Private/RenderGraphBuilder.cpp
--- a/Engine/Source/Runtime/RenderCore/Private/RenderGraphBuilder.cpp
+++ b/Engine/Source/Runtime/RenderCore/Private/RenderGraphBuilder.cpp
@@ -2054,7 +2054,7 @@ void FRDGBuilder::Execute()
 			}
 		};
 
-		if (GSupportsEfficientAsyncCompute)
+		if (bInitialAsyncComputeFence)
 		{
 			// Insert a manual fence from async compute to graphics to synchronize any all pipeline external access resources from the last run.
 			RHICmdList.Transition({}, ERHIPipeline::AsyncCompute, ERHIPipeline::Graphics);
diff --git a/Engine/Source/Runtime/RenderCore/Public/RenderGraphBuilder.h b/Engine/Source/Runtime/RenderCore/Public/RenderGraphBuilder.h
--- a/Engine/Source/Runtime/RenderCore/Public/RenderGraphBuilder.h
+++ b/Engine/Source/Runtime/RenderCore/Public/RenderGraphBuilder.h
@@ -319,6 +319,11 @@ public:
 	 */
 	RENDERCORE_API void SetFlushResourcesRHI();
 
+	/** Tells the RDG Builder that it is safe not to issue a fence from Async Compute -> Graphics at the start of the graph.
+	 *  By default RDG will issue the fence to keep new work from overlapping with work from prior builders.
+	 */
+	void SkipInitialAsyncComputeFence();
+
 	/** Queues a buffer upload operation prior to execution. The resource lifetime is extended and the data is uploaded prior to executing passes. */
 	void QueueBufferUpload(FRDGBufferRef Buffer, const void* InitialData, uint64 InitialDataSize, ERDGInitialDataFlags InitialDataFlags = ERDGInitialDataFlags::None);
 
@@ -490,6 +495,7 @@ private:
 	FRDGPass* ProloguePass = nullptr;
 	FRDGPass* EpiloguePass = nullptr;
 
+	bool bInitialAsyncComputeFence = GSupportsEfficientAsyncCompute;
 	bool bSupportsAsyncCompute = false;
 	bool bSupportsRenderPassMerge = false;
 
diff --git a/Engine/Source/Runtime/RenderCore/Public/RenderGraphBuilder.inl b/Engine/Source/Runtime/RenderCore/Public/RenderGraphBuilder.inl
--- a/Engine/Source/Runtime/RenderCore/Public/RenderGraphBuilder.inl
+++ b/Engine/Source/Runtime/RenderCore/Public/RenderGraphBuilder.inl
@@ -367,6 +367,11 @@ inline void FRDGBuilder::SetPassWorkload(FRDGPass* Pass, uint32 Workload)
 	Pass->Workload = Workload;
 }
 
+inline void FRDGBuilder::SkipInitialAsyncComputeFence()
+{
+	bInitialAsyncComputeFence = false;
+}
+
 inline void FRDGBuilder::QueueBufferUpload(FRDGBufferRef Buffer, const void* InitialData, uint64 InitialDataSize, ERDGInitialDataFlags InitialDataFlags)
 {
 	IF_RDG_ENABLE_DEBUG(UserValidation.ValidateUploadBuffer(Buffer, InitialData, InitialDataSize));
