Fix GNaniteFallbackMeshesPlatformMask being initialized with incorrect default value.

diff --git a/Engine/Source/Runtime/RenderCore/Private/RenderUtils.cpp b/Engine/Source/Runtime/RenderCore/Private/RenderUtils.cpp
--- a/Engine/Source/Runtime/RenderCore/Private/RenderUtils.cpp
+++ b/Engine/Source/Runtime/RenderCore/Private/RenderUtils.cpp
@@ -779,7 +779,7 @@ RENDERCORE_API void RenderUtilsInit()
 
 	GSimpleSkyDiffusePlatformMask.Init(false, EShaderPlatform::SP_NumPlatforms);
 	GVelocityEncodeDepthPlatformMask.Init(false, EShaderPlatform::SP_NumPlatforms);
-	GNaniteFallbackMeshesPlatformMask.Init(false, EShaderPlatform::SP_NumPlatforms);
+	GNaniteFallbackMeshesPlatformMask.Init(true, EShaderPlatform::SP_NumPlatforms);
 	GRayTracingPlatformMask.Init(false, EShaderPlatform::SP_NumPlatforms);
 	GRayTracingShadersPlatformMask.Init(false, EShaderPlatform::SP_NumPlatforms);
 
