[shaders]
- store shader symbols as a compressed buffer owned by the FShaderCodeResource after finalizing code
- use references to these compressed buffers everywhere symbol data is needed - single shader and shadermap DDC records, shader job cache entries, and shadermap editor-only data, to reduce both in-editor/cook memory footprint and DDC storage
- since symbols are stored as compressed, decompress to a temporary buffer only at the point of writing symbols out to disk (they are never needed uncompressed otherwise)
- add tracking of code and symbol buffers to the calculated job cache storage size; with this change effectively the job cache "owns" all the shader code and symbol buffers stored in memory (assuming shader map resources are cleaned up as packages are unloaded during cook), and so the job cache size limit also limits the amount of shader code and symbols will be kept around during a cook. note this means when entries are evicted from the job cache and later needed, they will be retrieved again from DDC, resulting in more network traffic but hopefully minimal redundant compilation (provided the DDC records are finalized in the local Zen instance before pruned from the job cache and needed again)

diff --git a/Engine/Shaders/Public/ShaderVersion.ush b/Engine/Shaders/Public/ShaderVersion.ush
--- a/Engine/Shaders/Public/ShaderVersion.ush
+++ b/Engine/Shaders/Public/ShaderVersion.ush
@@ -3,4 +3,4 @@
 // Each time the console command is executed it generates a new GUID. As this file is included
 // in Platform.ush (which should be included in any shader) it allows to invalidate the shader DDC.
 // 
-#pragma message("UESHADERMETADATA_VERSION 7612E18B-B799-4D68-B8D1-84A56CFAED95")
+#pragma message("UESHADERMETADATA_VERSION FF5BEA66-2646-43E2-902C-736FC7521050")
diff --git a/Engine/Source/Developer/ShaderCompilerCommon/Public/ShaderSymbolExport.h b/Engine/Source/Developer/ShaderCompilerCommon/Public/ShaderSymbolExport.h
--- a/Engine/Source/Developer/ShaderCompilerCommon/Public/ShaderSymbolExport.h
+++ b/Engine/Source/Developer/ShaderCompilerCommon/Public/ShaderSymbolExport.h
@@ -6,6 +6,7 @@
 
 #if WITH_ENGINE
 
+#include "Compression/CompressedBuffer.h"
 #include "CoreTypes.h"
 #include "Containers/ContainersFwd.h"
 #include "Containers/Map.h"
@@ -28,7 +29,11 @@ public:
 	*   Template type is the platform specific symbol data structure.
 	*/
 	template<typename TPlatformShaderSymbolData>
-	void NotifyShaderCompiled(const TConstArrayView<uint8>& PlatformSymbolData, const FString& DebugInfo = FString());
+	UE_DEPRECATED(5.6, "Use overload accepting an FCompressedBuffer for symbol data.")
+	void NotifyShaderCompiled(const TConstArrayView<uint8>& PlatformSymbolData, const FString& DebugInfo = FString()) {}
+
+	template<typename TPlatformShaderSymbolData>
+	void NotifyShaderCompiled(const FCompressedBuffer& PlatformSymbolDataCompressed, const FString& DebugInfo = FString());
 
 	/** Called at the end of a cook to free resources and finalize artifacts created during the cook. */
 	void NotifyShaderCompilersShutdown();
@@ -68,7 +73,7 @@ private:
 };
 
 template<typename TPlatformShaderSymbolData>
-inline void FShaderSymbolExport::NotifyShaderCompiled(const TConstArrayView<uint8>& PlatformSymbolData, const FString& DebugInfo)
+inline void FShaderSymbolExport::NotifyShaderCompiled(const FCompressedBuffer& PlatformSymbolDataCompressed, const FString& DebugInfo)
 {
 	FScopeLock Lock(&Cs);
 	if (!bInitialized)
@@ -83,7 +88,8 @@ inline void FShaderSymbolExport::NotifyShaderCompiled(const TConstArrayView<uint
 	{
 		// Deserialize the platform symbol data
 		TPlatformShaderSymbolData FullSymbolData;
-		FMemoryReaderView Ar(PlatformSymbolData);
+		FSharedBuffer PlatformSymbolData = PlatformSymbolDataCompressed.Decompress();
+		FMemoryReaderView Ar(PlatformSymbolData.GetView());
 		Ar << FullSymbolData;
 
 		for (const auto& SymbolData : FullSymbolData.GetAllSymbolData())
diff --git a/Engine/Source/Developer/TargetPlatform/Public/Interfaces/IShaderFormat.h b/Engine/Source/Developer/TargetPlatform/Public/Interfaces/IShaderFormat.h
--- a/Engine/Source/Developer/TargetPlatform/Public/Interfaces/IShaderFormat.h
+++ b/Engine/Source/Developer/TargetPlatform/Public/Interfaces/IShaderFormat.h
@@ -9,6 +9,7 @@
 
 class FSerializedShaderArchive;
 class FShaderPreprocessOutput;
+class FCompressedBuffer;
 class FSharedBuffer;
 struct FShaderCompilerEnvironment;
 struct FShaderCompilerInput;
@@ -117,11 +118,14 @@ public:
 	 */
 	virtual void ModifyShaderCompilerInput(FShaderCompilerInput& Input) const { }
 
+	UE_DEPRECATED(5.6, "Use overload accepting a const FCompressedBuffer& for symbol data")
+	virtual void NotifyShaderCompiled(const TConstArrayView<uint8>& SymbolData, FName Format, const FString& DebugInfo = FString()) const { }
+
 	/**
 	 * Called when a shader resource is cooked, so the shader format can perform platform-specific operations on the debug data.
-	 * Does nothing on platforms that make no use of the platform debug data.
+	 * Does nothing on platforms that make no use of the symbol data.
 	 */
-	virtual void NotifyShaderCompiled(const TConstArrayView<uint8>& PlatformDebugData, FName Format, const FString& DebugInfo = FString()) const { }
+	virtual void NotifyShaderCompiled(const FCompressedBuffer& SymbolData, FName Format, const FString& DebugInfo = FString()) const { }
 
 	/** Called at the end of a cook to free resources and finalize artifacts created during the cook. */
 	virtual void NotifyShaderCompilersShutdown(FName Format) const { }
diff --git a/Engine/Source/Developer/Windows/ShaderFormatD3D/Private/D3DShaderCompilerDXC.cpp b/Engine/Source/Developer/Windows/ShaderFormatD3D/Private/D3DShaderCompilerDXC.cpp
--- a/Engine/Source/Developer/Windows/ShaderFormatD3D/Private/D3DShaderCompilerDXC.cpp
+++ b/Engine/Source/Developer/Windows/ShaderFormatD3D/Private/D3DShaderCompilerDXC.cpp
@@ -931,7 +931,7 @@ bool CompileAndProcessD3DShaderDXC(
 		// as we use embedded PDBs. Once we are able to use external PDBs, the PDB contents
 		// can be exported too.
 		
-		FMemoryWriter Ar(Output.PlatformDebugData);
+		FMemoryWriter Ar(Output.ShaderCode.GetSymbolWriteAccess());
 		Ar << DebugData;
 	}
 
diff --git a/Engine/Source/Developer/Windows/ShaderFormatD3D/Private/ShaderFormatD3D.cpp b/Engine/Source/Developer/Windows/ShaderFormatD3D/Private/ShaderFormatD3D.cpp
--- a/Engine/Source/Developer/Windows/ShaderFormatD3D/Private/ShaderFormatD3D.cpp
+++ b/Engine/Source/Developer/Windows/ShaderFormatD3D/Private/ShaderFormatD3D.cpp
@@ -168,11 +168,11 @@ public:
 	}
 
 #if WITH_ENGINE
-	virtual void NotifyShaderCompiled(const TConstArrayView<uint8>& PlatformDebugData, FName Format, const FString& DebugInfo) const override
+	virtual void NotifyShaderCompiled(const FCompressedBuffer& SymbolData, FName Format, const FString& DebugInfo) const override
 	{
 		if (Format == NAME_PCD3D_SM6)
 		{
-			ShaderSymbolExportSM6.NotifyShaderCompiled<FD3DSM6ShaderDebugData>(PlatformDebugData, DebugInfo);
+			ShaderSymbolExportSM6.NotifyShaderCompiled<FD3DSM6ShaderDebugData>(SymbolData, DebugInfo);
 		}
 	}
 	virtual void NotifyShaderCompilersShutdown(FName Format) const override
diff --git a/Engine/Source/Runtime/Engine/Private/ShaderCompiler/ShaderCompilerJobCache.cpp b/Engine/Source/Runtime/Engine/Private/ShaderCompiler/ShaderCompilerJobCache.cpp
--- a/Engine/Source/Runtime/Engine/Private/ShaderCompiler/ShaderCompilerJobCache.cpp
+++ b/Engine/Source/Runtime/Engine/Private/ShaderCompiler/ShaderCompilerJobCache.cpp
@@ -368,6 +368,9 @@ public:
 	/** Separate blobs for shader code */
 	TArray<FCompositeBuffer> JobCode;
 
+	/** Separate blobs for shader symbols */
+	TArray<FCompressedBuffer> JobSymbols;
+
 	/** Path to where the cached debug info is stored. */
 	FString CachedDebugInfoPath;
 
@@ -403,7 +406,16 @@ public:
 
 	uint64 GetAllocatedSize() const
 	{
-		return static_cast<uint64>(JobOutput.GetSize() + sizeof(*this));
+		uint64 Size = JobOutput.GetSize() + sizeof(*this);
+		for (const FCompositeBuffer& CodeBuf : JobCode)
+		{
+			Size += CodeBuf.GetSize();
+		}
+		for (const FCompressedBuffer& SymbolBuf : JobSymbols)
+		{
+			Size += SymbolBuf.GetCompressedSize();
+		}
+		return Size;
 	}
 };
 
@@ -1632,7 +1644,7 @@ void FShaderJobCache::AddToCacheAndProcessPending(FShaderCommonCompileJob* Finis
 
 		check(SaveContext.HasData());
 		// Construct a single load context pointing to the data in the save context used above
-		FShaderCacheLoadContext LoadContext(SaveContext.ShaderObjectData, SaveContext.ShaderCode);
+		FShaderCacheLoadContext LoadContext(SaveContext.ShaderObjectData, SaveContext.ShaderCode, SaveContext.ShaderSymbols);
 		for (FShaderCommonCompileJob* DuplicateJob : FinishedDuplicateJobs)
 		{
 			// reuse the same load context for each duplicate job to avoid reallocating anything
@@ -1763,7 +1775,7 @@ FShaderCacheLoadContext FShaderJobCache::FindOrAdd(const FShaderCompilerInputHas
 		// update the output hit count
 		(*CannedOutput)->NumHits++;
 
-		return FShaderCacheLoadContext((*CannedOutput)->JobOutput, (*CannedOutput)->JobCode);
+		return FShaderCacheLoadContext((*CannedOutput)->JobOutput, (*CannedOutput)->JobCode, (*CannedOutput)->JobSymbols);
 	}
 #if WITH_EDITOR
 	else
@@ -1897,7 +1909,7 @@ FShaderCacheLoadContext FShaderJobCache::FindOrAdd(const FShaderCompilerInputHas
 							check(LoadContext.HasData()); // sanity check that the load context was populated properly
 							StoredOutput = new FStoredOutput();
 							StoredOutput->JobOutput = LoadContext.ShaderObjectData;
-							LoadContext.MoveCode(StoredOutput->JobCode);
+							LoadContext.MoveCode(StoredOutput->JobCode, StoredOutput->JobSymbols);
 							Outputs.Add(OutputHash, StoredOutput);
 							CurrentlyAllocatedMemory += StoredOutput->GetAllocatedSize();
 						}
@@ -1944,7 +1956,7 @@ FShaderCacheLoadContext FShaderJobCache::FindOrAdd(const FShaderCompilerInputHas
 							JobLock.WriteUnlock();
 
 							// Re-initialize load context pointing to the stored code array in the cache and reuse for each duplicate job needing population
-							LoadContext.Reset(StoredOutput->JobOutput, StoredOutput->JobCode);
+							LoadContext.Reset(StoredOutput->JobOutput, StoredOutput->JobCode, StoredOutput->JobSymbols);
 
 							// Call ProcessFinishedJob on main job and duplicates
 							for (FShaderCommonCompileJob* FinishedJob : FinishedJobs)
@@ -2018,7 +2030,7 @@ FShaderCacheLoadContext FShaderJobCache::FindOrAdd(const FShaderCompilerInputHas
 
 				if (StoredOutput)
 				{
-					return FShaderCacheLoadContext(StoredOutput->JobOutput, StoredOutput->JobCode);
+					return FShaderCacheLoadContext(StoredOutput->JobOutput, StoredOutput->JobCode, StoredOutput->JobSymbols);
 				}
 			}
 		}
@@ -2193,7 +2205,7 @@ void FShaderJobCache::AddJobOutput(FShaderJobData& JobData, const FShaderCommonC
 			FStoredOutput* NewStoredOutput = new FStoredOutput();
 			NewStoredOutput->NumHits = InitialHitCount;
 			NewStoredOutput->JobOutput = SaveContext.ShaderObjectData;
-			SaveContext.MoveCode(NewStoredOutput->JobCode);
+			SaveContext.MoveCode(NewStoredOutput->JobCode, NewStoredOutput->JobSymbols);
 			NewStoredOutput->CachedDebugInfoPath = InputDebugInfoPath;
 			NewStoredOutput->AddRef();
 			Outputs.Add(OutputHash, NewStoredOutput);
diff --git a/Engine/Source/Runtime/RenderCore/Internal/ShaderSerialization.h b/Engine/Source/Runtime/RenderCore/Internal/ShaderSerialization.h
--- a/Engine/Source/Runtime/RenderCore/Internal/ShaderSerialization.h
+++ b/Engine/Source/Runtime/RenderCore/Internal/ShaderSerialization.h
@@ -2,6 +2,7 @@
 
 #pragma once
 
+#include "Compression/CompressedBuffer.h"
 #include "Memory/CompositeBuffer.h"
 #include "Memory/SharedBuffer.h"
 #include "Templates/Function.h"
@@ -49,9 +50,6 @@ struct FShaderSerializeContext
 	/* FName of the asset which triggered the serialization, this is used only for diagnostic messages */
 	FName SerializingAsset = NAME_None;
 
-	/* Optional function which can be used to reserve space for a number of code entries as a minor optimization. */
-	TFunction<void(int32)> ReserveCodeFunc{};
-
 	/* Helper function to retrieve an FArchive reference; exists for convenience and validation purposes. */
 	FArchive& GetMainArchive()
 	{
@@ -69,6 +67,9 @@ struct FShaderCacheSerializeContext : public FShaderSerializeContext
 	FShaderCacheSerializeContext() : FShaderSerializeContext() 
 	{
 		ShaderCode = OwnedShaderCode;
+		ShaderSymbols = OwnedShaderSymbols;
+
+		checkf(ShaderCode.Num() == ShaderSymbols.Num(), TEXT("It is required to serialize a (possibly empty, but non-null) symbols buffer for every code buffer."));
 	}
 
 	virtual ~FShaderCacheSerializeContext() = default;
@@ -81,11 +82,21 @@ struct FShaderCacheSerializeContext : public FShaderSerializeContext
 	 */
 	TArrayView<FCompositeBuffer> ShaderCode;
 
-	/* Array of buffers actually owned by this context object; it is valid for this to be empty in cases
+	/* View on array of compressed buffers which store the symbols for the object, one per shader/stage.
+	 * Note that this may or may not point to the OwnedShaderSymbols array below, depending on usage.
+	 */
+	TArrayView<FCompressedBuffer> ShaderSymbols;
+
+	/* Array of code buffers actually owned by this context object; it is valid for this to be empty in cases
 	 * where the array of ShaderCode buffers is stored externally.
 	 */
 	TArray<FCompositeBuffer> OwnedShaderCode;
 
+	/* Array of symbol buffers actually owned by this context object; it is valid for this to be empty in cases
+	 * where the array of ShaderSymbols buffers is stored externally.
+	 */
+	TArray<FCompressedBuffer> OwnedShaderSymbols;
+
 	/* Subclasses of this type must implement custom code serialization function */
 	virtual bool EnableCustomCodeSerialize() override { return true; }
 
@@ -108,10 +119,14 @@ struct FShaderCacheSerializeContext : public FShaderSerializeContext
 	}
 
 	/* Populates the given code array (transfering ownership) and resets the internal view to point to the new owning array's data */
-	void MoveCode(TArray<FCompositeBuffer>& TargetCode)
+	void MoveCode(TArray<FCompositeBuffer>& TargetCode, TArray<FCompressedBuffer>& TargetSymbols)
 	{
 		TargetCode = MoveTemp(OwnedShaderCode);
 		ShaderCode = TargetCode;
+		TargetSymbols = MoveTemp(OwnedShaderSymbols);
+		ShaderSymbols = TargetSymbols;
+
+		checkf(ShaderCode.Num() == ShaderSymbols.Num(), TEXT("It is required to serialize a (possibly empty, but non-null) symbols buffer for every code buffer."));
 	}
 
 	/* Returns true if there is valid serialized data referenced by this context. */
@@ -161,12 +176,12 @@ struct FShaderCacheLoadContext : public FShaderCacheSerializeContext
 	FShaderCacheLoadContext() = default;
 
 	/* Constructor which references buffers (and an array of code buffers) owned elsewhere, does not allocate the OwnedShaderCode array */
-	RENDERCORE_API FShaderCacheLoadContext(FSharedBuffer ShaderObjectData, TArrayView<FCompositeBuffer> CodeBuffers);
+	RENDERCORE_API FShaderCacheLoadContext(FSharedBuffer ShaderObjectData, TArrayView<FCompositeBuffer> CodeBuffers, TArrayView<FCompressedBuffer> SymbolBuffers);
 
 	virtual ~FShaderCacheLoadContext() = default;
 
 	/* Resets internal state to point to the given buffers and recreates the owned memoryreader */
-	RENDERCORE_API void Reset(FSharedBuffer ShaderObjectData, TArrayView<FCompositeBuffer> CodeBuffers);
+	RENDERCORE_API void Reset(FSharedBuffer ShaderObjectData, TArrayView<FCompositeBuffer> CodeBuffers, TArrayView<FCompressedBuffer> SymbolBuffers);
 
 	/* Call to reset reader to start position so the same load context can be used to populate multiple objects. */
 	RENDERCORE_API void Reuse();
diff --git a/Engine/Source/Runtime/RenderCore/Private/ShaderCore.cpp b/Engine/Source/Runtime/RenderCore/Private/ShaderCore.cpp
--- a/Engine/Source/Runtime/RenderCore/Private/ShaderCore.cpp
+++ b/Engine/Source/Runtime/RenderCore/Private/ShaderCore.cpp
@@ -1704,11 +1704,13 @@ public:
 		CombinedSource.Append(PrimaryCode);
 		CombinedSource.Append(SecondaryCode);
 
-		// Replace Output shader code with the combined result
+		// this could be cleaner, but we currently require backends with secondary outputs to manually merge their symbols during the compilation process
+		// as such we need to save off the symbol buffer before we reset the shader code struct below, then re-set the buffer when finalizing the new combined code
+		FCompressedBuffer SymbolsBuffer = Job.Output.GetFinalizedCodeResource().GetSymbolsBuffer();
+		// Replace Output shader code with the combined result; need to reset first since the code struct is in the finalized state and write access is disallowed
 		Job.Output.ShaderCode = {};
-		TArray<uint8>& FinalShaderCode = Job.Output.ShaderCode.GetWriteAccess();
-		FinalShaderCode = MoveTemp(CombinedSource);
-		Job.Output.ShaderCode.FinalizeShaderCode();
+		Job.Output.ShaderCode.GetWriteAccess() = MoveTemp(CombinedSource);
+		Job.Output.ShaderCode.FinalizeShaderCode(SymbolsBuffer);
 	}
 
 	static void InvokeCompile(const IShaderFormat* Compiler, FShaderCompileJob& Job, const FString& WorkingDirectory, FString& OutExceptionMsg, FString& OutExceptionCallstack)
@@ -3967,7 +3969,10 @@ FArchive& operator<<(FArchive& Ar, FSharedBuffer& Buffer)
 	}
 	else if (Ar.IsSaving())
 	{
-		Ar.Serialize(const_cast<void*>(Buffer.GetData()), Len);
+		if (Len > 0)
+		{
+			Ar.Serialize(const_cast<void*>(Buffer.GetData()), Len);
+		}
 	}
 
 	return Ar;
@@ -3977,6 +3982,7 @@ FArchive& operator<<(FArchive& Ar, FShaderCodeResource& Resource)
 {
 	Ar << Resource.Header;
 	Ar << Resource.Code;
+	Ar << Resource.Symbols;
 	return Ar;
 }
 
@@ -3994,6 +4000,10 @@ FArchive& operator<<(FArchive& Ar, FShaderCode& Output)
 	// Note: this serialize is used to pass between UE and the shader compile worker, recompile both when modifying
 	FSharedBuffer& CodeBuffer = Output.ShaderCodeResource.Code;
 	Ar << CodeBuffer;
+
+	FCompressedBuffer& SymbolsBuffer = Output.ShaderCodeResource.Symbols;
+	Ar << SymbolsBuffer;
+
 	Ar << Output.UncompressedSize;
 	{
 		FString CompressionFormatString(Output.CompressionFormat.ToString());
@@ -4240,7 +4250,9 @@ void FShaderCompileJob::OnComplete()
 	Output.PreprocessTime = PreprocessOutput.ElapsedTime;
 
 	const IShaderFormat* ShaderFormat = nullptr;
-	if (!Output.PlatformDebugData.IsEmpty() || Input.DumpDebugInfoEnabled())
+	FCompressedBuffer SymbolsBuffer = Output.GetFinalizedCodeResource().GetSymbolsBuffer();
+	const bool bHasSymbols = SymbolsBuffer.GetRawSize() > 0;
+	if (bHasSymbols || Input.DumpDebugInfoEnabled())
 	{
 		// don't bother looking up the shader format unless needed, this can have non-trivial cost due to thread safety
 		// and the possibility of available target platforms needing to be refreshed
@@ -4275,10 +4287,10 @@ void FShaderCompileJob::OnComplete()
 		}
 	}
 
-	if (!Output.PlatformDebugData.IsEmpty())
+	if (bHasSymbols)
 	{
 		check(ShaderFormat);
-		ShaderFormat->NotifyShaderCompiled(Output.PlatformDebugData, Input.ShaderFormat, Input.GenerateDebugInfo());
+		ShaderFormat->NotifyShaderCompiled(SymbolsBuffer, Input.ShaderFormat, Input.GenerateDebugInfo());
 	}
 }
 
@@ -4453,11 +4465,8 @@ void FShaderPipelineCompileJob::SerializeOutput(FShaderCacheSerializeContext& Ct
 {
 	bool bAllStagesSucceeded = true;
 
-	if (Ctx.ReserveCodeFunc)
-	{
-		Ctx.ReserveCodeFunc(StageJobs.Num());
-	}
-
+	Ctx.ReserveCode(StageJobs.Num());
+	
 	for (int32 Index = 0, Num = StageJobs.Num(); Index < Num; ++Index)
 	{
 		StageJobs[Index]->SerializeOutput(Ctx, Index);
diff --git a/Engine/Source/Runtime/RenderCore/Private/ShaderResource.cpp b/Engine/Source/Runtime/RenderCore/Private/ShaderResource.cpp
--- a/Engine/Source/Runtime/RenderCore/Private/ShaderResource.cpp
+++ b/Engine/Source/Runtime/RenderCore/Private/ShaderResource.cpp
@@ -6,6 +6,7 @@
 
 #include "Shader.h"
 #include "Compression/OodleDataCompression.h"
+#include "Containers/ArrayView.h"
 #include "DataDrivenShaderPlatformInfo.h"
 #include "Misc/Compression.h"
 #include "Misc/CoreMisc.h"
@@ -13,6 +14,7 @@
 #include "Interfaces/ITargetPlatform.h"
 #include "Interfaces/ITargetPlatformManagerModule.h"
 #include "Interfaces/IShaderFormat.h"
+#include "IO/IoHash.h"
 #include "Misc/MemStack.h"
 #include "Misc/ScopeLock.h"
 #include "RenderingThread.h"
@@ -256,9 +258,15 @@ int32 FShaderMapResourceCode::FindShaderIndex(const FSHAHash& InHash) const
 }
 
 #if WITH_EDITORONLY_DATA
-uint64 FShaderMapResourceCode::FShaderEditorOnlyDataEntry::GetDebugDataHash(TConstArrayView<uint8> InPlatformDebugData)
+void FShaderMapResourceCode::FShaderEditorOnlyDataEntry::ConditionalSetSymbolBuffer(FCompressedBuffer InSymbols)
 {
-	return CityHash64(reinterpret_cast<const char*>(InPlatformDebugData.GetData()), InPlatformDebugData.Num());
+	// Intentional truncation of hash; 20 bytes is overkill for this purpose so we just take the first 8
+	const uint64 InSymbolHash = reinterpret_cast<uint64*>(InSymbols.GetRawHash().GetBytes())[0];
+	if (SymbolBuffer.IsNull() || InSymbolHash < SymbolHash)
+	{
+		SymbolBuffer = InSymbols;
+		SymbolHash = InSymbolHash;
+	}
 }
 #endif
 
@@ -268,45 +276,43 @@ void FShaderMapResourceCode::AddShaderCompilerOutput(const FShaderCompilerOutput
 
 	const FSHAHash& InHash = Output.OutputHash;
 	const int32 Index = Algo::LowerBound(ShaderHashes, InHash);
+	const FShaderCodeResource& CodeResource = Output.GetFinalizedCodeResource();
 	if (Index >= ShaderHashes.Num() || ShaderHashes[Index] != InHash)
 	{
 		ShaderHashes.Insert(InHash, Index);
 
 #if WITH_EDITORONLY_DATA
 		// Output.Errors contains warnings in the case any exist (no errors since if there were the job would have failed)
-		AddEditorOnlyData(Index, DebugName, Output.PlatformDebugData, Output.Errors, Output.ShaderStatistics, DebugInfo);
+		AddEditorOnlyData(Index, DebugName, CodeResource.GetSymbolsBuffer(), Output.Errors, Output.ShaderStatistics, DebugInfo);
 #endif
 
-		ShaderCodeResources.Insert(Output.GetFinalizedCodeResource(), Index);
+		ShaderCodeResources.Insert(CodeResource, Index);
 	}
 #if WITH_EDITORONLY_DATA
 	else
 	{
 		// Output.Errors contains warnings in the case any exist (no errors since if there were the job would have failed)
 		// We append the warnings and deduplicate other data like DebugInfo for any additional jobs which resulted in the
-		// same bytecode for the sake of determinism in the results saved to DDC. 
-		UpdateEditorOnlyData(Index, DebugName, Output.PlatformDebugData, Output.Errors, DebugInfo);
+		// same bytecode for the sake of determinism in the results saved to DDC.
+		UpdateEditorOnlyData(Index, DebugName, CodeResource.GetSymbolsBuffer(), Output.Errors, DebugInfo);
 		ValidateShaderStatisticsEditorOnlyData(Index, Output.ShaderStatistics);
 	}
 #endif
 }
 
 #if WITH_EDITORONLY_DATA
-void FShaderMapResourceCode::AddEditorOnlyData(int32 Index, const FString& DebugName, TConstArrayView<uint8> InPlatformDebugData, TConstArrayView<FShaderCompilerError> InCompilerWarnings, const TArray<FGenericShaderStat>& ShaderStatistics, const FString& DebugInfo)
+void FShaderMapResourceCode::AddEditorOnlyData(int32 Index, const FString& DebugName, FCompressedBuffer InSymbols, TConstArrayView<FShaderCompilerError> InCompilerWarnings, const TArray<FGenericShaderStat>& ShaderStatistics, const FString& DebugInfo)
 {
 	FShaderEditorOnlyDataEntry& Entry = ShaderEditorOnlyDataEntries.InsertDefaulted_GetRef(Index);
 
-	Entry.PlatformDebugData = InPlatformDebugData;
-	Entry.PlatformDebugDataHash = Entry.GetDebugDataHash(InPlatformDebugData);
-
 	// This should be a newly created shader entry.
 	check(Entry.ShaderStatistics.Num() == 0);
 	Entry.ShaderStatistics = ShaderStatistics;
 
-	UpdateEditorOnlyData(Index, DebugName, InPlatformDebugData, InCompilerWarnings, DebugInfo);
+	UpdateEditorOnlyData(Index, DebugName, InSymbols, InCompilerWarnings, DebugInfo);
 }
 
-void FShaderMapResourceCode::UpdateEditorOnlyData(int32 Index, const FString& DebugName, TConstArrayView<uint8> InPlatformDebugData, TConstArrayView<FShaderCompilerError> InCompilerWarnings, const FString& DebugInfo)
+void FShaderMapResourceCode::UpdateEditorOnlyData(int32 Index, const FString& DebugName, FCompressedBuffer InSymbols, TConstArrayView<FShaderCompilerError> InCompilerWarnings, const FString& DebugInfo)
 {
 	FShaderEditorOnlyDataEntry& Entry = ShaderEditorOnlyDataEntries[Index];
 
@@ -316,14 +322,7 @@ void FShaderMapResourceCode::UpdateEditorOnlyData(int32 Index, const FString& De
 		Entry.DebugInfo = DebugInfo;
 	}
 
-	// As above, keep only a single copy of the platform debug data, it doesn't matter which, but use the one with the smallest
-	// hash to ensure DDC data is deterministic
-	uint64 InPlatformDebugDataHash = Entry.GetDebugDataHash(InPlatformDebugData);
-	if (InPlatformDebugDataHash < Entry.PlatformDebugDataHash)
-	{
-		Entry.PlatformDebugDataHash = InPlatformDebugDataHash;
-		Entry.PlatformDebugData = InPlatformDebugData;
-	}
+	Entry.ConditionalSetSymbolBuffer(InSymbols);
 
 	for (const FShaderCompilerError& Warning : InCompilerWarnings)
 	{
@@ -411,11 +410,8 @@ void FShaderMapResourceCode::Serialize(FShaderSerializeContext& Ctx)
 		{
 			ShaderCodeResources.SetNum(ShaderHashes.Num());
 		}
-		
-		if (Ctx.ReserveCodeFunc)
-		{
-			Ctx.ReserveCodeFunc(ShaderCodeResources.Num());
-		}
+	
+		Ctx.ReserveCode(ShaderCodeResources.Num());
 
 		for (int32 CodeIndex = 0; CodeIndex < ShaderCodeResources.Num(); ++CodeIndex)
 		{
@@ -429,6 +425,16 @@ void FShaderMapResourceCode::Serialize(FShaderSerializeContext& Ctx)
 	{
 		Ar << ShaderEditorOnlyDataEntries;
 	}
+	if (Ar.IsLoading())
+	{
+		// only need to set SymbolBuffer when loading, saving is handled automatically by the serialization of ShaderCodeResources
+		// (this just sets the FSharedBuffer reference in a second place to avoid needing to fix up a bunch of code which retrieves
+		// symbols from the editoronly data)
+		for (int32 CodeIndex = 0; CodeIndex < ShaderCodeResources.Num(); ++CodeIndex)
+		{
+			ShaderEditorOnlyDataEntries[CodeIndex].ConditionalSetSymbolBuffer(ShaderCodeResources[CodeIndex].GetSymbolsBuffer());
+		}
+	}
 #endif // WITH_EDITORONLY_DATA
 	ApplyResourceStats(*this);
 
@@ -452,7 +458,7 @@ void FShaderMapResourceCode::NotifyShadersCompiled(FName FormatName)
 		{
 			for (const FShaderEditorOnlyDataEntry& Entry : ShaderEditorOnlyDataEntries)
 			{
-				ShaderFormat->NotifyShaderCompiled(Entry.PlatformDebugData, FormatName, Entry.DebugInfo);
+				ShaderFormat->NotifyShaderCompiled(Entry.SymbolBuffer, FormatName, Entry.DebugInfo);
 			}
 		}
 	}
diff --git a/Engine/Source/Runtime/RenderCore/Private/ShaderSerialization.cpp b/Engine/Source/Runtime/RenderCore/Private/ShaderSerialization.cpp
--- a/Engine/Source/Runtime/RenderCore/Private/ShaderSerialization.cpp
+++ b/Engine/Source/Runtime/RenderCore/Private/ShaderSerialization.cpp
@@ -18,7 +18,9 @@ FShaderCacheSaveContext::FShaderCacheSaveContext()
 #if WITH_EDITOR
 const UE::DerivedData::FValueId ShaderObjectDataValue = UE::DerivedData::FValueId::FromName(TEXT("ShaderObjectData"));
 const UE::DerivedData::FValueId ShaderCodeDataValue = UE::DerivedData::FValueId::FromName(TEXT("ShaderCodeData"));
+const UE::DerivedData::FValueId ShaderSymbolsDataValue = UE::DerivedData::FValueId::FromName(TEXT("ShaderSymbolsData"));
 const FAnsiStringView CodeCountMetaField = ANSITEXTVIEW("CodeCount");
+const FAnsiStringView HasSymbolsMetaField = ANSITEXTVIEW("bHasSymbols");
 #endif
 
 void FShaderCacheSaveContext::SerializeCode(FShaderCodeResource& Resource, int32 Index)
@@ -27,17 +29,24 @@ void FShaderCacheSaveContext::SerializeCode(FShaderCodeResource& Resource, int32
 	// reset the array view any time an entry is added; we do this instead of calling Resize in the reserve delegate
 	// and setting it there since not all code paths (i.e. single job cache records) call reserve
 	ShaderCode = OwnedShaderCode;
+
+	OwnedShaderSymbols.Add(Resource.GetSymbolsBuffer());
+	ShaderSymbols = OwnedShaderSymbols;
+
+	checkf(ShaderCode.Num() == ShaderSymbols.Num(), TEXT("It is required to serialize a (possibly empty, but non-null) symbols buffer for every code buffer."));
 }
 
 void FShaderCacheSaveContext::ReserveCode(int32 Count)
 {
 	OwnedShaderCode.Reserve(Count);
+	OwnedShaderSymbols.Reserve(Count);
 }
 
 void FShaderCacheSaveContext::Reset()
 {
 	ShaderObjectData.Reset();
 	OwnedShaderCode.Reset();
+	OwnedShaderSymbols.Reset();
 	Writer = MakeUnique<FMemoryWriter64>(ShaderObjectRawData);
 	Ar = Writer.Get();
 }
@@ -61,14 +70,26 @@ UE::DerivedData::FCacheRecord FShaderCacheSaveContext::BuildCacheRecord(const UE
 	// Code buffers are already compressed, don't waste cycles attempting (and failing) to recompress them
 	const ECompressedBufferCompressor CodeComp = ECompressedBufferCompressor::NotSet;
 	const ECompressedBufferCompressionLevel CodeCompLevel = ECompressedBufferCompressionLevel::None;
+
+	// Use a meta field to indicate whether we are adding symbol values to this record; we do so to prevent per-value overhead of symbol buffers if they are empty
+	// (each value in a cache record has a 64 byte overhead, which can be significant when we're pushing millions of individual shader cache records)
+	const bool bHasSymbols = ShaderSymbols[CodeIndex].GetRawSize() > 0;
 	for (FCompositeBuffer& CodeBuf : ShaderCode)
 	{
-		RecordBuilder.AddValue(ShaderCodeDataValue.MakeIndexed(CodeIndex++), UE::DerivedData::FValue(FCompressedBuffer::Compress(CodeBuf, CodeComp, CodeCompLevel)));
+		RecordBuilder.AddValue(ShaderCodeDataValue.MakeIndexed(CodeIndex), UE::DerivedData::FValue(FCompressedBuffer::Compress(CodeBuf, CodeComp, CodeCompLevel)));
+		// ensure that if one symbol buffer is non-empty, they all are
+		checkf(bHasSymbols == (ShaderSymbols[CodeIndex].GetRawSize() > 0), TEXT("It is expected that either all or none of the shader code entries have associated symbols."));
+		if (bHasSymbols)
+		{
+			RecordBuilder.AddValue(ShaderSymbolsDataValue.MakeIndexed(CodeIndex), UE::DerivedData::FValue(ShaderSymbols[CodeIndex]));
+		}
+		CodeIndex++;
 	}
 
 	TCbWriter<16> MetaWriter;
 	MetaWriter.BeginObject();
 	MetaWriter.AddInteger(CodeCountMetaField, ShaderCode.Num());
+	MetaWriter.AddBool(HasSymbolsMetaField, bHasSymbols);
 	MetaWriter.EndObject();
 
 	RecordBuilder.SetMeta(MetaWriter.Save().AsObject());
@@ -78,15 +99,18 @@ UE::DerivedData::FCacheRecord FShaderCacheSaveContext::BuildCacheRecord(const UE
 #endif
 
 
-FShaderCacheLoadContext::FShaderCacheLoadContext(FSharedBuffer InShaderObjectData, TArrayView<FCompositeBuffer> InCodeBuffers)
+FShaderCacheLoadContext::FShaderCacheLoadContext(FSharedBuffer InShaderObjectData, TArrayView<FCompositeBuffer> InCodeBuffers, TArrayView<FCompressedBuffer> InSymbolBuffers)
 {
-	Reset(InShaderObjectData, InCodeBuffers);
+	Reset(InShaderObjectData, InCodeBuffers, InSymbolBuffers);
 }
 
-void FShaderCacheLoadContext::Reset(FSharedBuffer InShaderObjectData, TArrayView<FCompositeBuffer> InCodeBuffers)
+void FShaderCacheLoadContext::Reset(FSharedBuffer InShaderObjectData, TArrayView<FCompositeBuffer> InCodeBuffers, TArrayView<FCompressedBuffer> InSymbolBuffers)
 {
 	ShaderObjectData = InShaderObjectData;
 	ShaderCode = InCodeBuffers;
+	ShaderSymbols = InSymbolBuffers;
+
+	checkf(ShaderCode.Num() == ShaderSymbols.Num(), TEXT("It is required to serialize a (possibly empty, but non-null) symbols buffer for every code buffer."));
 
 	Reader = MakeUnique<FMemoryReaderView>(ShaderObjectData);
 	Ar = Reader.Get();
@@ -94,7 +118,7 @@ void FShaderCacheLoadContext::Reset(FSharedBuffer InShaderObjectData, TArrayView
 
 void FShaderCacheLoadContext::SerializeCode(FShaderCodeResource& Resource, int32 Index)
 {
-	Resource.PopulateFromComposite(ShaderCode[Index]);
+	Resource.PopulateFromComposite(ShaderCode[Index], ShaderSymbols[Index]);
 }
 
 void FShaderCacheLoadContext::Reuse()
@@ -114,12 +138,30 @@ void FShaderCacheLoadContext::ReadFromRecord(const UE::DerivedData::FCacheRecord
 	
 	int32 CodeCount = Record.GetMeta()[CodeCountMetaField].AsInt32();
 	OwnedShaderCode.Reserve(CodeCount);
+	OwnedShaderSymbols.Reserve(CodeCount);
+
+	const bool bHasSymbols = Record.GetMeta()[HasSymbolsMetaField].AsBool();
+
 	for (int32 CodeIndex = 0; CodeIndex < CodeCount; ++CodeIndex)
 	{
 		FSharedBuffer CombinedBuffer = Record.GetValue(ShaderCodeDataValue.MakeIndexed(CodeIndex)).GetData().Decompress();
 		OwnedShaderCode.Add(FShaderCodeResource::Unpack(CombinedBuffer));
+		if (bHasSymbols)
+		{
+			UE::DerivedData::FValueWithId Value = Record.GetValue(ShaderSymbolsDataValue.MakeIndexed(CodeIndex));
+			FCompressedBuffer SymbolsBuffer = Value.GetData();
+			OwnedShaderSymbols.Add(SymbolsBuffer);
+		}
+		else
+		{
+			// We need to set an empty symbol buffer if symbols were not saved to the cache, instead of a null one
+			// This is due to FCompressedBuffer serialization not supporting null buffers
+			static FCompressedBuffer Empty = FCompressedBuffer::Compress(FUniqueBuffer::Alloc(0).MoveToShared());
+			OwnedShaderSymbols.Add(Empty);
+		}
 	}
 	ShaderCode = OwnedShaderCode;
+	ShaderSymbols = OwnedShaderSymbols;
 }
 #endif
 
diff --git a/Engine/Source/Runtime/RenderCore/Public/Shader.h b/Engine/Source/Runtime/RenderCore/Public/Shader.h
--- a/Engine/Source/Runtime/RenderCore/Public/Shader.h
+++ b/Engine/Source/Runtime/RenderCore/Public/Shader.h
@@ -530,11 +530,16 @@ public:
 #if WITH_EDITORONLY_DATA
 	struct FShaderEditorOnlyDataEntry
 	{
-		static uint64 GetDebugDataHash(TConstArrayView<uint8> PlatformDebugData);
-	
-		uint64 PlatformDebugDataHash = 0u; // Used only to ensure the copy of debug data saved is stable for determinism's sake; does not need to be serialized.
+		UE_DEPRECATED(5.6, "Use SymbolBuffer")
 		TArray<uint8> PlatformDebugData;
 
+		// Sets the symbols for this entry to the given buffer if either unset or the hash of the new buffer is "less" than the hash of the stored buffer.
+		// This is done purely for determinism's sake (we always want to store the same copy of symbols in the editoronly data for any given shader bytecode).
+		void ConditionalSetSymbolBuffer(FCompressedBuffer InSymbols);
+	
+		uint64 SymbolHash = 0u; // Used only to ensure the copy of symbols saved is stable for determinism's sake; does not need to be serialized.
+		FCompressedBuffer SymbolBuffer;
+
 		/** This field contains a debug string stored in a ShaderSymbols.info file when using r.Shaders.SymbolInfo=1
 		*   Used to facilitate reverse lookup of a shader from a platform shader hash when no full shader symbol 
 		*   information is available */
@@ -551,13 +556,18 @@ public:
 
 		friend FArchive& operator<<(FArchive& Ar, FShaderEditorOnlyDataEntry& Entry)
 		{
-			Ar << Entry.PlatformDebugData << Entry.DebugInfo << Entry.CompilerWarnings << Entry.ShaderStatistics;
-			if (Ar.IsLoading())
-			{
-				Entry.PlatformDebugDataHash = GetDebugDataHash(Entry.PlatformDebugData);
-			}
-			return Ar;
+			return Ar << Entry.DebugInfo << Entry.CompilerWarnings << Entry.ShaderStatistics;
 		}
+
+		PRAGMA_DISABLE_DEPRECATION_WARNINGS
+		// Explicitly-defaulted copy/move ctors & assignment operators are needed temporarily due to deprecation
+		// of PlatformDebugData field. These can be removed once the deprecation window for said field ends.
+		FShaderEditorOnlyDataEntry() = default;
+		FShaderEditorOnlyDataEntry(FShaderEditorOnlyDataEntry&&) = default;
+		FShaderEditorOnlyDataEntry(const FShaderEditorOnlyDataEntry&) = default;
+		FShaderEditorOnlyDataEntry& operator=(FShaderEditorOnlyDataEntry&&) = default;
+		FShaderEditorOnlyDataEntry& operator=(const FShaderEditorOnlyDataEntry&) = default;
+		PRAGMA_ENABLE_DEPRECATION_WARNINGS
 	};
 #endif // WITH_EDITORONLY_DATA
 
@@ -582,8 +592,8 @@ public:
 	int32 FindShaderIndex(const FSHAHash& InHash) const;
 
 #if WITH_EDITORONLY_DATA
-	void AddEditorOnlyData(int32 Index, const FString& DebugName, TConstArrayView<uint8> InPlatformDebugData, TConstArrayView<FShaderCompilerError> InCompilerWarnings, const TArray<FGenericShaderStat>& ShaderStatistics, const FString& DebugInfo = FString());
-	void UpdateEditorOnlyData(int32 Index, const FString& DebugName, TConstArrayView<uint8> InPlatformDebugData, TConstArrayView<FShaderCompilerError> InCompilerWarnings, const FString& DebugInfo);
+	void AddEditorOnlyData(int32 Index, const FString& DebugName, FCompressedBuffer InSymbols, TConstArrayView<FShaderCompilerError> InCompilerWarnings, const TArray<FGenericShaderStat>& ShaderStatistics, const FString& DebugInfo = FString());
+	void UpdateEditorOnlyData(int32 Index, const FString& DebugName, FCompressedBuffer InSymbols, TConstArrayView<FShaderCompilerError> InCompilerWarnings, const FString& DebugInfo);
 
 	RENDERCORE_API void LogShaderCompilerWarnings();
 #endif
diff --git a/Engine/Source/Runtime/RenderCore/Public/ShaderCompilerCore.h b/Engine/Source/Runtime/RenderCore/Public/ShaderCompilerCore.h
--- a/Engine/Source/Runtime/RenderCore/Public/ShaderCompilerCore.h
+++ b/Engine/Source/Runtime/RenderCore/Public/ShaderCompilerCore.h
@@ -38,7 +38,7 @@ inline const int32 ShaderCompileWorkerPipelineJobHeader = 'P';
 
 // Modify this to invalidate _just_ the cache/DDC entries for individual shaders (will not cause shadermaps to rebuild if they are not otherwise out-of-date).
 // This should be bumped for changes to the FShaderCompilerOutput data structure (in addition to ShaderCompileWorkerOutputVersion)
-inline static const FGuid UE_SHADER_CACHE_VERSION = FGuid("092A667D-E873-4E2D-957B-262A421AED5A");
+inline static const FGuid UE_SHADER_CACHE_VERSION = FGuid("B04528D6-1A2C-44C8-A742-E7D497D3562C");
 
 namespace UE::ShaderCompiler
 {
@@ -617,6 +617,7 @@ struct FShaderCompilerOutput
 	 * a different entry point than was set on the FShaderCompilerInput struct. */
 	FString ModifiedEntryPointName;
 
+	UE_DEPRECATED(5.6, "Use symbols accessors (GetSymbolWriteAccess, GetSymbolReadView) on FShaderCode object instead (ShaderCode member)")
 	TArray<uint8> PlatformDebugData;
 
 	TArray<FGenericShaderStat> ShaderStatistics;
@@ -674,7 +675,6 @@ struct FShaderCompilerOutput
 		Ar << Output.bSucceeded;
 		Ar << Output.ModifiedShaderSource;
 		Ar << Output.ModifiedEntryPointName;
-		Ar << Output.PlatformDebugData;
 		Ar << Output.ShaderStatistics;
 
 		// note: intentionally never serializing the following fields:
diff --git a/Engine/Source/Runtime/RenderCore/Public/ShaderCore.h b/Engine/Source/Runtime/RenderCore/Public/ShaderCore.h
--- a/Engine/Source/Runtime/RenderCore/Public/ShaderCore.h
+++ b/Engine/Source/Runtime/RenderCore/Public/ShaderCore.h
@@ -7,6 +7,7 @@
 #pragma once
 
 #include "Compression/OodleDataCompression.h"
+#include "Compression/CompressedBuffer.h"
 #include "Containers/Array.h"
 #include "Containers/ArrayView.h"
 #include "Containers/ContainerAllocationPolicies.h"
@@ -1122,6 +1123,8 @@ public:
 
 class FShaderCode;
 
+RENDERCORE_API FArchive& operator<<(FArchive& Ar, FSharedBuffer& Buffer);
+
 class FShaderCodeResource
 {
 	struct FHeader
@@ -1138,6 +1141,7 @@ class FShaderCodeResource
 
 	FSharedBuffer Header;		// The above FHeader struct persisted in a shared buffer
 	FSharedBuffer Code;			// The bytecode buffer as constructed by FShaderCode::FinalizeShaderCode
+	FCompressedBuffer Symbols;	// Buffer containing the symbols for this bytecode; will be empty if symbols are disabled
 
 	friend class FShaderCode;
 	friend RENDERCORE_API FArchive& operator<<(FArchive& Ar, FShaderCode& Code);
@@ -1159,6 +1163,12 @@ public:
 		return Code;
 	}
 
+	/* Return the buffer storing the (compressed) symbols for this resource */
+	FCompressedBuffer GetSymbolsBuffer() const
+	{
+		return Symbols;
+	}
+
 	/* Returns a single composite buffer referencing both the header and code data to be cached. */
 	FCompositeBuffer GetCacheBuffer() const
 	{
@@ -1180,12 +1190,13 @@ public:
 	/* Sets the Header and Code shared buffer references in this resource to the segments referenced
 	 * by the given composite buffer.
 	 */
-	void PopulateFromComposite(FCompositeBuffer CacheBuffer)
+	void PopulateFromComposite(FCompositeBuffer CacheBuffer, FCompressedBuffer SymbolsBuffer)
 	{
 		check(CacheBuffer.GetSegments().Num() == 2);
 		Header = CacheBuffer.GetSegments()[0];
 		check(Header.GetSize() == sizeof(FHeader));
 		Code = CacheBuffer.GetSegments()[1];
+		Symbols = SymbolsBuffer;
 	}
 
 	/* Populates the header for this code resource with the given sizes and frequency. 
@@ -1229,6 +1240,8 @@ class FShaderCode
 	// access through class methods
 	mutable TArray<uint8> ShaderCodeWithOptionalData;
 
+	mutable TArray<uint8> SymbolData;
+
 	mutable FShaderCodeResource ShaderCodeResource;
 
 	/** ShaderCode may be compressed in SCWs on demand. If this value isn't null, the shader code is compressed. */
@@ -1258,8 +1271,8 @@ public:
 	{
 	}
 
-	// adds CustomData or does nothing if that was already done before
-	void FinalizeShaderCode() const
+	// converts code and symbols into shared buffer representations, optionally overriding the symbols buffer with the given input
+	void FinalizeShaderCode(FCompressedBuffer OverrideSymbolsBuffer = FCompressedBuffer()) const
 	{
 		if (OptionalDataSize != -1)
 		{
@@ -1268,7 +1281,11 @@ public:
 			ShaderCodeWithOptionalData.Append((const uint8*)&OptionalDataSize, sizeof(OptionalDataSize));
 			OptionalDataSize = -1;
 
+			// sanity check: the override symbol buffer is only used currently when merging multiple code outputs into a single one, and in this
+			// case we expect the symbols to be empty (as merging the symbols currently needs to be handled differently in each shader format)
+			check(OverrideSymbolsBuffer.IsNull() || SymbolData.IsEmpty());
 			ShaderCodeResource.Code = MakeSharedBufferFromArray(MoveTemp(ShaderCodeWithOptionalData));
+			ShaderCodeResource.Symbols = !OverrideSymbolsBuffer.IsNull() ? OverrideSymbolsBuffer : FCompressedBuffer::Compress(MakeSharedBufferFromArray(MoveTemp(SymbolData)));
 		}
 	}
 
@@ -1282,6 +1299,18 @@ public:
 		return ShaderCodeWithOptionalData;
 	}
 
+	TArray<uint8>& GetSymbolWriteAccess()
+	{
+		checkf(OptionalDataSize != -1, TEXT("Tried to add shader symbols after being finalized!"));
+		return SymbolData;
+	}
+
+	TConstArrayView<uint8> GetSymbolReadView() const
+	{
+		checkf(OptionalDataSize != -1, TEXT("Tried to read uncompressed symbol data from bytecode after FinalizeShaderCode was called (which compresses the symbol data)"));
+		return SymbolData;
+	}
+
 	int32 GetShaderCodeSize() const
 	{
 		// use the cached size whenever available
