Only process game thread tasks during end of frame syncs. Avoids pathological behavior where a game thread task flushes the rendering thread only to pick up another GT task that has an implicit dependency on the one being processed.

diff --git a/Engine/Source/Runtime/RenderCore/Private/RenderingThread.cpp b/Engine/Source/Runtime/RenderCore/Private/RenderingThread.cpp
--- a/Engine/Source/Runtime/RenderCore/Private/RenderingThread.cpp
+++ b/Engine/Source/Runtime/RenderCore/Private/RenderingThread.cpp
@@ -2146,7 +2146,10 @@ namespace FFrameEndSync
 
 		PipelineFences.Emplace_GetRef().BeginFence(SyncDepth);
 
-		if (!FTaskGraphInterface::Get().IsThreadProcessingTasks(ENamedThreads::GameThread))
+		// Don't process game thread tasks when flushing all threads. This can result in strange behavior where the game thread
+		// is flushing the render thread and then gets pre-empted by another task that has an implicit dependency on the one
+		// being processed.
+		if (FlushMode == EFlushMode::EndFrame && !FTaskGraphInterface::Get().IsThreadProcessingTasks(ENamedThreads::GameThread))
 		{
 			// need to process gamethread tasks at least once a frame no matter what
 			FTaskGraphInterface::Get().ProcessThreadUntilIdle(ENamedThreads::GameThread);
