[shaders] fix cases of non-finalized code tripping asserts on finished jobs
- don't try to get the symbols buffer if the job failed
- explicitly finalize shader code in compile execution path to handle case where compilation executes in-process, instead of relying on serialization to do it for us (we don't hit the serialization path in the in-process compile case)

diff --git a/Engine/Source/Runtime/RenderCore/Private/ShaderCore.cpp b/Engine/Source/Runtime/RenderCore/Private/ShaderCore.cpp
--- a/Engine/Source/Runtime/RenderCore/Private/ShaderCore.cpp
+++ b/Engine/Source/Runtime/RenderCore/Private/ShaderCore.cpp
@@ -1764,7 +1764,8 @@ public:
 				Job.SecondaryOutput->Target = Job.Input.Target;
 				// also set the job input hash on the output struct for validation purposes
 				Job.SecondaryOutput->ValidateInputHash = Job.Input.Hash;
-				if (Job.Output.bSucceeded)
+				Job.SecondaryOutput->ShaderCode.FinalizeShaderCode();
+				if (Job.SecondaryOutput->bSucceeded)
 				{
 					Job.SecondaryOutput->GenerateOutputHash();
 				}
@@ -1780,6 +1781,7 @@ public:
 		Job.Output.Target = Job.Input.Target;
 		// also set the job input hash on the output struct for validation purposes
 		Job.Output.ValidateInputHash = Job.Input.Hash;
+		Job.Output.ShaderCode.FinalizeShaderCode();
 
 		// clear out the modified source/entry point fields if they aren't needed to avoid the data transfer/ddc overhead
 		if (!Job.Input.NeedsOriginalShaderSource())
@@ -4250,7 +4252,7 @@ void FShaderCompileJob::OnComplete()
 	Output.PreprocessTime = PreprocessOutput.ElapsedTime;
 
 	const IShaderFormat* ShaderFormat = nullptr;
-	FCompressedBuffer SymbolsBuffer = Output.GetFinalizedCodeResource().GetSymbolsBuffer();
+	FCompressedBuffer SymbolsBuffer = Output.bSucceeded ? Output.GetFinalizedCodeResource().GetSymbolsBuffer() : FCompressedBuffer();
 	const bool bHasSymbols = SymbolsBuffer.GetRawSize() > 0;
 	if (bHasSymbols || Input.DumpDebugInfoEnabled())
 	{
