[ShaderCore]
- Replace custom code and simplify by using recently introduced TStripedMap

diff --git a/Engine/Source/Runtime/RenderCore/Private/ShaderCore.cpp b/Engine/Source/Runtime/RenderCore/Private/ShaderCore.cpp
--- a/Engine/Source/Runtime/RenderCore/Private/ShaderCore.cpp
+++ b/Engine/Source/Runtime/RenderCore/Private/ShaderCore.cpp
@@ -26,6 +26,7 @@
 #include "Misc/ScopeRWLock.h"
 #include "Misc/ScopedSlowTask.h"
 #include "Misc/StringBuilder.h"
+#include "Containers/StripedMap.h"
 #include "Modules/ModuleManager.h"
 #include "RHIShaderFormatDefinitions.inl"
 #include "Serialization/MemoryHasher.h"
@@ -237,9 +238,6 @@ void FShaderCompilerDefinitions::InitializeInitialDefines(const FShaderCompilerD
 	GInitialDefines->InitialDefineCount = InDefines.Pairs.Num();
 }
 
-// Apply lock striping as we're mostly reader lock bound.
-constexpr int32 GSHADERFILECACHE_BUCKETS = 31; /* prime number for best distribution using modulo */
-
 struct FShaderFileCacheEntry
 {
 	FString Source;
@@ -247,14 +245,9 @@ struct FShaderFileCacheEntry
 	FShaderPreprocessDependenciesShared Dependencies;		// Stripped source with include dependencies, all in one shareable struct
 };
 
-struct FShaderFileCache
-{
-	/** Protects Map from simultaneous access by multiple threads. */
-	FRWLock Lock;
-
-	/** The shader file cache, used to minimize shader file reads */
-	TMap<FString, FShaderFileCacheEntry> Map;
-} GShaderFileCache[GSHADERFILECACHE_BUCKETS];
+// Apply lock striping as we're mostly reader lock bound.
+// Use prime number for the number of buckets for best distribution using modulo.
+TStripedMap<31, FString, FShaderFileCacheEntry> GShaderFileCache;
 
 class FShaderHashCache
 {
@@ -2515,75 +2508,36 @@ bool LoadShaderSourceFile(const TCHAR* InVirtualFilePath, EShaderPlatform Shader
 		FString VirtualFilePath(InVirtualFilePath);
 		FixupShaderFilePath(VirtualFilePath, ShaderPlatform, ShaderPlatformName);
 
-		FShaderFileCacheEntry* CachedFile = nullptr;
-
-		// First try a shared lock and only acquire exclusive access if element is not found in cache
-		uint32 CurrentHash = GetTypeHash(VirtualFilePath);
-		FShaderFileCache& ShaderFileCache = GShaderFileCache[CurrentHash % GSHADERFILECACHE_BUCKETS];
-		{
-			FRWScopeLock ScopeLock(ShaderFileCache.Lock, SLT_ReadOnly);
-			CachedFile = ShaderFileCache.Map.FindByHash(CurrentHash, VirtualFilePath);
-		}
-
-		if (CachedFile)
-		{
-			if (OutFileContents)
-			{
-				*OutFileContents = CachedFile->Source;
-			}
-			if (OutStrippedContents)
-			{
-				*OutStrippedContents = CachedFile->StrippedSource;
-			}
-			bResult = true;
-		}
-		else 
-		{
-			FRWScopeLock ScopeLock(ShaderFileCache.Lock, SLT_Write);
-
-			// Double-check the cache while holding exclusive lock as another thread may have added the item we're looking for
-			CachedFile = ShaderFileCache.Map.FindByHash(CurrentHash, VirtualFilePath);
-
-			// if this file has already been loaded and cached, use that
-			if (CachedFile)
-			{
-				if (OutFileContents)
-				{
-					*OutFileContents = CachedFile->Source;
-				}
-				if (OutStrippedContents)
+		bResult = 
+			GShaderFileCache.FindOrTryProduceAndApply(VirtualFilePath,
+				[&VirtualFilePath, &OutCompileErrors](FShaderFileCacheEntry& CachedFile)
 				{
-					*OutStrippedContents = CachedFile->StrippedSource;
-				}
-				bResult = true;
-			}
-			else
-			{
-				FString ShaderFilePath = GetShaderSourceFilePath(VirtualFilePath, OutCompileErrors);
+					FString ShaderFilePath = GetShaderSourceFilePath(VirtualFilePath, OutCompileErrors);
 
-				// verify SHA hash of shader files on load. missing entries trigger an error
-				FShaderFileCacheEntry FileContents;
-				if (!ShaderFilePath.IsEmpty() && FFileHelper::LoadFileToString(FileContents.Source, *ShaderFilePath, FFileHelper::EHashOptions::EnableVerify|FFileHelper::EHashOptions::ErrorMissingHash) )
-				{
-					TArray<ANSICHAR>* StrippedSource = new TArray<ANSICHAR>;
-					ShaderConvertAndStripComments(FileContents.Source, *StrippedSource);
-					FileContents.StrippedSource = MakeShareable(StrippedSource);
+					// verify SHA hash of shader files on load. missing entries trigger an error
+					if (!ShaderFilePath.IsEmpty() && FFileHelper::LoadFileToString(CachedFile.Source, *ShaderFilePath, FFileHelper::EHashOptions::EnableVerify | FFileHelper::EHashOptions::ErrorMissingHash))
+					{
+						TArray<ANSICHAR>* StrippedSource = new TArray<ANSICHAR>;
+						ShaderConvertAndStripComments(CachedFile.Source, *StrippedSource);
+						CachedFile.StrippedSource = MakeShareable(StrippedSource);
 
-					//update the shader file cache
-					ShaderFileCache.Map.AddByHash(CurrentHash, VirtualFilePath, FileContents);
+						return true;
+					}
 
+					return false;
+				},
+				[&OutFileContents, &OutStrippedContents](const FShaderFileCacheEntry& CachedFile)
+				{
 					if (OutFileContents)
 					{
-						*OutFileContents = FileContents.Source;
+						*OutFileContents = CachedFile.Source;
 					}
 					if (OutStrippedContents)
 					{
-						*OutStrippedContents = FileContents.StrippedSource;
+						*OutStrippedContents = CachedFile.StrippedSource;
 					}
-					bResult = true;
 				}
-			}
-		}
+			);
 	}
 
 	INC_FLOAT_STAT_BY(STAT_ShaderCompiling_LoadingShaderFiles,(float)ShaderFileLoadingTime);
@@ -2686,20 +2640,17 @@ static TUniquePtr<FShaderPreprocessDependencies> ShaderPreprocessDependenciesBeg
 // Adds finished dependencies to the cache
 static void ShaderPreprocessDependenciesEnd(const TCHAR* VirtualFilePath, TUniquePtr<FShaderPreprocessDependencies> PreprocessDependencies, EShaderPlatform Platform)
 {
-	uint32 CurrentHash = FCrc::Strihash_DEPRECATED(VirtualFilePath);
-	FShaderFileCache& ShaderFileCache = GShaderFileCache[CurrentHash % GSHADERFILECACHE_BUCKETS];
-	{
-		FRWScopeLock ScopeLock(ShaderFileCache.Lock, SLT_Write);
-		FShaderFileCacheEntry* CachedFile = ShaderFileCache.Map.FindByHash(CurrentHash, VirtualFilePath);
-		if (CachedFile)
+	GShaderFileCache.FindAndApply(
+		VirtualFilePath,
+		[&PreprocessDependencies](FShaderFileCacheEntry& CachedFile)
 		{
 			// Another thread could have finished the job...  If not, set the dependencies.
-			if (!CachedFile->Dependencies.IsValid())
+			if (!CachedFile.Dependencies.IsValid())
 			{
-				CachedFile->Dependencies = MakeShareable(PreprocessDependencies.Release());
+				CachedFile.Dependencies = MakeShareable(PreprocessDependencies.Release());
 			}
 		}
-	}
+	);
 }
 
 static void AddPreprocessDependency(FShaderPreprocessDependencies& Dependencies, const FShaderPreprocessDependency& Dependency)
@@ -2928,19 +2879,18 @@ static void InternalGetShaderIncludes(const TCHAR* EntryPointVirtualFilePath, co
 
 bool GetShaderPreprocessDependencies(const TCHAR* VirtualFilePath, EShaderPlatform ShaderPlatform, FShaderPreprocessDependenciesShared& OutDependencies)
 {
-	// Same case insensitive hash used by FString
-	uint32 CurrentHash = FCrc::Strihash_DEPRECATED(VirtualFilePath);
-	FShaderFileCache& ShaderFileCache = GShaderFileCache[CurrentHash % GSHADERFILECACHE_BUCKETS];
-
-	FRWScopeLock ScopeLock(ShaderFileCache.Lock, SLT_ReadOnly);
-	FShaderFileCacheEntry* CachedFile = ShaderFileCache.Map.FindByHash(CurrentHash, VirtualFilePath);
-	if (CachedFile && CachedFile->Dependencies.IsValid())
-	{
-		OutDependencies = CachedFile->Dependencies;
-		return true;
-	}
-
-	return false;
+	bool bResult = false;
+	GShaderFileCache.FindAndApply(VirtualFilePath,
+		[&bResult, &OutDependencies](const FShaderFileCacheEntry& CachedFile)
+		{
+			if (CachedFile.Dependencies.IsValid())
+			{
+				OutDependencies = CachedFile.Dependencies;
+				bResult = true;
+			}
+		}
+	);
+	return bResult;
 }
 
 /**
@@ -3323,13 +3273,8 @@ void FlushShaderFileCache()
 		FRWScopeLock ShaderHashAccessLock(GShaderHashAccessRWLock, SLT_Write);
 		GShaderHashCache.Empty();
 	}
-	{
-		for (int32 Index = 0; Index < UE_ARRAY_COUNT(GShaderFileCache); ++Index)
-		{
-			FRWScopeLock ScopeLock(GShaderFileCache[Index].Lock, SLT_Write);
-			GShaderFileCache[Index].Map.Empty();
-		}
-	}
+
+	GShaderFileCache.Empty();
 
 	UE_LOG(LogShaders, Log, TEXT("FlushShaderFileCache() end"));
 }
@@ -3339,12 +3284,8 @@ void InvalidateShaderFileCacheEntry(const TCHAR* InVirtualFilePath, EShaderPlatf
 	FString VirtualFilePath(InVirtualFilePath);
 	FixupShaderFilePath(VirtualFilePath, InShaderPlatform, InShaderPlatformName);
 
-	const uint32 CurrentHash = GetTypeHash(VirtualFilePath);
-	FShaderFileCache& ShaderFileCache = GShaderFileCache[CurrentHash % GSHADERFILECACHE_BUCKETS];
-	{
-		FRWScopeLock ScopeLock(ShaderFileCache.Lock, SLT_Write);
-		ShaderFileCache.Map.RemoveByHash(CurrentHash, VirtualFilePath);
-	}
+	GShaderFileCache.Remove(VirtualFilePath);
+
 	{
 		FRWScopeLock ShaderHashAccessLock(GShaderHashAccessRWLock, SLT_Write);
 		GShaderHashCache.RemoveHash(InShaderPlatform, VirtualFilePath);
