[Shaders] Use read-lock in diagnostics helper function GatherShaderMapResourceReferences().

#rnx

diff --git a/Engine/Source/Runtime/RenderCore/Private/ShaderCodeLibrary.cpp b/Engine/Source/Runtime/RenderCore/Private/ShaderCodeLibrary.cpp
--- a/Engine/Source/Runtime/RenderCore/Private/ShaderCodeLibrary.cpp
+++ b/Engine/Source/Runtime/RenderCore/Private/ShaderCodeLibrary.cpp
@@ -2835,20 +2835,22 @@ public:
 #endif
 	}
 
-	void GatherShaderMapResourceReferences(const FShaderMapResource_SharedCode* InResource, TArray<FString>& OutLibrariesWithRefToResource) const
+	void GatherShaderMapResourceReferences(const FShaderMapResource_SharedCode* InResource, TArray<FString>& OutLibrariesWithRefToResource)
 	{
 		using namespace UE::ShaderLibrary::Private;
 
-		for (auto It = NamedLibrariesStack.CreateConstIterator(); It; ++It)
+		FRWScopeLock ReadLock(NamedLibrariesMutex, SLT_ReadOnly);
+
+		for (const TTuple<FString, TUniquePtr<UE::ShaderLibrary::Private::FNamedShaderLibrary>>& NamedLibraryPair : NamedLibrariesStack)
 		{
-			const FNamedShaderLibrary* Library = It.Value().Get();
+			const FNamedShaderLibrary* Library = NamedLibraryPair.Value.Get();
 			for (const TUniquePtr<FShaderLibraryInstance>& LibraryInstance : Library->Components)
 			{
 				for (int32 LibraryResourceIndex = 0; LibraryResourceIndex < LibraryInstance->GetNumResources(); ++LibraryResourceIndex)
 				{
 					if (LibraryInstance->GetResource(LibraryResourceIndex).GetReference() == InResource)
 					{
-						OutLibrariesWithRefToResource.Add(It.Value()->LogicalName);
+						OutLibrariesWithRefToResource.Add(NamedLibraryPair.Value->LogicalName);
 					}
 				}
 			}
