Substrate - Fixed dbuffer decal still used if r.Substrate.BlendableGBuffer=1 but  r.DBuffer has also been left to be 1.

#rnx
#tests Editor dx12

diff --git a/Engine/Source/Runtime/RenderCore/Private/RenderUtils.cpp b/Engine/Source/Runtime/RenderCore/Private/RenderUtils.cpp
--- a/Engine/Source/Runtime/RenderCore/Private/RenderUtils.cpp
+++ b/Engine/Source/Runtime/RenderCore/Private/RenderUtils.cpp
@@ -1468,7 +1468,7 @@ bool IsSingleLayerWaterDepthPrepassEnabled(const FStaticShaderPlatform Platform,
 
 RENDERCORE_API bool IsUsingDBuffers(const FStaticShaderPlatform Platform)
 {
-	return GDBufferPlatformMask[(int)Platform];
+	return GDBufferPlatformMask[(int)Platform] && (!Substrate::IsSubstrateEnabled() || !Substrate::IsSubstrateBlendableGBufferEnabled(Platform));
 }
 
 RENDERCORE_API bool AreSkinCacheShadersEnabled(EShaderPlatform Platform)
@@ -2144,7 +2144,7 @@ namespace Substrate
 		const bool bDBufferPassSupported = NormalQuality == 0 || (NormalQuality > 0 && IsConsolePlatform(InPlatform));
 		
 		static FShaderPlatformCachedIniValue<int32> CVarSubstrateDBufferPassPlatform(TEXT("r.Substrate.DBufferPass"));
-		return IsSubstrateEnabled() && !IsSubstrateBlendableGBufferEnabled(InPlatform) && IsUsingDBuffers(InPlatform) && bDBufferPassSupported && CVarSubstrateDBufferPassPlatform.Get(InPlatform) > 0;
+		return IsSubstrateEnabled() && IsUsingDBuffers(InPlatform) && bDBufferPassSupported && CVarSubstrateDBufferPassPlatform.Get(InPlatform) > 0;
 	}
 
 	bool IsOpaqueRoughRefractionEnabled(EShaderPlatform InPlatform)
