[Shaders] Provide string representation for shader permutation domains.

This adds "GetPermutationVectorString" as new entry to the custom vtable in shader types which prints a string representation of the shader permutation domain.
This is intended for diagnostics tools to provide a human readable format of all shader permutation vectors.

#rnx

diff --git a/Engine/Plugins/Compositing/OpenColorIO/Source/OpenColorIO/Public/OpenColorIOShaderType.h b/Engine/Plugins/Compositing/OpenColorIO/Source/OpenColorIO/Public/OpenColorIOShaderType.h
--- a/Engine/Plugins/Compositing/OpenColorIO/Source/OpenColorIO/Public/OpenColorIOShaderType.h
+++ b/Engine/Plugins/Compositing/OpenColorIO/Source/OpenColorIO/Public/OpenColorIOShaderType.h
@@ -93,6 +93,9 @@ public:
 #endif // WITH_EDITOR
 		uint32 InTypeSize,
 		const FShaderParametersMetadata* InRootParametersMetadata = nullptr
+#if WITH_EDITOR
+		, GetPermutationIdStringType InGetPermutationIdStringRef = nullptr
+#endif // WITH_EDITOR
 		):
 		FShaderType(EShaderTypeForDynamicCast::OCIO, InTypeLayout, InName, InSourceFilename, InFunctionName, SF_Pixel, InTotalPermutationCount,
 			InConstructSerializedRef,
@@ -107,7 +110,11 @@ public:
 			InGetOverrideJobPriorityRef,
 #endif // WITH_EDITOR
 			InTypeSize,
-			InRootParametersMetadata)
+			InRootParametersMetadata
+#if WITH_EDITOR
+			, InGetPermutationIdStringRef
+#endif // WITH_EDITOR
+		)
 	{
 		check(InTotalPermutationCount == 1);
 	}
diff --git a/Engine/Plugins/Experimental/NNERuntimeIREE/Source/NNERuntimeIREEShader/Internal/NNERuntimeIREEShaderType.h b/Engine/Plugins/Experimental/NNERuntimeIREE/Source/NNERuntimeIREEShader/Internal/NNERuntimeIREEShaderType.h
--- a/Engine/Plugins/Experimental/NNERuntimeIREE/Source/NNERuntimeIREEShader/Internal/NNERuntimeIREEShaderType.h
+++ b/Engine/Plugins/Experimental/NNERuntimeIREE/Source/NNERuntimeIREEShader/Internal/NNERuntimeIREEShaderType.h
@@ -86,6 +86,9 @@ public:
 #endif // WITH_EDITOR
 		uint32 InTypeSize,
 		const FShaderParametersMetadata* InRootParametersMetadata = nullptr
+#if WITH_EDITOR
+		, GetPermutationIdStringType InGetPermutationIdStringRef = nullptr
+#endif // WITH_EDITOR
 		):
 		FShaderType(
 			EShaderTypeForDynamicCast::NNERuntimeIREE, 
@@ -108,6 +111,9 @@ public:
 #endif // WITH_EDITOR
 			InTypeSize,
 			InRootParametersMetadata
+#if WITH_EDITOR
+			, InGetPermutationIdStringRef
+#endif // WITH_EDITOR
 			)
 	{
 	}
diff --git a/Engine/Plugins/FX/Niagara/Source/NiagaraShader/Public/NiagaraShaderType.h b/Engine/Plugins/FX/Niagara/Source/NiagaraShader/Public/NiagaraShaderType.h
--- a/Engine/Plugins/FX/Niagara/Source/NiagaraShader/Public/NiagaraShaderType.h
+++ b/Engine/Plugins/FX/Niagara/Source/NiagaraShader/Public/NiagaraShaderType.h
@@ -99,6 +99,9 @@ public:
 #endif // WITH_EDITOR
 		uint32 InTypeSize,
 		const FShaderParametersMetadata* InRootParametersMetadata = nullptr
+#if WITH_EDITOR
+		, GetPermutationIdStringType InGetPermutationIdStringRef = nullptr
+#endif // WITH_EDITOR
 	)
 		: FShaderType(EShaderTypeForDynamicCast::Niagara, InTypeLayout, InName, InSourceFilename, InFunctionName, SF_Compute, InTotalPermutationCount,
 			InConstructSerializedRef,
@@ -114,6 +117,9 @@ public:
 #endif // WITH_EDITOR
 			InTypeSize,
 			InRootParametersMetadata
+#if WITH_EDITOR
+			, InGetPermutationIdStringRef
+#endif // WITH_EDITOR
 		)
 	{
 		check(InTotalPermutationCount == 1);
diff --git a/Engine/Plugins/Runtime/ComputeFramework/Source/ComputeFramework/Private/ComputeFramework/ComputeKernelShaderType.h b/Engine/Plugins/Runtime/ComputeFramework/Source/ComputeFramework/Private/ComputeFramework/ComputeKernelShaderType.h
--- a/Engine/Plugins/Runtime/ComputeFramework/Source/ComputeFramework/Private/ComputeFramework/ComputeKernelShaderType.h
+++ b/Engine/Plugins/Runtime/ComputeFramework/Source/ComputeFramework/Private/ComputeFramework/ComputeKernelShaderType.h
@@ -84,6 +84,9 @@ public:
 #endif // WITH_EDITOR
 		uint32 InTypeSize,
 		const FShaderParametersMetadata* InRootParametersMetadata = nullptr
+#if WITH_EDITOR
+		, GetPermutationIdStringType InGetPermutationIdStringRef = nullptr
+#endif // WITH_EDITOR
 		):
 		FShaderType(
 			EShaderTypeForDynamicCast::ComputeKernel, 
@@ -106,6 +109,9 @@ public:
 #endif // WITH_EDITOR
 			InTypeSize,
 			InRootParametersMetadata
+#if WITH_EDITOR
+			, InGetPermutationIdStringRef
+#endif // WITH_EDITOR
 			)
 	{
 	}
diff --git a/Engine/Source/Runtime/Engine/Public/MaterialShaderType.h b/Engine/Source/Runtime/Engine/Public/MaterialShaderType.h
--- a/Engine/Source/Runtime/Engine/Public/MaterialShaderType.h
+++ b/Engine/Source/Runtime/Engine/Public/MaterialShaderType.h
@@ -130,47 +130,7 @@ public:
 		{}
 	};
 
-	FMaterialShaderType(
-		FTypeLayoutDesc& InTypeLayout,
-		const TCHAR* InName,
-		const TCHAR* InSourceFilename,
-		const TCHAR* InFunctionName,
-		uint32 InFrequency,
-		int32 InTotalPermutationCount,
-		ConstructSerializedType InConstructSerializedRef,
-		ConstructCompiledType InConstructCompiledRef,
-		ShouldCompilePermutationType InShouldCompilePermutationRef,
-		ShouldPrecachePermutationType InShouldPrecachePermutationRef,
-		GetRayTracingPayloadTypeType InGetRayTracingPayloadTypeRef,
-		GetShaderBindingLayoutType InGetShaderBindingLayoutTypeRef,
-#if WITH_EDITOR
-		ModifyCompilationEnvironmentType InModifyCompilationEnvironmentRef,
-		ValidateCompiledResultType InValidateCompiledResultRef,
-		GetOverrideJobPriorityType InGetOverrideJobPriorityRef,
-#endif // WITH_EDITOR
-		uint32 InTypeSize,
-		const FShaderParametersMetadata* InRootParametersMetadata = nullptr
-		):
-		FShaderType(EShaderTypeForDynamicCast::Material, InTypeLayout, InName, InSourceFilename, InFunctionName, InFrequency, InTotalPermutationCount,
-			InConstructSerializedRef,
-			InConstructCompiledRef,
-			InShouldCompilePermutationRef,
-			InShouldPrecachePermutationRef,
-			InGetRayTracingPayloadTypeRef,
-			InGetShaderBindingLayoutTypeRef,
-#if WITH_EDITOR
-			InModifyCompilationEnvironmentRef,
-			InValidateCompiledResultRef,
-			InGetOverrideJobPriorityRef,
-#endif // WITH_EDITOR
-			InTypeSize,
-			InRootParametersMetadata
-		)
-	{
-		checkf(FPaths::GetExtension(InSourceFilename) == TEXT("usf"),
-			TEXT("Incorrect virtual shader path extension for material shader '%s': Only .usf files should be compiled."),
-			InSourceFilename);
-	}
+	IMPLEMENT_SHADER_TYPE_CONSTRUCTOR(FMaterialShaderType, EShaderTypeForDynamicCast::Material, TEXT("material shader"));
 
 #if WITH_EDITOR
 	/**
diff --git a/Engine/Source/Runtime/Engine/Public/MeshMaterialShaderType.h b/Engine/Source/Runtime/Engine/Public/MeshMaterialShaderType.h
--- a/Engine/Source/Runtime/Engine/Public/MeshMaterialShaderType.h
+++ b/Engine/Source/Runtime/Engine/Public/MeshMaterialShaderType.h
@@ -43,47 +43,7 @@ public:
 		{}
 	};
 
-	FMeshMaterialShaderType(
-		FTypeLayoutDesc& InTypeLayout,
-		const TCHAR* InName,
-		const TCHAR* InSourceFilename,
-		const TCHAR* InFunctionName,
-		uint32 InFrequency,
-		int32 InTotalPermutationCount,
-		ConstructSerializedType InConstructSerializedRef,
-		ConstructCompiledType InConstructCompiledRef,
-		ShouldCompilePermutationType InShouldCompilePermutationRef,
-		ShouldPrecachePermutationType InShouldPrecachePermutationRef,
-		GetRayTracingPayloadTypeType InGetRayTracingPayloadTypeRef,
-		GetShaderBindingLayoutType InGetShaderBindingLayoutTypeRef,
-#if WITH_EDITOR
-		ModifyCompilationEnvironmentType InModifyCompilationEnvironmentRef,
-		ValidateCompiledResultType InValidateCompiledResultRef,
-		GetOverrideJobPriorityType InGetOverrideJobPriorityRef,
-#endif // WITH_EDITOR
-		uint32 InTypeSize,
-		const FShaderParametersMetadata* InRootParametersMetadata = nullptr
-		):
-		FShaderType(EShaderTypeForDynamicCast::MeshMaterial, InTypeLayout, InName,InSourceFilename,InFunctionName,InFrequency,InTotalPermutationCount,
-			InConstructSerializedRef,
-			InConstructCompiledRef,
-			InShouldCompilePermutationRef,
-			InShouldPrecachePermutationRef,
-			InGetRayTracingPayloadTypeRef,
-			InGetShaderBindingLayoutTypeRef,
-#if WITH_EDITOR
-			InModifyCompilationEnvironmentRef,
-			InValidateCompiledResultRef,
-			InGetOverrideJobPriorityRef,
-#endif // WITH_EDITOR
-			InTypeSize,
-			InRootParametersMetadata
-		)
-	{
-		checkf(FPaths::GetExtension(InSourceFilename) == TEXT("usf"),
-			TEXT("Incorrect virtual shader path extension for mesh material shader '%s': Only .usf files should be compiled."),
-			InSourceFilename);
-	}
+	IMPLEMENT_SHADER_TYPE_CONSTRUCTOR(FMeshMaterialShaderType, EShaderTypeForDynamicCast::MeshMaterial, TEXT("mesh material shader"));
 
 #if WITH_EDITOR
 	/**
diff --git a/Engine/Source/Runtime/RenderCore/Private/Shader.cpp b/Engine/Source/Runtime/RenderCore/Private/Shader.cpp
--- a/Engine/Source/Runtime/RenderCore/Private/Shader.cpp
+++ b/Engine/Source/Runtime/RenderCore/Private/Shader.cpp
@@ -255,31 +255,35 @@ FShaderType::FShaderType(
 #endif // WITH_EDITOR
 	uint32 InTypeSize,
 	const FShaderParametersMetadata* InRootParametersMetadata
-):
-	ShaderTypeForDynamicCast(InShaderTypeForDynamicCast),
-	TypeLayout(&InTypeLayout),
-	Name(InName),
-	TypeName(InName),
-	HashedName(TypeName),
-	HashedSourceFilename(InSourceFilename),
-	SourceFilename(InSourceFilename),
-	FunctionName(InFunctionName),
-	Frequency(InFrequency),
-	TypeSize(InTypeSize),
-	TotalPermutationCount(InTotalPermutationCount),
-	ConstructSerializedRef(InConstructSerializedRef),
-	ConstructCompiledRef(InConstructCompiledRef),
-	ShouldCompilePermutationRef(InShouldCompilePermutationRef),
-	ShouldPrecachePermutationRef(InShouldPrecachePermutationRef),
-	GetRayTracingPayloadTypeRef(InGetRayTracingPayloadTypeRef),
-	GetShaderBindingLayoutTypeRef(InGetShaderBindingLayoutTypeRef),
 #if WITH_EDITOR
-	ModifyCompilationEnvironmentRef(InModifyCompilationEnvironmentRef),
-	ValidateCompiledResultRef(InValidateCompiledResultRef),
-	GetOverrideJobPriorityRef(InGetOverrideJobPriorityRef),
+	, GetPermutationIdStringType InGetPermutationIdStringRef
+#endif
+)
+	: ShaderTypeForDynamicCast(InShaderTypeForDynamicCast)
+	, TypeLayout(&InTypeLayout)
+	, Name(InName)
+	, TypeName(InName)
+	, HashedName(TypeName)
+	, HashedSourceFilename(InSourceFilename)
+	, SourceFilename(InSourceFilename)
+	, FunctionName(InFunctionName)
+	, Frequency(InFrequency)
+	, TypeSize(InTypeSize)
+	, TotalPermutationCount(InTotalPermutationCount)
+	, ConstructSerializedRef(InConstructSerializedRef)
+	, ConstructCompiledRef(InConstructCompiledRef)
+	, ShouldCompilePermutationRef(InShouldCompilePermutationRef)
+	, ShouldPrecachePermutationRef(InShouldPrecachePermutationRef)
+	, GetRayTracingPayloadTypeRef(InGetRayTracingPayloadTypeRef)
+	, GetShaderBindingLayoutTypeRef(InGetShaderBindingLayoutTypeRef)
+#if WITH_EDITOR
+	, ModifyCompilationEnvironmentRef(InModifyCompilationEnvironmentRef)
+	, ValidateCompiledResultRef(InValidateCompiledResultRef)
+	, GetOverrideJobPriorityRef(InGetOverrideJobPriorityRef)
+	, GetPermutationIdStringRef(InGetPermutationIdStringRef)
 #endif // WITH_EDITOR
-	RootParametersMetadata(InRootParametersMetadata),
-	GlobalListLink(this)
+	, RootParametersMetadata(InRootParametersMetadata)
+	, GlobalListLink(this)
 {
 	FTypeLayoutDesc::Register(InTypeLayout);
 
@@ -530,6 +534,18 @@ void FShaderType::UpdateReferencedUniformBufferNames(const TMap<FString, TArray<
 	ReferencedUniformBuffers.Empty();
 	GenerateReferencedUniformBuffers(SourceFilename, Name, ShaderFileToUniformBufferVariables, ReferencedUniformBuffers);
 }
+
+FString FShaderType::GetPermutationIdString(int32 PermutationId, bool bFullNames) const
+{
+	// This function is for diagnostics only, so ignore it if the shader type did not initialize its format string
+	if (GetPermutationIdStringRef)
+	{
+		FString PermutationIdentifier;
+		(*GetPermutationIdStringRef)(PermutationId, PermutationIdentifier, bFullNames);
+		return PermutationIdentifier;
+	}
+	return FString();
+}
 #endif // WITH_EDITOR
 
 ERayTracingPayloadType FShaderType::GetRayTracingPayloadType(const int32 PermutationId) const
diff --git a/Engine/Source/Runtime/RenderCore/Public/GlobalShader.h b/Engine/Source/Runtime/RenderCore/Public/GlobalShader.h
--- a/Engine/Source/Runtime/RenderCore/Public/GlobalShader.h
+++ b/Engine/Source/Runtime/RenderCore/Public/GlobalShader.h
@@ -90,48 +90,7 @@ public:
 
 	typedef FShader::CompiledShaderInitializerType CompiledShaderInitializerType;
 
-	FGlobalShaderType(
-		FTypeLayoutDesc& InTypeLayout,
-		const TCHAR* InName,
-		const TCHAR* InSourceFilename,
-		const TCHAR* InFunctionName,
-		uint32 InFrequency,
-		int32 InTotalPermutationCount,
-		ConstructSerializedType InConstructSerializedRef,
-		ConstructCompiledType InConstructCompiledRef,
-		ShouldCompilePermutationType InShouldCompilePermutationRef,
-		ShouldPrecachePermutationType InShouldPrecachePermutationRef,
-		GetRayTracingPayloadTypeType InGetRayTracingPayloadTypeRef,
-		GetShaderBindingLayoutType InGetShaderBindingLayoutTypeRef,
-#if WITH_EDITOR
-		ModifyCompilationEnvironmentType InModifyCompilationEnvironmentRef,
-		ValidateCompiledResultType InValidateCompiledResultRef,
-		GetOverrideJobPriorityType InGetOverrideJobPriorityRef,
-#endif // WITH_EDITOR
-		uint32 InTypeSize,
-		const FShaderParametersMetadata* InRootParametersMetadata = nullptr
-		):
-		FShaderType(EShaderTypeForDynamicCast::Global, InTypeLayout, InName, InSourceFilename, InFunctionName, InFrequency, InTotalPermutationCount,
-			InConstructSerializedRef,
-			InConstructCompiledRef,
-			InShouldCompilePermutationRef,
-			InShouldPrecachePermutationRef,
-			InGetRayTracingPayloadTypeRef,
-			InGetShaderBindingLayoutTypeRef,
-#if WITH_EDITOR
-			InModifyCompilationEnvironmentRef,
-			InValidateCompiledResultRef,
-			InGetOverrideJobPriorityRef,
-#endif // WITH_EDITOR
-			InTypeSize,
-			InRootParametersMetadata
-		)
-	{
-		checkf(FPaths::GetExtension(InSourceFilename) == TEXT("usf"),
-			TEXT("Incorrect virtual shader path extension for global shader '%s': Only .usf files should be "
-			     "compiled."),
-			InSourceFilename);
-	}
+	IMPLEMENT_SHADER_TYPE_CONSTRUCTOR(FGlobalShaderType, EShaderTypeForDynamicCast::Global, TEXT("global shader"));
 
 	/**
 	 * Checks if the shader type should be cached for a particular platform.
diff --git a/Engine/Source/Runtime/RenderCore/Public/Shader.h b/Engine/Source/Runtime/RenderCore/Public/Shader.h
--- a/Engine/Source/Runtime/RenderCore/Public/Shader.h
+++ b/Engine/Source/Runtime/RenderCore/Public/Shader.h
@@ -42,6 +42,7 @@
 #include "ShaderParameterMetadata.h"
 #include "ShaderParameters.h"
 #include "ShaderPermutation.h"
+#include "ShaderPermutationUtils.h"
 #include "Templates/RefCounting.h"
 #include "Templates/TypeHash.h"
 #include "Templates/UniquePtr.h"
@@ -811,14 +812,6 @@ inline int16 GetParameterIndex(const FShaderParameterBindings::FResourceParamete
 inline int16 GetParameterIndex(const FShaderParameterBindings::FBindlessResourceParameter& Parameter) { return Parameter.GlobalConstantOffset; }
 inline int16 GetParameterIndex(const FShaderParameterBindings::FParameterStructReference& Parameter) { return Parameter.BufferIndex; }
 
-// Flags that can specialize shader permutations compiled for specific platforms
-enum class EShaderPermutationFlags : uint32
-{
-	None = 0u,
-	HasEditorOnlyData = (1u << 0),
-};
-ENUM_CLASS_FLAGS(EShaderPermutationFlags);
-
 enum class EShaderPermutationPrecacheRequest : uint8
 {
 	Precached,			//< Permutation should be precached because it can be used at runtime
@@ -828,25 +821,6 @@ enum class EShaderPermutationPrecacheRequest : uint8
 
 RENDERCORE_API EShaderPermutationFlags GetShaderPermutationFlags(const FPlatformTypeLayoutParameters& LayoutParams);
 
-struct FShaderPermutationParameters
-{
-	// Shader platform to compile to.
-	const EShaderPlatform Platform;
-
-	// Unique permutation identifier of the material shader type.
-	const int32 PermutationId;
-
-	// Flags that describe the permutation
-	const EShaderPermutationFlags Flags;
-
-	// Default to include editor-only shaders, to maintain backwards-compatibility
-	explicit FShaderPermutationParameters(EShaderPlatform InPlatform, int32 InPermutationId = 0, EShaderPermutationFlags InFlags = EShaderPermutationFlags::HasEditorOnlyData)
-		: Platform(InPlatform)
-		, PermutationId(InPermutationId)
-		, Flags(InFlags)
-	{}
-};
-
 namespace Freeze
 {
 	RENDERCORE_API void IntrinsicToString(const TIndexedPtr<FShaderType>& Object, const FTypeLayoutDesc& TypeDesc, const FPlatformTypeLayoutParameters& LayoutParams, FMemoryToStringContext& OutContext);
@@ -1297,6 +1271,7 @@ public:
 	typedef void (*ModifyCompilationEnvironmentType)(const FShaderPermutationParameters&, FShaderCompilerEnvironment&);
 	typedef bool (*ValidateCompiledResultType)(EShaderPlatform, const FShaderParameterMap&, TArray<FString>&);
 	typedef EShaderCompileJobPriority(*GetOverrideJobPriorityType)();
+	typedef void (*GetPermutationIdStringType)(int32 PermutationId, FString&, bool);
 #endif // WITH_EDITOR
 
 	/** @return The global shader factory list. */
@@ -1338,6 +1313,10 @@ public:
 #endif // WITH_EDITOR
 		uint32 InTypeSize,
 		const FShaderParametersMetadata* InRootParametersMetadata
+#if WITH_EDITOR
+		, GetPermutationIdStringType InGetPermutationIdStringRef = nullptr
+#endif
+
 	);
 
 	RENDERCORE_API virtual ~FShaderType();
@@ -1363,6 +1342,13 @@ public:
 
 	/** Returns a custom compile job priority or EShaderCompileJobPriority::None to ignore the customization. */
 	RENDERCORE_API EShaderCompileJobPriority GetOverrideJobPriority() const;
+
+	/**
+	 * Returns a string representation for the specified permutation Id.
+	 * @param PermutationId Zero based index in the half-open range [0, GetPermutationCount()).
+	 * @param bFullNames Specifies whether the full shader defines are meant to be included. If false, a short form is returned as comma-separated integers.
+	 */
+	RENDERCORE_API FString GetPermutationIdString(int32 PermutationId, bool bFullNames) const;
 #endif // WITH_EDITOR
 
 	RENDERCORE_API ERayTracingPayloadType GetRayTracingPayloadType(const int32 PermutationId) const;
@@ -1573,6 +1559,7 @@ private:
 	ModifyCompilationEnvironmentType ModifyCompilationEnvironmentRef;
 	ValidateCompiledResultType ValidateCompiledResultRef;
 	GetOverrideJobPriorityType GetOverrideJobPriorityRef;
+	GetPermutationIdStringType GetPermutationIdStringRef;
 #endif
 	const FShaderParametersMetadata* const RootParametersMetadata;
 
@@ -1654,6 +1641,11 @@ struct FShaderCompiledShaderInitializerType
 			const typename ShaderClass::FPermutationDomain PermutationVector(Parameters.PermutationId); \
 			PermutationVector.ModifyCompilationEnvironment(OutEnvironment); \
 			ShaderClass::ModifyCompilationEnvironment(static_cast<const typename ShaderClass::FPermutationParameters&>(Parameters), OutEnvironment); \
+		} \
+		static void GetPermutationIdStringImpl(int32 PermutationId, FString& OutString, bool bFullNames) \
+		{ \
+			const typename ShaderClass::FPermutationDomain PermutationVector(PermutationId); \
+			UE::ShaderPermutationUtils::FormatPermutationDomain(PermutationVector, OutString, bFullNames, nullptr); \
 		}
 #else
 	#define SHADER_DECLARE_EDITOR_VTABLE(ShaderClass)
@@ -1702,12 +1694,20 @@ struct FShaderCompiledShaderInitializerType
 	public:
 
 #if WITH_EDITOR
+
 #define SHADER_TYPE_EDITOR_VTABLE(ShaderClass) \
 	, ShaderClass::ModifyCompilationEnvironmentImpl \
 	, ShaderClass::ValidateCompiledResult \
 	, ShaderClass::GetOverrideJobPriority
+
+#define SHADER_TYPE_EDITOR_PERMUTATION_METADATA(ShaderClass) \
+	, ShaderClass::GetPermutationIdStringImpl
+
 #else
+
 #define SHADER_TYPE_EDITOR_VTABLE(ShaderClass)
+#define SHADER_TYPE_EDITOR_PERMUTATION_METADATA(ShaderClass)
+
 #endif
 
 #define SHADER_TYPE_VTABLE(ShaderClass) \
@@ -1736,6 +1736,7 @@ struct FShaderCompiledShaderInitializerType
 			SHADER_TYPE_VTABLE(ShaderClass), \
 			sizeof(ShaderClass), \
 			ShaderClass::GetRootParametersMetadata() \
+			SHADER_TYPE_EDITOR_PERMUTATION_METADATA(ShaderClass) \
 		); \
 		return StaticType; \
 	} \
@@ -1757,6 +1758,7 @@ struct FShaderCompiledShaderInitializerType
 			SHADER_TYPE_VTABLE(ShaderClass), \
 			sizeof(ShaderClass), \
 			ShaderClass::GetRootParametersMetadata() \
+			SHADER_TYPE_EDITOR_PERMUTATION_METADATA(ShaderClass) \
 		); \
 		return StaticType; \
 	} \
@@ -1778,6 +1780,7 @@ struct FShaderCompiledShaderInitializerType
 			SHADER_TYPE_VTABLE(ShaderClass), \
 			sizeof(ShaderClass), \
 			ShaderClass::GetRootParametersMetadata() \
+			SHADER_TYPE_EDITOR_PERMUTATION_METADATA(ShaderClass) \
 		); \
 		return StaticType; \
 	} \
@@ -1801,6 +1804,7 @@ struct FShaderCompiledShaderInitializerType
 			SHADER_TYPE_VTABLE(ShaderClass), \
 			sizeof(ShaderClass), \
 			ShaderClass::GetRootParametersMetadata() \
+			SHADER_TYPE_EDITOR_PERMUTATION_METADATA(ShaderClass) \
 		); \
 		return StaticType; \
 	} \
@@ -1822,11 +1826,106 @@ struct FShaderCompiledShaderInitializerType
 			SHADER_TYPE_VTABLE(ShaderClass), \
 			sizeof(ShaderClass), \
 			ShaderClass::GetRootParametersMetadata() \
+			SHADER_TYPE_EDITOR_PERMUTATION_METADATA(ShaderClass) \
 		); \
 		return StaticType; \
 	} \
 	TemplatePrefix FShaderTypeRegistration ShaderClass::ShaderTypeRegistration{TFunctionRef<::FShaderType&()>{ShaderClass::GetStaticType}};
 
+#if WITH_EDITOR
+
+#define IMPLEMENT_SHADER_TYPE_CONSTRUCTOR(ShaderClass, InShaderTypeForDynamicCast, InShaderClassUserFriendlyName) \
+	ShaderClass( \
+		FTypeLayoutDesc& InTypeLayout, \
+		const TCHAR* InName, \
+		const TCHAR* InSourceFilename, \
+		const TCHAR* InFunctionName, \
+		uint32 InFrequency, \
+		int32 InTotalPermutationCount, \
+		ConstructSerializedType InConstructSerializedRef, \
+		ConstructCompiledType InConstructCompiledRef, \
+		ShouldCompilePermutationType InShouldCompilePermutationRef, \
+		ShouldPrecachePermutationType InShouldPrecachePermutationRef, \
+		GetRayTracingPayloadTypeType InGetRayTracingPayloadTypeRef, \
+		GetShaderBindingLayoutType InGetShaderBindingLayoutTypeRef, \
+		ModifyCompilationEnvironmentType InModifyCompilationEnvironmentRef, \
+		ValidateCompiledResultType InValidateCompiledResultRef, \
+		GetOverrideJobPriorityType InGetOverrideJobPriorityRef, \
+		uint32 InTypeSize, \
+		const FShaderParametersMetadata* InRootParametersMetadata = nullptr, \
+		GetPermutationIdStringType InGetPermutationIdStringRef = nullptr \
+		): \
+		FShaderType( \
+			InShaderTypeForDynamicCast, \
+			InTypeLayout, \
+			InName, \
+			InSourceFilename, \
+			InFunctionName, \
+			InFrequency, \
+			InTotalPermutationCount, \
+			InConstructSerializedRef, \
+			InConstructCompiledRef, \
+			InShouldCompilePermutationRef, \
+			InShouldPrecachePermutationRef, \
+			InGetRayTracingPayloadTypeRef, \
+			InGetShaderBindingLayoutTypeRef, \
+			InModifyCompilationEnvironmentRef, \
+			InValidateCompiledResultRef, \
+			InGetOverrideJobPriorityRef, \
+			InTypeSize, \
+			InRootParametersMetadata, \
+			InGetPermutationIdStringRef \
+		) \
+	{ \
+		checkf(FPaths::GetExtension(InSourceFilename) == TEXT("usf"), \
+			TEXT("Incorrect virtual shader path extension for %s '%s': Only .usf files should be compiled."), \
+			InShaderClassUserFriendlyName, InSourceFilename); \
+	}
+
+#else
+
+#define IMPLEMENT_SHADER_TYPE_CONSTRUCTOR(ShaderClass, InShaderTypeForDynamicCast, InShaderClassUserFriendlyName) \
+	ShaderClass( \
+		FTypeLayoutDesc& InTypeLayout, \
+		const TCHAR* InName, \
+		const TCHAR* InSourceFilename, \
+		const TCHAR* InFunctionName, \
+		uint32 InFrequency, \
+		int32 InTotalPermutationCount, \
+		ConstructSerializedType InConstructSerializedRef, \
+		ConstructCompiledType InConstructCompiledRef, \
+		ShouldCompilePermutationType InShouldCompilePermutationRef, \
+		ShouldPrecachePermutationType InShouldPrecachePermutationRef, \
+		GetRayTracingPayloadTypeType InGetRayTracingPayloadTypeRef, \
+		GetShaderBindingLayoutType InGetShaderBindingLayoutTypeRef, \
+		uint32 InTypeSize, \
+		const FShaderParametersMetadata* InRootParametersMetadata = nullptr \
+		): \
+		FShaderType( \
+			InShaderTypeForDynamicCast, \
+			InTypeLayout, \
+			InName, \
+			InSourceFilename, \
+			InFunctionName, \
+			InFrequency, \
+			InTotalPermutationCount, \
+			InConstructSerializedRef, \
+			InConstructCompiledRef, \
+			InShouldCompilePermutationRef, \
+			InShouldPrecachePermutationRef, \
+			InGetRayTracingPayloadTypeRef, \
+			InGetShaderBindingLayoutTypeRef, \
+			InTypeSize, \
+			InRootParametersMetadata \
+		) \
+	{ \
+		checkf(FPaths::GetExtension(InSourceFilename) == TEXT("usf"), \
+			TEXT("Incorrect virtual shader path extension for %s '%s': Only .usf files should be compiled."), \
+			InShaderClassUserFriendlyName, InSourceFilename); \
+	}
+
+#endif // WITH_EDITOR
+
 // Binding of a set of shader stages in a single pipeline
 class FShaderPipelineType
 {
diff --git a/Engine/Source/Runtime/RenderCore/Public/ShaderPermutation.h b/Engine/Source/Runtime/RenderCore/Public/ShaderPermutation.h
--- a/Engine/Source/Runtime/RenderCore/Public/ShaderPermutation.h
+++ b/Engine/Source/Runtime/RenderCore/Public/ShaderPermutation.h
@@ -14,6 +14,34 @@ inline constexpr bool bLogPermutations = false;
 
 struct FShaderCompilerEnvironment;
 
+// Flags that can specialize shader permutations compiled for specific platforms
+enum class EShaderPermutationFlags : uint32
+{
+	None = 0u,
+	HasEditorOnlyData = (1u << 0),
+};
+ENUM_CLASS_FLAGS(EShaderPermutationFlags);
+
+struct FShaderPermutationParameters
+{
+	// Shader platform to compile to.
+	const EShaderPlatform Platform;
+
+	// Unique permutation identifier of the material shader type.
+	const int32 PermutationId;
+
+	// Flags that describe the permutation
+	const EShaderPermutationFlags Flags;
+
+	// Default to include editor-only shaders, to maintain backwards-compatibility
+	explicit FShaderPermutationParameters(EShaderPlatform InPlatform, int32 InPermutationId = 0, EShaderPermutationFlags InFlags = EShaderPermutationFlags::HasEditorOnlyData)
+		: Platform(InPlatform)
+		, PermutationId(InPermutationId)
+		, Flags(InFlags)
+	{
+	}
+};
+
 /** Defines at compile time a boolean permutation dimension. */
 struct FShaderPermutationBool
 {
@@ -273,7 +301,7 @@ struct TShaderPermutationDomain
 
 // C++11 doesn't allow partial specialization of templates method or function. So we spetialise class that have
 // non spetialised static method, but leave templated static function.
-template<bool BooleanSpetialization>
+template<bool BooleanSpecialization>
 class TShaderPermutationDomainSpetialization
 {
 public:
@@ -380,6 +408,13 @@ struct TShaderPermutationDomain<TDimension, Ts...>
 	}
 
 
+	/** Get the tail of the dimensions. */
+	FORCEINLINE const typename Super::Type& GetTail() const
+	{
+		return Tail;
+	}
+
+
 	/** Modify the shader's compilation environment. */
 	void ModifyCompilationEnvironment(FShaderCompilerEnvironment& OutEnvironment) const
 	{
@@ -419,7 +454,7 @@ struct TShaderPermutationDomain<TDimension, Ts...>
 	}
 
 private:
-	template<bool BooleanSpetialization>
+	template<bool BooleanSpecialization>
 	friend class TShaderPermutationDomainSpetialization;
 
 	typename TDimension::Type DimensionValue;
@@ -483,4 +518,4 @@ using FShaderPermutationNone = TShaderPermutationDomain<>;
  * class FMyShaderDim : SHADER_PERMUTATION_ENUM_CLASS("MY_SHADER_DEFINE_NAME", EMyEnum);
  */
 #define SHADER_PERMUTATION_ENUM_CLASS(InDefineName, EnumName) \
-	DECLARE_SHADER_PERMUTATION_IMPL(InDefineName, TShaderPermutationInt, EnumName, static_cast<int32>(EnumName::MAX))
\ No newline at end of file
+	DECLARE_SHADER_PERMUTATION_IMPL(InDefineName, TShaderPermutationInt, EnumName, static_cast<int32>(EnumName::MAX))
diff --git a/Engine/Source/Runtime/RenderCore/Public/ShaderPermutationUtils.h b/Engine/Source/Runtime/RenderCore/Public/ShaderPermutationUtils.h
--- a/Engine/Source/Runtime/RenderCore/Public/ShaderPermutationUtils.h
+++ b/Engine/Source/Runtime/RenderCore/Public/ShaderPermutationUtils.h
@@ -4,7 +4,7 @@
 
 #include "DataDrivenShaderPlatformInfo.h"
 #include "RHIGlobals.h"
-#include "Shader.h"
+#include "ShaderPermutation.h"
 
 namespace UE::ShaderPermutationUtils
 {
@@ -39,4 +39,86 @@ namespace UE::ShaderPermutationUtils
 
 		return true;
 	}
+
+	template <typename TDimension>
+	void FormatPermutationParameter(const typename TDimension::Type& E, FString& OutString, bool bFullNames, const TCHAR* Prefix);
+	
+	FORCEINLINE void FormatPermutationDomain(const TShaderPermutationDomain<>& InShaderPermutationDomain, FString& OutString, bool bFullNames, const TCHAR* Prefix)
+	{
+		// Dummy
+	}
+
+	template <typename TDimension, typename... Ts>
+	void FormatPermutationDomain(const TShaderPermutationDomain<TDimension, Ts...>& InShaderPermutationDomain, FString& OutString, bool bFullNames, const TCHAR* Prefix)
+	{
+		const int32 Value = TDimension::ToDimensionValueId(InShaderPermutationDomain.template Get<TDimension>());
+		if (bFullNames)
+		{
+			if (!OutString.IsEmpty())
+			{
+				OutString.AppendChar(TEXT('\n'));
+			}
+
+			if (Prefix != nullptr)
+			{
+				OutString.Append(Prefix);
+			}
+
+			if (TDimension::IsMultiDimensional)
+			{
+				OutString.Appendf(TEXT("TShaderPermutationDomain[%d]"), TDimension::PermutationCount);
+			}
+
+			if (Prefix != nullptr)
+			{
+				const FString ExtendedPrefix = FString::Printf(TEXT("%s -> "), Prefix);
+				FormatPermutationParameter<TDimension>(InShaderPermutationDomain.template Get<TDimension>(), OutString, bFullNames, *ExtendedPrefix);
+			}
+			else
+			{
+				FormatPermutationParameter<TDimension>(InShaderPermutationDomain.template Get<TDimension>(), OutString, bFullNames, TEXT(" -> "));
+			}
+		}
+		else
+		{
+			if (!OutString.IsEmpty())
+			{
+				OutString.AppendChar(TEXT(','));
+			}
+			FormatPermutationParameter<TDimension>(InShaderPermutationDomain.template Get<TDimension>(), OutString, bFullNames, TEXT(" -> "));
+		}
+
+		FormatPermutationDomain(InShaderPermutationDomain.GetTail(), OutString, bFullNames, Prefix);
+	}
+
+	template <typename TDimension>
+	void FormatPermutationParameter(const typename TDimension::Type& E, FString& OutString, bool bFullNames, const TCHAR* Prefix)
+	{
+		if constexpr (TDimension::IsMultiDimensional)
+		{
+			FormatPermutationDomain(E, OutString, bFullNames, Prefix);
+		}
+		else if constexpr (std::is_base_of_v<FShaderPermutationBool, TDimension>)
+		{
+			if (bFullNames)
+			{
+				OutString.Appendf(TEXT("%s (%s)"), TDimension::DefineName, !FShaderPermutationBool::ToDefineValue(E) ? TEXT("false") : TEXT("true"));
+			}
+			else
+			{
+				OutString.Appendf(TEXT("%d"), static_cast<int32>(FShaderPermutationBool::ToDefineValue(E)));
+			}
+		}
+		else
+		{
+			if (bFullNames)
+			{
+				OutString.Appendf(TEXT("%s (%d)"), TDimension::DefineName, static_cast<int32>(TDimension::ToDefineValue(E)));
+			}
+			else
+			{
+				OutString.Appendf(TEXT("%d"), static_cast<int32>(TDimension::ToDefineValue(E)));
+			}
+		}
+	}
 }
