Add support for an optionally generated VertexFactoryFwd.ush file included by the MaterialTemplate which allows shader expressions to be defined by vertex factories and invoked from within the material graph
#tests compiled all shaders (preflight) and used a new vertex factory which used the new Fwd file and external code declarations for VF specific material expressions

diff --git a/Engine/Shaders/Private/MaterialTemplate.ush b/Engine/Shaders/Private/MaterialTemplate.ush
--- a/Engine/Shaders/Private/MaterialTemplate.ush
+++ b/Engine/Shaders/Private/MaterialTemplate.ush
@@ -12,6 +12,14 @@
 #define UE_DF_FORCE_FP32_OPS 1
 #endif
 
+#ifndef USE_VERTEX_FACTORY_FWD
+#define USE_VERTEX_FACTORY_FWD 0
+#endif
+
+#if USE_VERTEX_FACTORY_FWD
+#include "/Engine/Generated/VertexFactoryFwd.ush"
+#endif
+
 #include "/Engine/Private/WorldSpaceMath.ush"
 #include "/Engine/Private/SceneTexturesCommon.ush"
 #include "/Engine/Private/EyeAdaptationCommon.ush"
diff --git a/Engine/Shaders/Private/VertexFactoryDefaultFwd.ush b/Engine/Shaders/Private/VertexFactoryDefaultFwd.ush
new file mode 100644
--- /dev/null
+++ b/Engine/Shaders/Private/VertexFactoryDefaultFwd.ush
@@ -0,0 +1,3 @@
+// Copyright Epic Games, Inc. All Rights Reserved.
+
+// default/empty implementation of VertexFactoryFwd
\ No newline at end of file
diff --git a/Engine/Source/Runtime/Engine/Public/LocalVertexFactory.h b/Engine/Source/Runtime/Engine/Public/LocalVertexFactory.h
--- a/Engine/Source/Runtime/Engine/Public/LocalVertexFactory.h
+++ b/Engine/Source/Runtime/Engine/Public/LocalVertexFactory.h
@@ -207,7 +207,7 @@ protected:
 /**
  * Shader parameters for all LocalVertexFactory derived classes.
  */
-class FLocalVertexFactoryShaderParametersBase : public FVertexFactoryShaderParameters
+class ENGINE_API FLocalVertexFactoryShaderParametersBase : public FVertexFactoryShaderParameters
 {
 	DECLARE_TYPE_LAYOUT(FLocalVertexFactoryShaderParametersBase, NonVirtual);
 public:
diff --git a/Engine/Source/Runtime/RenderCore/Private/ShaderCore.cpp b/Engine/Source/Runtime/RenderCore/Private/ShaderCore.cpp
--- a/Engine/Source/Runtime/RenderCore/Private/ShaderCore.cpp
+++ b/Engine/Source/Runtime/RenderCore/Private/ShaderCore.cpp
@@ -1983,6 +1983,10 @@ void GetAllVirtualShaderSourcePaths(TArray<FString>& OutVirtualFilePaths, EShade
 		{
 			FString ShaderFilename(VertexFactoryType->GetShaderFilename());
 			AddShaderSourceFileEntry(OutVirtualFilePaths, ShaderFilename, ShaderPlatform, ShaderPlatformName);
+			if (VertexFactoryType->IncludesFwdShaderFile())
+			{
+				AddShaderSourceFileEntry(OutVirtualFilePaths, VertexFactoryType->GetShaderFwdFilename(), ShaderPlatform, ShaderPlatformName);
+			}
 		}
 	}
 	for( TLinkedList<FShaderType*>::TIterator ShaderIt(FShaderType::GetTypeList()); ShaderIt; ShaderIt.Next() )
diff --git a/Engine/Source/Runtime/RenderCore/Private/ShaderDiagnostics.cpp b/Engine/Source/Runtime/RenderCore/Private/ShaderDiagnostics.cpp
--- a/Engine/Source/Runtime/RenderCore/Private/ShaderDiagnostics.cpp
+++ b/Engine/Source/Runtime/RenderCore/Private/ShaderDiagnostics.cpp
@@ -70,6 +70,13 @@ int32 FShaderDiagnosticInfo::AddAndProcessErrorsForFailedJobFiltered(FShaderComp
 				check(CurrentJob.Key.VFType);
 				CurrentError.ErrorVirtualFilePath = FString(CurrentJob.Key.VFType->GetShaderFilename());
 			}
+			else if (CurrentError.ErrorVirtualFilePath == TEXT("/Engine/Generated/VertexFactoryFwd.ush"))
+			{
+				// VertexFactoryFwd.usf is dynamically included from whichever vertex factory the shader was compiled with.
+				check(CurrentJob.Key.VFType);
+				check(CurrentJob.Key.VFType->IncludesFwdShaderFile());
+				CurrentError.ErrorVirtualFilePath = FString(CurrentJob.Key.VFType->GetShaderFwdFilename());
+			}
 			else if (CurrentError.ErrorVirtualFilePath == TEXT("") && CurrentJob.Key.ShaderType)
 			{
 				// Some shader compiler errors won't have a file and line number, so we just assume the error happened in file containing the entrypoint function.
diff --git a/Engine/Source/Runtime/RenderCore/Private/VertexFactory.cpp b/Engine/Source/Runtime/RenderCore/Private/VertexFactory.cpp
--- a/Engine/Source/Runtime/RenderCore/Private/VertexFactory.cpp
+++ b/Engine/Source/Runtime/RenderCore/Private/VertexFactory.cpp
@@ -74,6 +74,10 @@ void FVertexFactoryType::UpdateReferencedUniformBufferNames(const TMap<FString,
 {
 	ReferencedUniformBuffers.Empty();
 	GenerateReferencedUniformBuffers(ShaderFilename, Name, ShaderFileToUniformBufferVariables, ReferencedUniformBuffers);
+	if (IncludesFwdShaderFile())
+	{
+		GenerateReferencedUniformBuffers(GetShaderFwdFilename(), Name, ShaderFileToUniformBufferVariables, ReferencedUniformBuffers);
+	}
 }
 #endif
 
@@ -132,6 +136,12 @@ FVertexFactoryType::FVertexFactoryType(
 		TEXT("Incorrect virtual shader path extension for vertex factory shader header '%s': Only .ush files should be included."),
 		InShaderFilename);
 
+	if (IncludesFwdShaderFile())
+	{
+		ShaderFwdFilename = FPaths::GetBaseFilename(InShaderFilename, false) + TEXT("Fwd.ush");
+		check(CheckVirtualShaderFilePath(ShaderFwdFilename));
+	}
+
 	// This will trigger if an IMPLEMENT_VERTEX_FACTORY_TYPE was in a module not loaded before InitializeShaderTypes
 	// Vertex factory types need to be implemented in modules that are loaded before that
 	checkf(!bInitializedSerializationHistory, TEXT("VF type was loaded after engine init, use ELoadingPhase::PostConfigInit on your module to cause it to load earlier."));
diff --git a/Engine/Source/Runtime/RenderCore/Public/VertexFactory.h b/Engine/Source/Runtime/RenderCore/Public/VertexFactory.h
--- a/Engine/Source/Runtime/RenderCore/Public/VertexFactory.h
+++ b/Engine/Source/Runtime/RenderCore/Public/VertexFactory.h
@@ -154,6 +154,7 @@ enum class EVertexFactoryFlags : uint32
 	SupportsComputeShading                = 1u << 17,
 	SupportsLumenMeshCards                = 1u << 18,
 	SupportsLandscape                     = 1u << 19,
+	IncludesFwdShader                     = 1u << 20,
 };
 ENUM_CLASS_FLAGS(EVertexFactoryFlags);
 
@@ -372,6 +373,7 @@ public:
 	FName GetFName() const { return TypeName; }
 	const FHashedName& GetHashedName() const { return HashedName; }
 	const TCHAR* GetShaderFilename() const { return ShaderFilename; }
+	const TCHAR* GetShaderFwdFilename() const { return *ShaderFwdFilename; }
 
 	FVertexFactoryShaderParameters* CreateShaderParameters(EShaderFrequency ShaderFrequency, const class FShaderParameterMap& ParameterMap) const { return (*ConstructParameters)(ShaderFrequency, ParameterMap); }
 	const FTypeLayoutDesc* GetShaderParameterLayout(EShaderFrequency ShaderFrequency) const { return (*GetParameterTypeLayout)(ShaderFrequency); }
@@ -417,6 +419,7 @@ public:
 	bool SupportsComputeShading() const                { return HasFlags(EVertexFactoryFlags::SupportsComputeShading); }
 	bool SupportsLumenMeshCards() const                { return HasFlags(EVertexFactoryFlags::SupportsLumenMeshCards); }
 	bool SupportsLandscape() const                     { return HasFlags(EVertexFactoryFlags::SupportsLandscape); }
+	bool IncludesFwdShaderFile() const                 { return HasFlags(EVertexFactoryFlags::IncludesFwdShader); }
 
 	bool SupportsManualVertexFetch(ERHIFeatureLevel::Type InFeatureLevel) const 
 	{
@@ -453,6 +456,17 @@ public:
 		FString VertexFactoryIncludeString = FString::Printf( TEXT("#include \"%s\""), GetShaderFilename() );
 		OutEnvironment.IncludeVirtualPathToContentsMap.Add(TEXT("/Engine/Generated/VertexFactory.ush"), VertexFactoryIncludeString);
 
+		if (IncludesFwdShaderFile())
+		{
+			FString VertexFactoryFwdIncludeString = FString::Printf(TEXT("#include \"%s\""), GetShaderFwdFilename());
+			OutEnvironment.IncludeVirtualPathToContentsMap.Add(TEXT("/Engine/Generated/VertexFactoryFwd.ush"), VertexFactoryFwdIncludeString);
+			OutEnvironment.SetDefine(TEXT("USE_VERTEX_FACTORY_FWD"), 1);
+		}
+		else
+		{
+			OutEnvironment.IncludeVirtualPathToContentsMap.Add(TEXT("/Engine/Generated/VertexFactoryFwd.ush"), TEXT("#include \"/Engine/Private/VertexFactoryDefaultFwd.ush\""));
+		}
+
 		OutEnvironment.SetDefine(TEXT("HAS_PRIMITIVE_UNIFORM_BUFFER"), 1);
 
 		(*ModifyCompilationEnvironmentRef)(Parameters, OutEnvironment);
@@ -489,6 +503,7 @@ private:
 
 	const TCHAR* Name;
 	const TCHAR* ShaderFilename;
+	FString ShaderFwdFilename;
 	FName TypeName;
 	FHashedName HashedName;
 	EVertexFactoryFlags Flags;
