Caching empty entries for missing shader files during initialization to avoid thread contention.

#rnx

diff --git a/Engine/Source/Runtime/RenderCore/Private/ShaderCore.cpp b/Engine/Source/Runtime/RenderCore/Private/ShaderCore.cpp
--- a/Engine/Source/Runtime/RenderCore/Private/ShaderCore.cpp
+++ b/Engine/Source/Runtime/RenderCore/Private/ShaderCore.cpp
@@ -243,12 +243,50 @@ struct FShaderFileCacheEntry
 	FString Source;
 	FShaderSharedAnsiStringPtr StrippedSource;				// Source with comments stripped out, and converted to ANSICHAR
 	FShaderPreprocessDependenciesShared Dependencies;		// Stripped source with include dependencies, all in one shareable struct
+	
+	bool IsEmpty() const
+	{
+		return Source.IsEmpty() && !StrippedSource.IsValid() && !Dependencies.IsValid();
+	}
 };
 
 // Apply lock striping as we're mostly reader lock bound.
 // Use prime number for the number of buckets for best distribution using modulo.
 TStripedMap<31, FString, FShaderFileCacheEntry> GShaderFileCache;
 
+// Enables transient missing shader file caching to avoid write lock contention on parallel loading.
+class FMissingShaderFileCacheGuard
+{
+public:
+	FMissingShaderFileCacheGuard()
+	{
+		++EnabledCount;
+	}
+	
+	~FMissingShaderFileCacheGuard()
+	{
+		if (--EnabledCount == 0)
+		{
+			// Clean up empty entries which were cached for parallelization performance and turn further caching off
+			// allowing for shaders that are dynamically generated to be loaded.
+			GShaderFileCache.RemoveIf([](const TPair<FString, FShaderFileCacheEntry>& Entry)
+			{
+				return Entry.Value.IsEmpty();
+			});
+		}
+	}
+
+	static bool IsEnabled()
+	{
+		return EnabledCount > 0;
+	}
+	
+private:
+	static int EnabledCount;
+};
+
+int FMissingShaderFileCacheGuard::EnabledCount = 0;
+
 class FShaderHashCache
 {
 public:
@@ -2503,6 +2541,7 @@ bool LoadShaderSourceFile(const TCHAR* InVirtualFilePath, EShaderPlatform Shader
 	}
 
 	bool bResult = false;
+	bool bFoundValidFile = false;
 
 	STAT(double ShaderFileLoadingTime = 0);
 
@@ -2528,9 +2567,10 @@ bool LoadShaderSourceFile(const TCHAR* InVirtualFilePath, EShaderPlatform Shader
 						return true;
 					}
 
-					return false;
+					// Create an empty entry if missing files are being cached.
+					return FMissingShaderFileCacheGuard::IsEnabled();
 				},
-				[&OutFileContents, &OutStrippedContents](const FShaderFileCacheEntry& CachedFile)
+				[&bFoundValidFile, &OutFileContents, &OutStrippedContents](const FShaderFileCacheEntry& CachedFile)
 				{
 					if (OutFileContents)
 					{
@@ -2540,10 +2580,20 @@ bool LoadShaderSourceFile(const TCHAR* InVirtualFilePath, EShaderPlatform Shader
 					{
 						*OutStrippedContents = CachedFile.StrippedSource;
 					}
+					if (FMissingShaderFileCacheGuard::IsEnabled())
+					{
+						bFoundValidFile = !CachedFile.IsEmpty();
+					}
 				}
 			);
 	}
 
+	// When caching missing files, the map find will always return true, so the result is based on whether it is an empty entry or not.
+	if (FMissingShaderFileCacheGuard::IsEnabled())
+	{
+		bResult = bFoundValidFile;
+	}
+
 	INC_FLOAT_STAT_BY(STAT_ShaderCompiling_LoadingShaderFiles,(float)ShaderFileLoadingTime);
 
 	return bResult;
@@ -3249,6 +3299,8 @@ void CheckShaderHashCacheInclude(const FString& VirtualFilePath, EShaderPlatform
 void InitializeShaderTypes()
 {
 	UE_LOG(LogShaders, Log, TEXT("InitializeShaderTypes() begin"));
+	
+	FMissingShaderFileCacheGuard Guard;
 
 	LogShaderSourceDirectoryMappings();
 
