Make async RDG builder setup task high priority or it will get pre-empted by other tasks.

diff --git a/Engine/Source/Runtime/RenderCore/Private/RenderGraphBuilder.cpp b/Engine/Source/Runtime/RenderCore/Private/RenderGraphBuilder.cpp
--- a/Engine/Source/Runtime/RenderCore/Private/RenderGraphBuilder.cpp
+++ b/Engine/Source/Runtime/RenderCore/Private/RenderGraphBuilder.cpp
@@ -1661,7 +1661,8 @@ void FRDGBuilder::LaunchAsyncSetupQueueTask()
 		AsyncSetupQueue.Pipe.Launch(UE_SOURCE_LOCATION, [this]() mutable
 		{
 			ProcessAsyncSetupQueue();
-		});
+
+		}, UE::Tasks::ETaskPriority::High);
 	}
 }
 
