[shaders] destruct FShader properly when adding a permutation to a pipeline that has already been added by another completed task.

diff --git a/Engine/Source/Runtime/RenderCore/Private/Shader.cpp b/Engine/Source/Runtime/RenderCore/Private/Shader.cpp
--- a/Engine/Source/Runtime/RenderCore/Private/Shader.cpp
+++ b/Engine/Source/Runtime/RenderCore/Private/Shader.cpp
@@ -1328,7 +1328,10 @@ FShader* FShaderPipeline::FindOrAddShader(FShader* Shader, int32 PermutationId)
 	FShader* PrevShader = Shaders[Frequency];
 	if (PrevShader && PermutationIds[Frequency] == PermutationId)
 	{
-		delete Shader;
+		if (Shader)
+		{
+			DeleteObjectFromLayout(Shader);
+		}
 		return PrevShader;
 	}
 
