remove unnecessary null check, causing SA failure

#rnx

diff --git a/Engine/Source/Runtime/RenderCore/Private/Shader.cpp b/Engine/Source/Runtime/RenderCore/Private/Shader.cpp
--- a/Engine/Source/Runtime/RenderCore/Private/Shader.cpp
+++ b/Engine/Source/Runtime/RenderCore/Private/Shader.cpp
@@ -1328,10 +1328,7 @@ FShader* FShaderPipeline::FindOrAddShader(FShader* Shader, int32 PermutationId)
 	FShader* PrevShader = Shaders[Frequency];
 	if (PrevShader && PermutationIds[Frequency] == PermutationId)
 	{
-		if (Shader)
-		{
-			DeleteObjectFromLayout(Shader);
-		}
+		DeleteObjectFromLayout(Shader);
 		return PrevShader;
 	}
 
