Render target pool memory dump sort by size and name, for increased legibility

#rnx

diff --git a/Engine/Source/Runtime/RenderCore/Private/RenderTargetPool.cpp b/Engine/Source/Runtime/RenderCore/Private/RenderTargetPool.cpp
--- a/Engine/Source/Runtime/RenderCore/Private/RenderTargetPool.cpp
+++ b/Engine/Source/Runtime/RenderCore/Private/RenderTargetPool.cpp
@@ -482,38 +482,69 @@ void FRenderTargetPool::FreeUnusedResources()
 	}
 }
 
+struct FRenderTargetWithSize
+{
+	const FPooledRenderTarget* RenderTarget;
+	uint32 SizeInKB;
+
+	bool operator<(const FRenderTargetWithSize& Other) const
+	{
+		if (SizeInKB != Other.SizeInKB)
+		{
+			// Sort larger sizes first
+			return SizeInKB > Other.SizeInKB;
+		}
+		int32 NameOrdinalCompare = FCString::Strcmp(RenderTarget->GetDesc().DebugName, Other.RenderTarget->GetDesc().DebugName);
+		if (NameOrdinalCompare != 0)
+		{
+			return NameOrdinalCompare < 0;
+		}
+		return RenderTarget < Other.RenderTarget;
+	}
+};
+
 void FRenderTargetPool::DumpMemoryUsage(FOutputDevice& OutputDevice)
 {
 	UE::TScopeLock Lock(Mutex);
 
 	uint32 UnusedAllocationInKB = 0;
 
-	OutputDevice.Logf(TEXT("Pooled Render Targets:"));
+	TArray<FRenderTargetWithSize> SortedRenderTargets;
+	SortedRenderTargets.Reserve(PooledRenderTargets.Num());
+
 	for (int32 i = 0; i < PooledRenderTargets.Num(); ++i)
 	{
 		FPooledRenderTarget* Element = PooledRenderTargets[i];
 
 		if (Element)
 		{
-			uint32 ElementAllocationInKB = ComputeSizeInKB(*Element);
-			if (Element->UnusedForNFrames > 2)
-			{
-				UnusedAllocationInKB += ElementAllocationInKB;
-			}
+			SortedRenderTargets.Add({ Element, ComputeSizeInKB(*Element) });
+		}
+	}
 
-			OutputDevice.Logf(
-				TEXT("  %6.3fMB %4dx%4d%s%s %2dmip(s) %s (%s) Unused frames: %d"),
-				ElementAllocationInKB / 1024.0f,
-				Element->Desc.Extent.X,
-				Element->Desc.Extent.Y,
-				Element->Desc.Depth > 1 ? *FString::Printf(TEXT("x%3d"), Element->Desc.Depth) : (Element->Desc.IsCubemap() ? TEXT("cube") : TEXT("    ")),
-				Element->Desc.bIsArray ? *FString::Printf(TEXT("[%3d]"), Element->Desc.ArraySize) : TEXT("     "),
-				Element->Desc.NumMips,
-				Element->Desc.DebugName,
-				GPixelFormats[Element->Desc.Format].Name,
-				Element->UnusedForNFrames
-			);
+	SortedRenderTargets.Sort();
+
+	OutputDevice.Logf(TEXT("Pooled Render Targets:"));
+	for (FRenderTargetWithSize& SortedRenderTarget : SortedRenderTargets)
+	{
+		if (SortedRenderTarget.RenderTarget->UnusedForNFrames > 2)
+		{
+			UnusedAllocationInKB += SortedRenderTarget.SizeInKB;
 		}
+
+		const FPooledRenderTarget* Element = SortedRenderTarget.RenderTarget;
+		OutputDevice.Logf(
+			TEXT("  %6.3fMB %4dx%4d%s%s %2dmip(s) %s (%s) Unused frames: %d"),
+			SortedRenderTarget.SizeInKB / 1024.0f,
+			Element->Desc.Extent.X,
+			Element->Desc.Extent.Y,
+			Element->Desc.Depth > 1 ? *FString::Printf(TEXT("x%3d"), Element->Desc.Depth) : (Element->Desc.IsCubemap() ? TEXT("cube") : TEXT("    ")),
+			Element->Desc.bIsArray ? *FString::Printf(TEXT("[%3d]"), Element->Desc.ArraySize) : TEXT("     "),
+			Element->Desc.NumMips,
+			Element->Desc.DebugName,
+			GPixelFormats[Element->Desc.Format].Name,
+			Element->UnusedForNFrames
+		);
 	}
 	uint32 NumTargets = 0;
 	uint32 UsedKB = 0;
