Deprecate the stereo Mobile Multi-View Fallback path.

This was a path that allowed us to perform single pass stereo through the mobile shading path when Instanced Stereo and layered RTs were supported by the RHI, but Mobile Multi-View was not. It was generally used in mobile previews. Users that want to test MMV on desktop should instead use Vulkan Editor for native MMV.

Supporting MMV is now gated behind UE_SUPPORT_MMV_FALLBACK, which is 0 by default. If support is disabled and we would previously have entered MMV Fallback, a warning will be logged.

diff --git a/Engine/Source/Runtime/Engine/Private/UnrealEngine.cpp b/Engine/Source/Runtime/Engine/Private/UnrealEngine.cpp
--- a/Engine/Source/Runtime/Engine/Private/UnrealEngine.cpp
+++ b/Engine/Source/Runtime/Engine/Private/UnrealEngine.cpp
@@ -3914,8 +3914,8 @@ public:
 			NumViews = FMath::Clamp(V, 1, 32);
 		}
 
-		// mobile rendering path does not use side-by-side normally
-		bLayeredRTs = GetFeatureLevelShadingPath(GMaxRHIFeatureLevel) == EShadingPath::Mobile;
+		const UE::StereoRenderUtils::FStereoShaderAspects Aspects(GMaxRHIShaderPlatform);
+		bLayeredRTs = Aspects.IsMobileMultiViewEnabled();
 	}
 
 	virtual ~FFakeStereoRenderingDevice() {}
diff --git a/Engine/Source/Runtime/RenderCore/Private/StereoRenderUtils.cpp b/Engine/Source/Runtime/RenderCore/Private/StereoRenderUtils.cpp
--- a/Engine/Source/Runtime/RenderCore/Private/StereoRenderUtils.cpp
+++ b/Engine/Source/Runtime/RenderCore/Private/StereoRenderUtils.cpp
@@ -9,6 +9,7 @@
 
 // enable this to printf-debug stereo rendering aspects on a device
 #define UE_DEBUG_STEREO_ASPECTS			(!UE_BUILD_SHIPPING && !UE_BUILD_TEST && !UE_EDITOR)
+#define UE_SUPPORT_MMV_FALLBACK			0
 
 #if UE_DEBUG_STEREO_ASPECTS
 	#define UE_DEBUG_SSA_LOG_INIT(Platform) \
@@ -97,6 +98,7 @@ RENDERCORE_API FStereoShaderAspects::FStereoShaderAspects(EShaderPlatform Platfo
 
 			if (bVertexShaderLayerCapable)
 			{
+				// Even if MMV fallback support is disabled, we still set it to true internally to log a deprecation warning
 				bMobileMultiViewFallback = true;
 			}
 		}
@@ -114,8 +116,14 @@ RENDERCORE_API FStereoShaderAspects::FStereoShaderAspects(EShaderPlatform Platfo
 
 	bInstancedMultiViewportEnabled = bInstancedStereoNative && bMultiViewportCapable;
 	// Since instanced stereo now relies on multi-viewport capability, it cannot be separately enabled from it.
-	bInstancedStereoEnabled = bInstancedStereoNative || bMobileMultiViewFallback;
-	bMobileMultiViewEnabled = bMobileMultiViewNative || bMobileMultiViewFallback;
+	bInstancedStereoEnabled = bInstancedStereoNative;
+	bMobileMultiViewEnabled = bMobileMultiViewNative;
+
+	// MMV fallback is deprecated as of UE 5.6
+#if UE_SUPPORT_MMV_FALLBACK
+	bInstancedStereoEnabled |= bMobileMultiViewFallback;
+	bMobileMultiViewEnabled |= bMobileMultiViewFallback;
+#endif
 
 	UE_DEBUG_SSA_LOG_BOOL(bInstancedMultiViewportEnabled);
 	UE_DEBUG_SSA_LOG_BOOL(bInstancedStereoEnabled);
@@ -124,8 +132,13 @@ RENDERCORE_API FStereoShaderAspects::FStereoShaderAspects(EShaderPlatform Platfo
 
 	// check the following invariants
 	checkf(!bMobileMultiViewNative || !bInstancedStereoEnabled, TEXT("When a platform supports MMV natively, ISR should not be enabled."));
+
+#if UE_SUPPORT_MMV_FALLBACK
 	checkf(!(bMobileMultiViewEnabled && bInstancedStereoEnabled) || bMobileMultiViewFallback, TEXT("If a platform uses MMV fallback, both MMV and ISR should be enabled."));
 	checkf(!bInstancedStereoEnabled || (bInstancedMultiViewportEnabled || bMobileMultiViewFallback), TEXT("If ISR is enabled, we need either multi-viewport (since we no longer support clip-distance method) or MMV fallback (which uses vertex layer)."));
+#else
+	checkf(!bInstancedStereoEnabled || bInstancedMultiViewportEnabled, TEXT("If ISR is enabled, we need either multi-viewport (since we no longer support clip-distance method)."));
+#endif
 }
 	
 RENDERCORE_API FStereoShaderAspects::FStereoShaderAspects() :
@@ -142,6 +155,17 @@ RENDERCORE_API void LogISRInit(const UE::StereoRenderUtils::FStereoShaderAspects
 	UE_LOG(LogInit, Log, TEXT("XR: Instanced Stereo Rendering is %s"), (Aspects.IsInstancedStereoEnabled() ? TEXT("Enabled") : TEXT("Disabled")));
 	UE_LOG(LogInit, Log, TEXT("XR: MultiViewport is %s"), (Aspects.IsInstancedMultiViewportEnabled() ? TEXT("Enabled") : TEXT("Disabled")));
 	UE_LOG(LogInit, Log, TEXT("XR: Mobile Multiview is %s"), (Aspects.IsMobileMultiViewEnabled() ? TEXT("Enabled") : TEXT("Disabled")));
+
+	PRAGMA_DISABLE_DEPRECATION_WARNINGS
+	if (Aspects.IsMobileMultiViewFallbackEnabled())
+	{
+#if UE_SUPPORT_MMV_FALLBACK
+		UE_LOG(LogInit, Log, TEXT("XR: Using MMV fallback (instanced stereo + mobile shaders via layered RTs)."))
+#else
+		UE_LOG(LogInit, Warning, TEXT("XR: Using mobile shader platform and vr.MobileMultiView=1, but MMV is not supported by the RHI. Falling back to view-by-view rendering. Try using Vulkan editor for native MMV support on desktop."))
+#endif
+	}
+	PRAGMA_ENABLE_DEPRECATION_WARNINGS
 }
 
 RENDERCORE_API void VerifyISRConfig(const UE::StereoRenderUtils::FStereoShaderAspects& Aspects, EShaderPlatform ShaderPlatform)
diff --git a/Engine/Source/Runtime/RenderCore/Public/StereoRenderUtils.h b/Engine/Source/Runtime/RenderCore/Public/StereoRenderUtils.h
--- a/Engine/Source/Runtime/RenderCore/Public/StereoRenderUtils.h
+++ b/Engine/Source/Runtime/RenderCore/Public/StereoRenderUtils.h
@@ -45,6 +45,13 @@ public:
 	*/
 	inline bool IsInstancedMultiViewportEnabled() const { return bInstancedMultiViewportEnabled; };
 
+	/**
+	* Whether MMV fallback was requested - i.e. using ISR-generated view indices to index into texture arrays.
+	* True when on a mobile shader platform and vr.MobileMultiView=1, but ISR is supported by the RHI and MMV is not (e.g. D3D12 mobile preview).
+	*/
+	UE_DEPRECATED(5.6, TEXT("The MMV fallback path is deprecated and will be disabled via UE_SUPPORT_MMV_FALLBACK by default even if this function returns true."))
+	inline bool IsMobileMultiViewFallbackEnabled() const { return bMobileMultiViewFallback; };
+
 private:
 	bool bInstancedStereoEnabled : 1;
 	bool bMobileMultiViewEnabled : 1;
