Fix compiler error warning on Mac due to unneccessary TEXT macro in use of UE_DEPRECATED.

diff --git a/Engine/Source/Runtime/RenderCore/Public/StereoRenderUtils.h b/Engine/Source/Runtime/RenderCore/Public/StereoRenderUtils.h
--- a/Engine/Source/Runtime/RenderCore/Public/StereoRenderUtils.h
+++ b/Engine/Source/Runtime/RenderCore/Public/StereoRenderUtils.h
@@ -49,7 +49,7 @@ public:
 	* Whether MMV fallback was requested - i.e. using ISR-generated view indices to index into texture arrays.
 	* True when on a mobile shader platform and vr.MobileMultiView=1, but ISR is supported by the RHI and MMV is not (e.g. D3D12 mobile preview).
 	*/
-	UE_DEPRECATED(5.6, TEXT("The MMV fallback path is deprecated and will be disabled via UE_SUPPORT_MMV_FALLBACK by default even if this function returns true."))
+	UE_DEPRECATED(5.6, "The MMV fallback path is deprecated and will be disabled via UE_SUPPORT_MMV_FALLBACK by default even if this function returns true.")
 	inline bool IsMobileMultiViewFallbackEnabled() const { return bMobileMultiViewFallback; };
 
 private:
