[shaders] modify cflags enum generation macros to allow deprecating entries, and to only add the CFLAG prefix to the enum members (not the SA string representation); also cleaned up some odd typos in the comments.

#rnx

diff --git a/Engine/Source/Runtime/RenderCore/Private/ShaderCore.cpp b/Engine/Source/Runtime/RenderCore/Private/ShaderCore.cpp
--- a/Engine/Source/Runtime/RenderCore/Private/ShaderCore.cpp
+++ b/Engine/Source/Runtime/RenderCore/Private/ShaderCore.cpp
@@ -4524,9 +4524,8 @@ static const TMap<FString, ECompilerFlags>& GetCompilerFlagsDictionary()
 {
 	static const TMap<FString, ECompilerFlags> CompilerFlagsDictionary =
 	{
-		#define SHADER_COMPILER_FLAGS_ENTRY(NAME) { TEXT(#NAME), NAME },
+		#define SHADER_COMPILER_FLAGS_ENTRY(Name) { TEXT(#Name), CFLAG_##Name },
 		#include "ShaderCompilerFlags.inl"
-		#undef SHADER_COMPILER_FLAGS_ENTRY
 	};
 	return CompilerFlagsDictionary;
 }
diff --git a/Engine/Source/Runtime/RenderCore/Public/ShaderCompilerCore.h b/Engine/Source/Runtime/RenderCore/Public/ShaderCompilerCore.h
--- a/Engine/Source/Runtime/RenderCore/Public/ShaderCompilerCore.h
+++ b/Engine/Source/Runtime/RenderCore/Public/ShaderCompilerCore.h
@@ -84,9 +84,10 @@ extern RENDERCORE_API bool AllowGlobalShaderLoad();
 
 enum ECompilerFlags
 {
-	#define SHADER_COMPILER_FLAGS_ENTRY(NAME) NAME,
+	#define SHADER_COMPILER_FLAGS_ENTRY(Name) CFLAG_##Name,
+	#define SHADER_COMPILER_FLAGS_ENTRY_DEPRECATED(Name, Version, Message) CFLAG_##Name UE_DEPRECATED(Version, Message),
 	#include "ShaderCompilerFlags.inl"
-	#undef SHADER_COMPILER_FLAGS_ENTRY
+	CFLAG_Max
 };
 static_assert(CFLAG_PreferFlowControl == 0, "First entry in ECompilerFlags must be 'CFLAG_PreferFlowControl' and assigned to 0");
 static_assert(CFLAG_Max < 64, "Out of bitfield space! Modify FShaderCompilerFlags");
diff --git a/Engine/Source/Runtime/RenderCore/Public/ShaderCompilerFlags.inl b/Engine/Source/Runtime/RenderCore/Public/ShaderCompilerFlags.inl
--- a/Engine/Source/Runtime/RenderCore/Public/ShaderCompilerFlags.inl
+++ b/Engine/Source/Runtime/RenderCore/Public/ShaderCompilerFlags.inl
@@ -4,98 +4,107 @@
 	ShaderCompilerFlags.inl: Flags used in ECompilerFlags enumeration.
 =============================================================================*/
 
-#ifndef SHADER_COMPILER_FLAGS_ENTRY
+#if !defined(SHADER_COMPILER_FLAGS_ENTRY) 
 #error SHADER_COMPILER_FLAGS_ENTRY must be defined before including ShaderCompilerFlags.inl
 #endif
 
-SHADER_COMPILER_FLAGS_ENTRY( CFLAG_PreferFlowControl )
-SHADER_COMPILER_FLAGS_ENTRY( CFLAG_Debug )
-SHADER_COMPILER_FLAGS_ENTRY( CFLAG_AvoidFlowControl )
+#if !defined(SHADER_COMPILER_FLAGS_ENTRY_DEPRECATED)
+#define SHADER_COMPILER_FLAGS_ENTRY_DEPRECATED(Name, Version, Message)
+#endif
+
+// To add a flag:
+// SHADER_COMPILER_FLAGS_ENTRY(FlagName) - generates CFLAG_FlagName
+// For deprecated flags:
+// SHADER_COMPILER_FLAGS_ENTRY_DEPRECATED(FlagName, X.X, "Deprecation Message") - generates CFLAG_FlagName UE_DEPRECATED(X.X, "Deprecation Message")
+
+SHADER_COMPILER_FLAGS_ENTRY(PreferFlowControl)
+SHADER_COMPILER_FLAGS_ENTRY(Debug)
+SHADER_COMPILER_FLAGS_ENTRY(AvoidFlowControl)
 // Disable shader validation
-SHADER_COMPILER_FLAGS_ENTRY( CFLAG_SkipValidation )
+SHADER_COMPILER_FLAGS_ENTRY(SkipValidation)
 // Only allows standard optimizations, not the longest compile times.
-SHADER_COMPILER_FLAGS_ENTRY( CFLAG_StandardOptimization )
-// Always optimize even when CFLAG_Debug is set. Required for some complex shaders and features.
-SHADER_COMPILER_FLAGS_ENTRY( CFLAG_ForceOptimization )
+SHADER_COMPILER_FLAGS_ENTRY(StandardOptimization)
+// Always optimize even whenDebug is set. Required for some complex shaders and features.
+SHADER_COMPILER_FLAGS_ENTRY(ForceOptimization)
 // Shader should generate symbols for debugging.
-SHADER_COMPILER_FLAGS_ENTRY( CFLAG_GenerateSymbols )
+SHADER_COMPILER_FLAGS_ENTRY(GenerateSymbols)
 // Shader should insert debug/name info at the risk of generating non-deterministic libraries
-SHADER_COMPILER_FLAGS_ENTRY( CFLAG_ExtraShaderData )
+SHADER_COMPILER_FLAGS_ENTRY(ExtraShaderData)
 // Allows the (external) symbols to be specific to each shader rather than trying to deduplicate.
-SHADER_COMPILER_FLAGS_ENTRY( CFLAG_AllowUniqueSymbols )
-SHADER_COMPILER_FLAGS_ENTRY( CFLAG_NoFastMath )
+SHADER_COMPILER_FLAGS_ENTRY(AllowUniqueSymbols)
+SHADER_COMPILER_FLAGS_ENTRY(NoFastMath)
 // Explicitly enforce zero initialization on shader platforms that may omit it.
-SHADER_COMPILER_FLAGS_ENTRY( CFLAG_ZeroInitialise )
+SHADER_COMPILER_FLAGS_ENTRY(ZeroInitialise)
 // Explicitly enforce bounds checking on shader platforms that may omit it.
-SHADER_COMPILER_FLAGS_ENTRY( CFLAG_BoundsChecking )
+SHADER_COMPILER_FLAGS_ENTRY(BoundsChecking)
 // Force removing unused interpolators for platforms that can opt out
-SHADER_COMPILER_FLAGS_ENTRY( CFLAG_ForceRemoveUnusedInterpolators )
+SHADER_COMPILER_FLAGS_ENTRY(ForceRemoveUnusedInterpolators)
 // Hint that it is a vertex to geometry shader
-SHADER_COMPILER_FLAGS_ENTRY( CFLAG_VertexToGeometryShader )
+SHADER_COMPILER_FLAGS_ENTRY(VertexToGeometryShader)
 // Hint that it is a vertex to primitive shader
-SHADER_COMPILER_FLAGS_ENTRY( CFLAG_VertexToPrimitiveShader )
+SHADER_COMPILER_FLAGS_ENTRY(VertexToPrimitiveShader)
 // Hint that a vertex shader should use automatic culling on certain platforms.
-SHADER_COMPILER_FLAGS_ENTRY( CFLAG_VertexUseAutoCulling )
+SHADER_COMPILER_FLAGS_ENTRY(VertexUseAutoCulling)
 // Prepare the shader for archiving in the native binary shader cache format
-SHADER_COMPILER_FLAGS_ENTRY( CFLAG_Archive )
+SHADER_COMPILER_FLAGS_ENTRY(Archive)
 // Shaders uses external texture so may need special runtime handling
-SHADER_COMPILER_FLAGS_ENTRY( CFLAG_UsesExternalTexture )
+SHADER_COMPILER_FLAGS_ENTRY(UsesExternalTexture)
 // Use emulated uniform buffers on supported platforms
-SHADER_COMPILER_FLAGS_ENTRY( CFLAG_UseEmulatedUB )
+SHADER_COMPILER_FLAGS_ENTRY(UseEmulatedUB)
 // Enable wave operation intrinsics (requires DX12 and DXC/DXIL on PC).
 // Check GRHISupportsWaveOperations before using shaders compiled with this flag at runtime.
 // [REDACTED_PROFILE_URL]/DirectXShaderCompiler/wiki/Wave-Intrinsics
-SHADER_COMPILER_FLAGS_ENTRY( CFLAG_WaveOperations )
-// Use DirectX Shader Compiler (DXC) to compile all shaders ) intended for compatibility testing.
-SHADER_COMPILER_FLAGS_ENTRY( CFLAG_ForceDXC )
-SHADER_COMPILER_FLAGS_ENTRY( CFLAG_SkipOptimizations )
-// Temporarily disable optimizations with DXC compiler only ) intended to workaround shader compiler bugs until they can be resolved with 1st party
-SHADER_COMPILER_FLAGS_ENTRY( CFLAG_SkipOptimizationsDXC )
-// Typed UAV loads are disallowed by default ) as Windows 7 D3D 11.0 does not support them; this flag allows a shader to use them.
-SHADER_COMPILER_FLAGS_ENTRY( CFLAG_AllowTypedUAVLoads )
+SHADER_COMPILER_FLAGS_ENTRY(WaveOperations)
+// Use DirectX Shader Compiler (DXC) to compile all shaders - intended for compatibility testing.
+SHADER_COMPILER_FLAGS_ENTRY(ForceDXC)
+SHADER_COMPILER_FLAGS_ENTRY(SkipOptimizations)
+// Temporarily disable optimizations with DXC compiler only - intended to workaround shader compiler bugs until they can be resolved with 1st party
+SHADER_COMPILER_FLAGS_ENTRY(SkipOptimizationsDXC)
+// Typed UAV loads are disallowed by default as Windows 7 D3D 11.0 does not support them; this flag allows a shader to use them.
+SHADER_COMPILER_FLAGS_ENTRY(AllowTypedUAVLoads)
 // Prefer shader execution in wave32 mode if possible.
-SHADER_COMPILER_FLAGS_ENTRY( CFLAG_Wave32 )
+SHADER_COMPILER_FLAGS_ENTRY(Wave32)
 // Enable support of inline raytracing in compute shader.
-SHADER_COMPILER_FLAGS_ENTRY( CFLAG_InlineRayTracing )
+SHADER_COMPILER_FLAGS_ENTRY(InlineRayTracing)
 // Enable support of C-style data types for platforms that can. Check for PLATFORM_SUPPORTS_REAL_TYPES and FDataDrivenShaderPlatformInfo::GetSupportsRealTypes()
-SHADER_COMPILER_FLAGS_ENTRY( CFLAG_AllowRealTypes )
-// Precompile HLSL to optimized HLSL ) then forward to FXC. Speeds up some shaders that take longer with FXC and works around crashes in FXC.
-SHADER_COMPILER_FLAGS_ENTRY( CFLAG_PrecompileWithDXC )
-// Enable HLSL 2021 version. Enables templates ) operator overloaing ) C++ style function overloading. Contains breaking change with short-circuiting evaluation.
-SHADER_COMPILER_FLAGS_ENTRY( CFLAG_HLSL2021 )
+SHADER_COMPILER_FLAGS_ENTRY(AllowRealTypes)
+// Precompile HLSL to optimized HLSL then forward to FXC. Speeds up some shaders that take longer with FXC and works around crashes in FXC.
+SHADER_COMPILER_FLAGS_ENTRY(PrecompileWithDXC)
+// Enable HLSL 2021 version. Enables templates, operator overloading, and C++ style function overloading. Contains breaking change with short-circuiting evaluation.
+SHADER_COMPILER_FLAGS_ENTRY(HLSL2021)
 // Allow warnings to be treated as errors
-SHADER_COMPILER_FLAGS_ENTRY( CFLAG_WarningsAsErrors )
+SHADER_COMPILER_FLAGS_ENTRY(WarningsAsErrors)
 // Enabled if bindless resources are enabled for the platform
-SHADER_COMPILER_FLAGS_ENTRY( CFLAG_BindlessResources )
+SHADER_COMPILER_FLAGS_ENTRY(BindlessResources)
 // Enabled if bindless samplers are enabled for the platform
-SHADER_COMPILER_FLAGS_ENTRY( CFLAG_BindlessSamplers )
+SHADER_COMPILER_FLAGS_ENTRY(BindlessSamplers)
 // Force disable bindless resources and samplers on a shader
-SHADER_COMPILER_FLAGS_ENTRY( CFLAG_ForceBindful )
+SHADER_COMPILER_FLAGS_ENTRY(ForceBindful)
 // EXPERIMENTAL: Run the shader re-writer that removes any unused functions/resources/types from source code before compilation.
-SHADER_COMPILER_FLAGS_ENTRY( CFLAG_RemoveDeadCode )
+SHADER_COMPILER_FLAGS_ENTRY(RemoveDeadCode)
 // Enable CullBeforeFetch optimization on supported platforms
-SHADER_COMPILER_FLAGS_ENTRY( CFLAG_CullBeforeFetch )
+SHADER_COMPILER_FLAGS_ENTRY(CullBeforeFetch)
 // Enable WarpCulling optimization on supported platforms
-SHADER_COMPILER_FLAGS_ENTRY( CFLAG_WarpCulling )
+SHADER_COMPILER_FLAGS_ENTRY(WarpCulling)
 // Shader should generate minimal symbols info
-SHADER_COMPILER_FLAGS_ENTRY( CFLAG_GenerateSymbolsInfo )
+SHADER_COMPILER_FLAGS_ENTRY(GenerateSymbolsInfo)
 // Enabled root constants optimization on supported platforms
-SHADER_COMPILER_FLAGS_ENTRY( CFLAG_RootConstants )
-// Specifies that a shader provides derivatives ) and the compiler should look in the compiled ISA for any instructions requiring
-// auto derivatives. If none are found ) the shader will be marked with EShaderResourceUsageFlags::NoDerivativeOps ) meaning that
+SHADER_COMPILER_FLAGS_ENTRY(RootConstants)
+// Specifies that a shader provides derivatives, and the compiler should look in the compiled ISA for any instructions requiring
+// auto derivatives. If none are found, the shader will be marked with EShaderResourceUsageFlags::NoDerivativeOps, meaning that
 // calling code can safely assume only provided derivatives are used.
-SHADER_COMPILER_FLAGS_ENTRY( CFLAG_CheckForDerivativeOps )
-// Shader is used with indirect draws. This flag is currently used to fix a platform specific problem with certain (rare) indirect draw setups ) but it is intended to be set for all indirect draw shaders in the future.
+SHADER_COMPILER_FLAGS_ENTRY(CheckForDerivativeOps)
+// Shader is used with indirect draws. This flag is currently used to fix a platform specific problem with certain (rare) indirect draw setups, but it is intended to be set for all indirect draw shaders in the future.
 // Must not be used on shaders that are used with direct draws. Doing so might cause crashes or visual corruption on certain platforms.
-SHADER_COMPILER_FLAGS_ENTRY( CFLAG_IndirectDraw )
+SHADER_COMPILER_FLAGS_ENTRY(IndirectDraw)
 // Shader is used with shader bundles.
-SHADER_COMPILER_FLAGS_ENTRY( CFLAG_ShaderBundle )
+SHADER_COMPILER_FLAGS_ENTRY(ShaderBundle)
 // Shader code should not be stripped of comments/whitespace/line directives at the end of preprocessing
-SHADER_COMPILER_FLAGS_ENTRY( CFLAG_DisableSourceStripping )
+SHADER_COMPILER_FLAGS_ENTRY(DisableSourceStripping)
 // Shader uses RHI Shader Binding Layout for global shader binding
-SHADER_COMPILER_FLAGS_ENTRY( CFLAG_ShaderBindingLayout )
+SHADER_COMPILER_FLAGS_ENTRY(ShaderBindingLayout)
 // Request full shader analysis artifacts in output statistics. This may contain multiple compilation steps in full text form.
-SHADER_COMPILER_FLAGS_ENTRY( CFLAG_OutputAnalysisArtifacts )
-SHADER_COMPILER_FLAGS_ENTRY( CFLAG_Max )
-
+SHADER_COMPILER_FLAGS_ENTRY(OutputAnalysisArtifacts)
 
+#undef SHADER_COMPILER_FLAGS_ENTRY
+#undef SHADER_COMPILER_FLAGS_ENTRY_DEPRECATED
