Fixed bug in RDG blackboard when creating a blackboard recursively from within another blackboard struct's constructor.

diff --git a/Engine/Source/Runtime/RenderCore/Public/RenderGraphBlackboard.h b/Engine/Source/Runtime/RenderCore/Public/RenderGraphBlackboard.h
--- a/Engine/Source/Runtime/RenderCore/Public/RenderGraphBlackboard.h
+++ b/Engine/Source/Runtime/RenderCore/Public/RenderGraphBlackboard.h
@@ -68,10 +68,10 @@ public:
 			Blackboard.SetNumZeroed(StructIndex + 1);
 		}
 
-		FStruct*& Result = Blackboard[StructIndex];
-		checkf(!Result, TEXT("RDGBlackboard duplicate Create called on struct '%s'. Only one Create call per struct is allowed."), GetGeneratedTypeName<StructType>());
-		Result = Allocator.Alloc<HelperStructType>(Forward<ArgsType&&>(Args)...);
+		checkf(!Blackboard[StructIndex], TEXT("RDGBlackboard duplicate Create called on struct '%s'. Only one Create call per struct is allowed."), GetGeneratedTypeName<StructType>());
+		FStruct* Result = Allocator.Alloc<HelperStructType>(Forward<ArgsType&&>(Args)...);
 		check(Result);
+		Blackboard[StructIndex] = Result;
 		return static_cast<HelperStructType*>(Result)->Struct;
 	}
 
