[RenderCore]
* Added alignment to variable that is using typedef that has alignment added in the typedef. This typedef+newalignment is borderline not legal for c++ and seems to work but does not work for header units/c++modules

diff --git a/Engine/Source/Runtime/RenderCore/Public/ShaderParameterMacros.h b/Engine/Source/Runtime/RenderCore/Public/ShaderParameterMacros.h
--- a/Engine/Source/Runtime/RenderCore/Public/ShaderParameterMacros.h
+++ b/Engine/Source/Runtime/RenderCore/Public/ShaderParameterMacros.h
@@ -1432,7 +1432,7 @@ struct TShaderParameterStructTypeInfo<StructType[InNumElements]>
 #define INTERNAL_SHADER_PARAMETER_EXPLICIT(BaseType,TypeInfo,MemberType,MemberName,ArrayDecl,DefaultValue,Precision,OptionalShaderType,IsMemberStruct) \
 		zzMemberId##MemberName; \
 	public: \
-		TypeInfo::TAlignedType MemberName DefaultValue; \
+		alignas(TypeInfo::Alignment) TypeInfo::TAlignedType MemberName DefaultValue; \
 		static_assert(BaseType != UBMT_INVALID, "Invalid type " #MemberType " of member " #MemberName "."); \
 	private: \
 		struct zzNextMemberId##MemberName { enum { HasDeclaredResource = zzMemberId##MemberName::HasDeclaredResource || !TypeInfo::bIsStoredInConstantBuffer }; }; \
