Use TStaticBitArray for shader platform masks.

These masks never grow or shrink so it's more efficient to hard-code their size.

diff --git a/Engine/Plugins/Experimental/GPULightmass/Source/GPULightmassEditor/Private/GPULightmassEditorModule.cpp b/Engine/Plugins/Experimental/GPULightmass/Source/GPULightmassEditor/Private/GPULightmassEditorModule.cpp
--- a/Engine/Plugins/Experimental/GPULightmass/Source/GPULightmassEditor/Private/GPULightmassEditorModule.cpp
+++ b/Engine/Plugins/Experimental/GPULightmass/Source/GPULightmassEditor/Private/GPULightmassEditorModule.cpp
@@ -5,6 +5,7 @@
 #include "Internationalization/Internationalization.h"
 #include "Modules/ModuleManager.h"
 #include "HAL/IConsoleManager.h"
+#include "Containers/StaticBitArray.h"
 #include "PropertyEditorModule.h"
 #include "RenderingThread.h"
 #include "DataDrivenShaderPlatformInfo.h"
@@ -95,7 +96,7 @@ enum class ERayTracingDisabledReason : int32
 	INCAPABLE_HARDWARE, // The video card isn't capable of hardware ray tracing (too old), or supports only software emulated ray tracing (GTX 10 series)
 };
 
-typedef TBitArray<TInlineAllocator<EShaderPlatform::SP_NumPlatforms / 8>> ShaderPlatformMaskType;
+using ShaderPlatformMaskType = TStaticBitArray<EShaderPlatform::SP_NumPlatforms>;
 extern RENDERCORE_API ShaderPlatformMaskType GRayTracingPlatformMask;
 
 ERayTracingDisabledReason GetRayTracingDisabledReason(EShaderPlatform ShaderPlatform)
diff --git a/Engine/Source/Runtime/Core/Public/Containers/StaticBitArray.h b/Engine/Source/Runtime/Core/Public/Containers/StaticBitArray.h
--- a/Engine/Source/Runtime/Core/Public/Containers/StaticBitArray.h
+++ b/Engine/Source/Runtime/Core/Public/Containers/StaticBitArray.h
@@ -377,6 +377,14 @@ public:
 
 		Ar.Serialize(Words, ArchivedNumWords * sizeof(Words[0]));
 	}
+	
+	/**
+	 * Fills the entire bit array with zeros or ones.
+	 */
+	void Fill(bool Value)
+	{
+		FMemory::Memset(Words, Value ? 0xFF : 0x00, sizeof(Words));
+	}
 
 	static constexpr uint32 NumOfBits = NumBits;
 
diff --git a/Engine/Source/Runtime/RenderCore/Private/RenderUtils.cpp b/Engine/Source/Runtime/RenderCore/Private/RenderUtils.cpp
--- a/Engine/Source/Runtime/RenderCore/Private/RenderUtils.cpp
+++ b/Engine/Source/Runtime/RenderCore/Private/RenderUtils.cpp
@@ -3,6 +3,7 @@
 #include "RenderUtils.h"
 
 #include "Containers/DynamicRHIResourceArray.h"
+#include "Containers/StaticBitArray.h"
 #include "DataDrivenShaderPlatformInfo.h"
 #include "Misc/ConfigCacheIni.h"
 #include "Misc/CoreMisc.h"
@@ -543,7 +544,7 @@ RENDERCORE_API uint32 GetPlatformShadingModelsMask(const FStaticShaderPlatform P
 	return 0xFFFFFFFF;
 }
 
-typedef TBitArray<TInlineAllocator<EShaderPlatform::SP_NumPlatforms / 8>> ShaderPlatformMaskType;
+using ShaderPlatformMaskType = TStaticBitArray<EShaderPlatform::SP_NumPlatforms>;
 
 RENDERCORE_API ShaderPlatformMaskType GMobileAmbientOcclusionPlatformMask;
 
@@ -771,27 +772,27 @@ RENDERCORE_API void RenderUtilsInit()
 
 	FReadOnlyCVARCache::Initialize();
 
-	GForwardShadingPlatformMask.Init(GUseForwardShading == 1, EShaderPlatform::SP_NumPlatforms);
+	GForwardShadingPlatformMask.Fill(GUseForwardShading == 1);
 
 	static IConsoleVariable* DBufferVar = IConsoleManager::Get().FindConsoleVariable(TEXT("r.DBuffer"));
-	GDBufferPlatformMask.Init(DBufferVar && DBufferVar->GetInt(), EShaderPlatform::SP_NumPlatforms);
+	GDBufferPlatformMask.Fill(DBufferVar && DBufferVar->GetInt());
 
 	static IConsoleVariable* SelectiveBasePassOutputsCVar = IConsoleManager::Get().FindConsoleVariable(TEXT("r.SelectiveBasePassOutputs"));
-	GSelectiveBasePassOutputsPlatformMask.Init(SelectiveBasePassOutputsCVar && SelectiveBasePassOutputsCVar->GetInt(), EShaderPlatform::SP_NumPlatforms);
+	GSelectiveBasePassOutputsPlatformMask.Fill(SelectiveBasePassOutputsCVar && SelectiveBasePassOutputsCVar->GetInt());
 
 	static IConsoleVariable* DistanceFieldsCVar = IConsoleManager::Get().FindConsoleVariable(TEXT("r.DistanceFields")); 
-	GDistanceFieldsPlatformMask.Init(DistanceFieldsCVar && DistanceFieldsCVar->GetInt(), EShaderPlatform::SP_NumPlatforms);
+	GDistanceFieldsPlatformMask.Fill(DistanceFieldsCVar && DistanceFieldsCVar->GetInt());
 
 	static IConsoleVariable* RayTracingCVar = IConsoleManager::Get().FindConsoleVariable(TEXT("r.RayTracing"));
 
-	GSimpleSkyDiffusePlatformMask.Init(false, EShaderPlatform::SP_NumPlatforms);
-	GVelocityEncodeDepthPlatformMask.Init(false, EShaderPlatform::SP_NumPlatforms);
-	GNaniteFallbackMeshesPlatformMask.Init(true, EShaderPlatform::SP_NumPlatforms);
-	GRayTracingPlatformMask.Init(false, EShaderPlatform::SP_NumPlatforms);
-	GRayTracingShadersPlatformMask.Init(false, EShaderPlatform::SP_NumPlatforms);
+	GSimpleSkyDiffusePlatformMask.Fill(false);
+	GVelocityEncodeDepthPlatformMask.Fill(false);
+	GNaniteFallbackMeshesPlatformMask.Fill(true);
+	GRayTracingPlatformMask.Fill(false);
+	GRayTracingShadersPlatformMask.Fill(false);
 
 	static IConsoleVariable* MobileAmbientOcclusionCVar = IConsoleManager::Get().FindConsoleVariable(TEXT("r.Mobile.AmbientOcclusion"));
-	GMobileAmbientOcclusionPlatformMask.Init(MobileAmbientOcclusionCVar && MobileAmbientOcclusionCVar->GetInt(), EShaderPlatform::SP_NumPlatforms);
+	GMobileAmbientOcclusionPlatformMask.Fill(MobileAmbientOcclusionCVar && MobileAmbientOcclusionCVar->GetInt());
 
 #if WITH_EDITOR
 	ITargetPlatformManagerModule* TargetPlatformManager = GetTargetPlatformManager();
@@ -871,16 +872,16 @@ RENDERCORE_API void RenderUtilsInit()
 	if (IsMobilePlatform(GMaxRHIShaderPlatform))
 	{
 		static IConsoleVariable* MobileDBufferCVar = IConsoleManager::Get().FindConsoleVariable(TEXT("r.Mobile.DBuffer"));
-		GDBufferPlatformMask.Init(MobileDBufferCVar && MobileDBufferCVar->GetInt(), EShaderPlatform::SP_NumPlatforms);
+		GDBufferPlatformMask.Fill(MobileDBufferCVar && MobileDBufferCVar->GetInt());
 	}
 
 	if (RayTracingCVar && RayTracingCVar->GetInt() && GRHISupportsRayTracing)
 	{
-		GRayTracingPlatformMask.Init(true, EShaderPlatform::SP_NumPlatforms);
+		GRayTracingPlatformMask.Fill(true);
 
 		if (GRHISupportsRayTracingShaders)
 		{
-			GRayTracingShadersPlatformMask.Init(true, EShaderPlatform::SP_NumPlatforms);
+			GRayTracingShadersPlatformMask.Fill(true);
 		}
 	}
 
@@ -898,13 +899,13 @@ RENDERCORE_API void RenderUtilsInit()
 			bool bNaniteFallbackMeshes = false;
 			if (PlatformIniFile.GetBool(*CategoryName, TEXT("bGenerateNaniteFallbackMeshes"), bNaniteFallbackMeshes))
 			{
-				GNaniteFallbackMeshesPlatformMask.Init(bNaniteFallbackMeshes, EShaderPlatform::SP_NumPlatforms);
+				GNaniteFallbackMeshesPlatformMask.Fill(bNaniteFallbackMeshes);
 			}
 
 			bool bDistanceFields = false;
 			if (PlatformIniFile.GetBool(*CategoryName, TEXT("bEnableDistanceFields"), bDistanceFields) && !bDistanceFields)
 			{
-				GDistanceFieldsPlatformMask.Init(false, EShaderPlatform::SP_NumPlatforms);
+				GDistanceFieldsPlatformMask.Fill(false);
 			}
 
 			{
@@ -913,7 +914,7 @@ RENDERCORE_API void RenderUtilsInit()
 
 				if (RayTracingModeString.Equals(TEXT("Inline"), ESearchCase::IgnoreCase))
 				{
-					GRayTracingShadersPlatformMask.Init(false, EShaderPlatform::SP_NumPlatforms);
+					GRayTracingShadersPlatformMask.Fill(false);
 				}
 				else if (RayTracingModeString.Equals(TEXT("Full"), ESearchCase::IgnoreCase))
 				{
@@ -921,15 +922,15 @@ RENDERCORE_API void RenderUtilsInit()
 				}
 				else if (RayTracingModeString.Equals(TEXT("Disabled"), ESearchCase::IgnoreCase))
 				{
-					GRayTracingShadersPlatformMask.Init(false, EShaderPlatform::SP_NumPlatforms);
-					GRayTracingPlatformMask.Init(false, EShaderPlatform::SP_NumPlatforms);
+					GRayTracingShadersPlatformMask.Fill(false);
+					GRayTracingPlatformMask.Fill(false);
 				}
 				else
 				{
 					bool bRayTracing = false;
 					if (PlatformIniFile.GetBool(*CategoryName, TEXT("bEnableRayTracing"), bRayTracing) && !bRayTracing)
 					{
-						GRayTracingPlatformMask.Init(false, EShaderPlatform::SP_NumPlatforms);
+						GRayTracingPlatformMask.Fill(false);
 					}
 				}
 			}
