Disable lumen and nanite per project and strip content at cook time. Created a jira to clean up both DoesRuntimeSupportNanite and DoesPlatformSupportLumenGI to remove the option to skip the per project. Removing the option pull on many other function also exposing the per project option, more time will be requiered.

#tests PF

diff --git a/Engine/Source/Runtime/Engine/Private/ShaderCompiler/ShaderCompiler.cpp b/Engine/Source/Runtime/Engine/Private/ShaderCompiler/ShaderCompiler.cpp
--- a/Engine/Source/Runtime/Engine/Private/ShaderCompiler/ShaderCompiler.cpp
+++ b/Engine/Source/Runtime/Engine/Private/ShaderCompiler/ShaderCompiler.cpp
@@ -3947,8 +3947,7 @@ void GlobalBeginCompileShader(
 	}
 
 	{
-		static IConsoleVariable* CVar = IConsoleManager::Get().FindConsoleVariable(TEXT("r.Lumen.Supported"));
-		const bool bLumenSupported = CVar->GetInt() != 0;
+		const bool bLumenSupported = DoesProjectSupportLumenGI((EShaderPlatform)Target.Platform);
 		SET_SHADER_DEFINE(Input.Environment, PROJECT_SUPPORTS_LUMEN, bLumenSupported ? 1 : 0);
 	}
 
diff --git a/Engine/Source/Runtime/RenderCore/Private/RenderUtils.cpp b/Engine/Source/Runtime/RenderCore/Private/RenderUtils.cpp
--- a/Engine/Source/Runtime/RenderCore/Private/RenderUtils.cpp
+++ b/Engine/Source/Runtime/RenderCore/Private/RenderUtils.cpp
@@ -1256,15 +1256,14 @@ RENDERCORE_API bool PlatformSupportsVelocityRendering(const FStaticShaderPlatfor
 bool DoesPlatformSupportNanite(EShaderPlatform Platform, bool bCheckForProjectSetting)
 {
 	// Nanite allowed for this project
-	if (bCheckForProjectSetting)
+	static FShaderPlatformCachedIniValue<int32> CVarPlatformNaniteProjectEnabled(TEXT("r.Nanite.ProjectEnabled"));
+	bool bIsNaniteProjectEnable = CVarPlatformNaniteProjectEnabled.Get(Platform) != 0;
+
+	if (UNLIKELY(!bIsNaniteProjectEnable))
 	{
-		const bool bNaniteSupported = GNaniteProjectEnabled != 0;
-		if (UNLIKELY(!bNaniteSupported))
-		{
-			return false;
-		}
+		return false;
 	}
-
+	
 	// Make sure the current platform has DDPI definitions.
 	const bool bValidPlatform = FDataDrivenShaderPlatformInfo::IsValid(Platform);
 
@@ -1397,7 +1396,7 @@ bool UseVirtualShadowMaps(EShaderPlatform ShaderPlatform)
 {
 	static const auto EnableVirtualSMCVar = IConsoleManager::Get().FindConsoleVariable(TEXT("r.Shadow.Virtual.Enable"));
 	const bool bVirtualShadowMapsEnabled = EnableVirtualSMCVar ? (EnableVirtualSMCVar->GetInt() != 0) : false;
-	return bVirtualShadowMapsEnabled && DoesRuntimeSupportNanite(ShaderPlatform, true /* check for atomics */, false /* check project setting */);
+	return bVirtualShadowMapsEnabled && DoesRuntimeSupportNanite(ShaderPlatform, true /* check for atomics */, true /* check project setting */);
 }
 
 bool UseVirtualShadowMaps(EShaderPlatform ShaderPlatform, const FStaticFeatureLevel FeatureLevel)
@@ -1407,13 +1406,13 @@ bool UseVirtualShadowMaps(EShaderPlatform ShaderPlatform, const FStaticFeatureLe
 
 bool DoesPlatformSupportVirtualShadowMaps(EShaderPlatform Platform)
 {
-	return DoesPlatformSupportNanite(Platform, false /* check project setting */);
+	return DoesPlatformSupportNanite(Platform);
 }
 
 bool DoesPlatformSupportNonNaniteVirtualShadowMaps(EShaderPlatform ShaderPlatform)
 {
 	static const auto EnableCVar = IConsoleManager::Get().FindConsoleVariable(TEXT("r.Shadow.Virtual.NonNaniteVSM"));
-	return EnableCVar->GetInt() != 0 && DoesPlatformSupportNanite(ShaderPlatform, false /* check project setting */);
+	return EnableCVar->GetInt() != 0 && DoesPlatformSupportNanite(ShaderPlatform);
 }
 
 bool UseNonNaniteVirtualShadowMaps(EShaderPlatform ShaderPlatform, const FStaticFeatureLevel FeatureLevel)
@@ -2201,10 +2200,16 @@ bool HasFirstPersonGBufferBit(EShaderPlatform ShaderPlatform)
 	return !IsMobilePlatform(ShaderPlatform) && !IsForwardShadingEnabled(ShaderPlatform) && !IsStaticLightingAllowed();
 }
 
+
+bool DoesProjectSupportLumenGI(EShaderPlatform Platform)
+{
+	static FShaderPlatformCachedIniValue<int32> CVarLumenSupported(TEXT("r.Lumen.Supported"));
+	return CVarLumenSupported.Get(Platform) != 0;
+}
+
 bool DoesPlatformSupportLumenGI(EShaderPlatform Platform, bool bSkipProjectCheck)
 {
-	static IConsoleVariable* CVarLumenSupported = IConsoleManager::Get().FindConsoleVariable(TEXT("r.Lumen.Supported"));
-	const bool bLumenSupported = CVarLumenSupported->GetInt() != 0;
+	bool bLumenSupported = DoesProjectSupportLumenGI(Platform);
 
 	static IConsoleVariable* CVarLumenSupportedSM5 = IConsoleManager::Get().FindConsoleVariable(TEXT("r.Lumen.Supported.SM5"));
 	const bool bLumenSupportedSM5 = CVarLumenSupportedSM5->GetInt() != 0;
diff --git a/Engine/Source/Runtime/RenderCore/Public/RenderUtils.h b/Engine/Source/Runtime/RenderCore/Public/RenderUtils.h
--- a/Engine/Source/Runtime/RenderCore/Public/RenderUtils.h
+++ b/Engine/Source/Runtime/RenderCore/Public/RenderUtils.h
@@ -576,8 +576,9 @@ inline bool IsStaticLightingAllowed()
 // Returns 'true' when first person primitives are tagged in the GBuffer.
 RENDERCORE_API bool HasFirstPersonGBufferBit(EShaderPlatform ShaderPlatform);
 
-extern RENDERCORE_API bool DoesPlatformSupportLumenGI(EShaderPlatform Platform, bool bSkipProjectCheck = false);
+RENDERCORE_API bool DoesProjectSupportLumenGI(EShaderPlatform Platform);
 
+extern RENDERCORE_API bool DoesPlatformSupportLumenGI(EShaderPlatform Platform, bool bSkipProjectCheck = false);
 extern RENDERCORE_API bool DoesProjectSupportLumenRayTracedTranslucentRefraction();
 
 RENDERCORE_API bool DoesProjectSupportExpFogMatchesVolumetricFog();
diff --git a/Engine/Source/Runtime/Renderer/Private/LightRendering.cpp b/Engine/Source/Runtime/Renderer/Private/LightRendering.cpp
--- a/Engine/Source/Runtime/Renderer/Private/LightRendering.cpp
+++ b/Engine/Source/Runtime/Renderer/Private/LightRendering.cpp
@@ -3283,7 +3283,7 @@ public:
 	{
 		FPermutationDomain PermutationVector(Parameters.PermutationId);
 
-		if (!DoesPlatformSupportNanite(Parameters.Platform, false) && PermutationVector.Get<FNaniteCompositeDim>() != 0)
+		if (!DoesPlatformSupportNanite(Parameters.Platform) && PermutationVector.Get<FNaniteCompositeDim>() != 0)
 		{
 			return false;
 		}
