[RenderCore]
* Fixed so dllstorage is on the right things in ShaderParameterMacros

diff --git a/Engine/Source/Runtime/RenderCore/Public/ShaderParameterMacros.h b/Engine/Source/Runtime/RenderCore/Public/ShaderParameterMacros.h
--- a/Engine/Source/Runtime/RenderCore/Public/ShaderParameterMacros.h
+++ b/Engine/Source/Runtime/RenderCore/Public/ShaderParameterMacros.h
@@ -1369,8 +1369,8 @@ struct TShaderParameterStructTypeInfo<StructType[InNumElements]>
 	static const FShaderParametersMetadata* GetStructMetadata() { return StructType::FTypeInfo::GetStructMetadata(); }
 };
 
-#define INTERNAL_BEGIN_UNIFORM_BUFFER_STRUCT \
-	static const FShaderParametersMetadata* GetStructMetadata(); 
+#define INTERNAL_BEGIN_UNIFORM_BUFFER_STRUCT(DllStorage) \
+	static DllStorage const FShaderParametersMetadata* GetStructMetadata(); 
 
 #define INTERNAL_UNIFORM_BUFFER_STRUCT_GET_STRUCT_METADATA(StructTypeName) \
 	{ return StructTypeName::GetStructMetadata(); }
@@ -1397,11 +1397,11 @@ struct TShaderParameterStructTypeInfo<StructType[InNumElements]>
 #define INTERNAL_UNIFORM_BUFFER_STRUCT_CREATE_UNIFORM_BUFFER return RHICreateUniformBuffer(&InContents, FTypeInfo::GetStructMetadata()->GetLayoutPtr(), InUsage);
 
 /** Begins a uniform buffer struct declaration. */
-#define INTERNAL_SHADER_PARAMETER_STRUCT_BEGIN(StructTypeName,PrefixKeywords,ConstructorSuffix,GetStructMetadataScope,CreateUniformBufferImpl) \
-	MS_ALIGN(SHADER_PARAMETER_STRUCT_ALIGNMENT) class PrefixKeywords StructTypeName \
+#define INTERNAL_SHADER_PARAMETER_STRUCT_BEGIN(StructTypeName,DllStorage,ConstructorSuffix,GetStructMetadataScope,CreateUniformBufferImpl) \
+	MS_ALIGN(SHADER_PARAMETER_STRUCT_ALIGNMENT) class StructTypeName \
 	{ \
 	public: \
-		StructTypeName () ConstructorSuffix \
+		DllStorage StructTypeName () ConstructorSuffix \
 		struct FTypeInfo { \
 			static constexpr int32 NumRows = 1; \
 			static constexpr int32 NumColumns = 1; \
@@ -1478,8 +1478,8 @@ extern RENDERCORE_API void BuildShaderBindingLayout(TConstArrayView<FShaderParam
  *	BEGIN_SHADER_PARAMETER_STRUCT(FMyParameterStruct, RENDERER_API)
  *	END_SHADER_PARAMETER_STRUCT()
  */
-#define BEGIN_SHADER_PARAMETER_STRUCT(StructTypeName, PrefixKeywords) \
-	INTERNAL_SHADER_PARAMETER_STRUCT_BEGIN(StructTypeName, PrefixKeywords, {}, INTERNAL_SHADER_PARAMETER_GET_STRUCT_METADATA(StructTypeName), INTERNAL_SHADER_PARAMETER_STRUCT_CREATE_UNIFORM_BUFFER)
+#define BEGIN_SHADER_PARAMETER_STRUCT(StructTypeName, DllStorage) \
+	INTERNAL_SHADER_PARAMETER_STRUCT_BEGIN(StructTypeName,, {}, INTERNAL_SHADER_PARAMETER_GET_STRUCT_METADATA(StructTypeName), INTERNAL_SHADER_PARAMETER_STRUCT_CREATE_UNIFORM_BUFFER)
 
 #define END_SHADER_PARAMETER_STRUCT() \
 		zzLastMemberId; \
@@ -1508,13 +1508,13 @@ extern RENDERCORE_API void BuildShaderBindingLayout(TConstArrayView<FShaderParam
  *	// C++ file
  *	IMPLEMENT_UNIFORM_BUFFER_STRUCT(FMyParameterStruct, "MyShaderBindingName");
  */
-#define BEGIN_UNIFORM_BUFFER_STRUCT(StructTypeName, PrefixKeywords) \
-	DECLARE_UNIFORM_BUFFER_STRUCT(StructTypeName, PrefixKeywords) \
-	INTERNAL_SHADER_PARAMETER_STRUCT_BEGIN(StructTypeName,PrefixKeywords,{} INTERNAL_BEGIN_UNIFORM_BUFFER_STRUCT, INTERNAL_UNIFORM_BUFFER_STRUCT_GET_STRUCT_METADATA(StructTypeName), INTERNAL_UNIFORM_BUFFER_STRUCT_CREATE_UNIFORM_BUFFER)
+#define BEGIN_UNIFORM_BUFFER_STRUCT(StructTypeName, DllStorage) \
+	DECLARE_UNIFORM_BUFFER_STRUCT(StructTypeName, DllStorage) \
+	INTERNAL_SHADER_PARAMETER_STRUCT_BEGIN(StructTypeName,,{} INTERNAL_BEGIN_UNIFORM_BUFFER_STRUCT(DllStorage), INTERNAL_UNIFORM_BUFFER_STRUCT_GET_STRUCT_METADATA(StructTypeName), INTERNAL_UNIFORM_BUFFER_STRUCT_CREATE_UNIFORM_BUFFER)
 
-#define BEGIN_UNIFORM_BUFFER_STRUCT_WITH_CONSTRUCTOR(StructTypeName, PrefixKeywords) \
-	DECLARE_UNIFORM_BUFFER_STRUCT(StructTypeName, PrefixKeywords) \
-	INTERNAL_SHADER_PARAMETER_STRUCT_BEGIN(StructTypeName,PrefixKeywords,; INTERNAL_BEGIN_UNIFORM_BUFFER_STRUCT, INTERNAL_UNIFORM_BUFFER_STRUCT_GET_STRUCT_METADATA(StructTypeName), INTERNAL_UNIFORM_BUFFER_STRUCT_CREATE_UNIFORM_BUFFER)
+#define BEGIN_UNIFORM_BUFFER_STRUCT_WITH_CONSTRUCTOR(StructTypeName, DllStorage) \
+	DECLARE_UNIFORM_BUFFER_STRUCT(StructTypeName, DllStorage) \
+	INTERNAL_SHADER_PARAMETER_STRUCT_BEGIN(StructTypeName,DllStorage,; INTERNAL_BEGIN_UNIFORM_BUFFER_STRUCT(DllStorage), INTERNAL_UNIFORM_BUFFER_STRUCT_GET_STRUCT_METADATA(StructTypeName), INTERNAL_UNIFORM_BUFFER_STRUCT_CREATE_UNIFORM_BUFFER)
 
 #define END_UNIFORM_BUFFER_STRUCT() \
 	END_SHADER_PARAMETER_STRUCT()
