[Shaders] Fix locking mechanism for compile job status.

Always return a zero hash when it's value is not set to separate the status value from the hash pointer.

#rnx

diff --git a/Engine/Source/Runtime/Engine/Private/ShaderCompiler/ShaderCompilerJobCache.cpp b/Engine/Source/Runtime/Engine/Private/ShaderCompiler/ShaderCompilerJobCache.cpp
--- a/Engine/Source/Runtime/Engine/Private/ShaderCompiler/ShaderCompilerJobCache.cpp
+++ b/Engine/Source/Runtime/Engine/Private/ShaderCompiler/ShaderCompilerJobCache.cpp
@@ -1538,7 +1538,7 @@ void FShaderJobCache::SubmitJobs(const TArray<FShaderCommonCompileJobPtr>& InJob
 					}
 
 					const bool bSubmitJob = PreprocessShader(Job);
-					Job->GetInputHash();
+					Job->UpdateInputHash();
 					Job->UpdateStatus(EShaderCompileJobStatus::Ready);
 
 					if (bSubmitJob)
diff --git a/Engine/Source/Runtime/RenderCore/Public/ShaderCompilerJobTypes.h b/Engine/Source/Runtime/RenderCore/Public/ShaderCompilerJobTypes.h
--- a/Engine/Source/Runtime/RenderCore/Public/ShaderCompilerJobTypes.h
+++ b/Engine/Source/Runtime/RenderCore/Public/ShaderCompilerJobTypes.h
@@ -103,7 +103,6 @@ public:
 
 	void SetStatus(EShaderCompileJobStatus NewStatus)
 	{
-		checkf(NewStatus == EShaderCompileJobStatus::Unset || InputHash != nullptr, TEXT("Unless the job status is reset back to 'Unset', the input hash must be provided"));
 		Status.store(NewStatus);
 	}
 
@@ -129,16 +128,7 @@ public:
 
 	const FShaderCompilerInputHash& GetInputHash()
 	{
-		// input hash is guaranteed to be set once status is set to Submitted (which is the first status that is always set by the system)
-		// if we somehow query the hash prior to this occurring, just return zero (as it's possible the hash could be in an invalid state
-		// if the submit job is in the process of setting it). This prevents the need for any other thread safety since the status is set
-		// atomically.
-		if (GetStatus() == EShaderCompileJobStatus::Unset)
-		{
-			return FShaderCompilerInputHash::Zero;
-		}
-		check(InputHash);
-		return *InputHash;
+		return InputHash ? *InputHash : FShaderCompilerInputHash::Zero;
 	}
 
 private:
@@ -207,14 +197,15 @@ public:
 
 	TSharedPtr<FShaderCompileJobStatus> JobStatusPtr;
 
-	void UpdateStatus(EShaderCompileJobStatus NewStatus)
+	void UpdateInputHash()
 	{
-		if (NewStatus == EShaderCompileJobStatus::Ready)
-		{
-			check(bInputHashSet);
-			JobStatusPtr->SetInputHash(&InputHash);
-		}
+		(void)GetInputHash();
+		check(bInputHashSet);
+		JobStatusPtr->SetInputHash(&InputHash);
+	}
 
+	void UpdateStatus(EShaderCompileJobStatus NewStatus)
+	{
 		JobStatusPtr->SetStatus(NewStatus);
 	}
 	
