Simplify logic in UseVirtualTexturing() implementation.
Removed deprecated versions of the function and UseVirtualTextureLightmap().
Removed PLATFORM_SUPPORTS_VIRTUAL_TEXTURE_STREAMING since it is on everywhere.
Removed FPlatformProperties::SupportsVirtualTextureStreaming() since it is on everywhere.
Deprecated us of ETargetPlatformFeatures::VirtualTextureStreaming enum since UseVirtualTexturing() now replaces the need for this.
#rnx

diff --git a/Engine/Source/Developer/Android/AndroidTargetPlatformSettings/Private/AndroidTargetPlatformSettings.cpp b/Engine/Source/Developer/Android/AndroidTargetPlatformSettings/Private/AndroidTargetPlatformSettings.cpp
--- a/Engine/Source/Developer/Android/AndroidTargetPlatformSettings/Private/AndroidTargetPlatformSettings.cpp
+++ b/Engine/Source/Developer/Android/AndroidTargetPlatformSettings/Private/AndroidTargetPlatformSettings.cpp
@@ -12,7 +12,6 @@ FAndroidTargetPlatformSettings::FAndroidTargetPlatformSettings(const TCHAR* Cook
 	, MobileShadingPath(0)
 	, bDistanceField(false)
 	, bMobileForwardEnableClusteredReflections(false)
-	, bMobileVirtualTextures(false)
 {
 #if WITH_ENGINE
 	TextureLODSettings = nullptr; // These are registered by the device profile system.
@@ -20,7 +19,6 @@ FAndroidTargetPlatformSettings::FAndroidTargetPlatformSettings(const TCHAR* Cook
 	GetConfigSystem()->GetBool(TEXT("/Script/Engine.RendererSettings"), TEXT("r.DistanceFields"), bDistanceField, GEngineIni);
 	GetConfigSystem()->GetInt(TEXT("/Script/Engine.RendererSettings"), TEXT("r.Mobile.ShadingPath"), MobileShadingPath, GEngineIni);
 	GetConfigSystem()->GetBool(TEXT("/Script/Engine.RendererSettings"), TEXT("r.Mobile.Forward.EnableClusteredReflections"), bMobileForwardEnableClusteredReflections, GEngineIni);
-	GetConfigSystem()->GetBool(TEXT("/Script/Engine.RendererSettings"), TEXT("r.Mobile.VirtualTextures"), bMobileVirtualTextures, GEngineIni);
 #endif
 }
 
@@ -73,10 +71,6 @@ bool FAndroidTargetPlatformSettings::SupportsFeature(ETargetPlatformFeatures Fea
 	case ETargetPlatformFeatures::DeferredRendering:
 		return SupportsVulkanSM5();
 
-	case ETargetPlatformFeatures::VirtualTextureStreaming:
-		// TODO: should it check r.VirtualTextures for SM5 renderer ?
-		return bMobileVirtualTextures;
-
 	case ETargetPlatformFeatures::DistanceFieldAO:
 		return UsesDistanceFields();
 
diff --git a/Engine/Source/Developer/Android/AndroidTargetPlatformSettings/Public/AndroidTargetPlatformSettings.h b/Engine/Source/Developer/Android/AndroidTargetPlatformSettings/Public/AndroidTargetPlatformSettings.h
--- a/Engine/Source/Developer/Android/AndroidTargetPlatformSettings/Public/AndroidTargetPlatformSettings.h
+++ b/Engine/Source/Developer/Android/AndroidTargetPlatformSettings/Public/AndroidTargetPlatformSettings.h
@@ -146,8 +146,6 @@ protected:
 	// r.Mobile.Forward.EnableClusteredReflections value
 	bool bMobileForwardEnableClusteredReflections;
 
-	// r.Mobile.VirtualTextures value
-	bool bMobileVirtualTextures;
 #if WITH_ENGINE
 	// Holds a cache of the target LOD settings.
 	const UTextureLODSettings* TextureLODSettings;
diff --git a/Engine/Source/Developer/IOS/IOSTargetPlatformSettings/Private/IOSTargetPlatformSettings.cpp b/Engine/Source/Developer/IOS/IOSTargetPlatformSettings/Private/IOSTargetPlatformSettings.cpp
--- a/Engine/Source/Developer/IOS/IOSTargetPlatformSettings/Private/IOSTargetPlatformSettings.cpp
+++ b/Engine/Source/Developer/IOS/IOSTargetPlatformSettings/Private/IOSTargetPlatformSettings.cpp
@@ -26,7 +26,6 @@ FIOSTargetPlatformSettings::FIOSTargetPlatformSettings(bool bInIsTVOS, bool bInI
 	, MobileShadingPath(0)
 	, bDistanceField(false)
 	, bMobileForwardEnableClusteredReflections(false)
-	, bMobileVirtualTextures(false)
 {
 #if WITH_ENGINE
 	TextureLODSettings = nullptr; // TextureLODSettings are registered by the device profile.
@@ -34,7 +33,6 @@ FIOSTargetPlatformSettings::FIOSTargetPlatformSettings(bool bInIsTVOS, bool bInI
 	GetConfigSystem()->GetBool(TEXT("/Script/Engine.RendererSettings"), TEXT("r.DistanceFields"), bDistanceField, GEngineIni);
 	GetConfigSystem()->GetInt(TEXT("/Script/Engine.RendererSettings"), TEXT("r.Mobile.ShadingPath"), MobileShadingPath, GEngineIni);
 	GetConfigSystem()->GetBool(TEXT("/Script/Engine.RendererSettings"), TEXT("r.Mobile.Forward.EnableClusteredReflections"), bMobileForwardEnableClusteredReflections, GEngineIni);
-	GetConfigSystem()->GetBool(TEXT("/Script/Engine.RendererSettings"), TEXT("r.Mobile.VirtualTextures"), bMobileVirtualTextures, GEngineIni);
 #endif // #if WITH_ENGINE
 }
 
@@ -86,10 +84,6 @@ bool FIOSTargetPlatformSettings::SupportsFeature( ETargetPlatformFeatures Featur
 		case ETargetPlatformFeatures::HighQualityLightmaps:
 			return SupportsMetalMRT();
 
-		case ETargetPlatformFeatures::VirtualTextureStreaming:
-			// TODO: should it check r.VirtualTextures for SM5 renderer?
-			return bMobileVirtualTextures;
-
 		case ETargetPlatformFeatures::DistanceFieldAO:
 			return UsesDistanceFields();
 
diff --git a/Engine/Source/Developer/IOS/IOSTargetPlatformSettings/Private/IOSTargetPlatformSettings.h b/Engine/Source/Developer/IOS/IOSTargetPlatformSettings/Private/IOSTargetPlatformSettings.h
--- a/Engine/Source/Developer/IOS/IOSTargetPlatformSettings/Private/IOSTargetPlatformSettings.h
+++ b/Engine/Source/Developer/IOS/IOSTargetPlatformSettings/Private/IOSTargetPlatformSettings.h
@@ -79,9 +79,6 @@ private:
 	// r.Mobile.Forward.EnableClusteredReflections value
 	bool bMobileForwardEnableClusteredReflections;
 
-	// r.Mobile.VirtualTextures value
-	bool bMobileVirtualTextures;
-
 #if WITH_ENGINE
 	// Holds the cache of the target LOD settings.
 	const UTextureLODSettings* TextureLODSettings;
diff --git a/Engine/Source/Developer/TargetPlatform/Public/Common/TargetPlatformBase.h b/Engine/Source/Developer/TargetPlatform/Public/Common/TargetPlatformBase.h
--- a/Engine/Source/Developer/TargetPlatform/Public/Common/TargetPlatformBase.h
+++ b/Engine/Source/Developer/TargetPlatform/Public/Common/TargetPlatformBase.h
@@ -507,9 +507,6 @@ public:
 		case ETargetPlatformFeatures::MemoryMappedAnimation:
 			return TPlatformProperties::SupportsMemoryMappedAnimation();
 
-		case ETargetPlatformFeatures::VirtualTextureStreaming:
-			return TPlatformProperties::SupportsVirtualTextureStreaming();
-
 		case ETargetPlatformFeatures::SdkConnectDisconnect:
 		case ETargetPlatformFeatures::UserCredentials:
 			break;
diff --git a/Engine/Source/Developer/TargetPlatform/Public/Common/TargetPlatformSettingsBase.h b/Engine/Source/Developer/TargetPlatform/Public/Common/TargetPlatformSettingsBase.h
--- a/Engine/Source/Developer/TargetPlatform/Public/Common/TargetPlatformSettingsBase.h
+++ b/Engine/Source/Developer/TargetPlatform/Public/Common/TargetPlatformSettingsBase.h
@@ -144,9 +144,6 @@ public:
 		case ETargetPlatformFeatures::MemoryMappedAnimation:
 			return TPlatformProperties::SupportsMemoryMappedAnimation();
 
-		case ETargetPlatformFeatures::VirtualTextureStreaming:
-			return TPlatformProperties::SupportsVirtualTextureStreaming();
-
 		case ETargetPlatformFeatures::SdkConnectDisconnect:
 		case ETargetPlatformFeatures::UserCredentials:
 			break;
diff --git a/Engine/Source/Developer/TargetPlatform/Public/Interfaces/ITargetPlatformSettings.h b/Engine/Source/Developer/TargetPlatform/Public/Interfaces/ITargetPlatformSettings.h
--- a/Engine/Source/Developer/TargetPlatform/Public/Interfaces/ITargetPlatformSettings.h
+++ b/Engine/Source/Developer/TargetPlatform/Public/Interfaces/ITargetPlatformSettings.h
@@ -89,7 +89,7 @@ enum class ETargetPlatformFeatures
 	SparseTextures,
 
 	/* Can we use the virtual texture streaming system on this platform. */
-	VirtualTextureStreaming,
+	VirtualTextureStreaming UE_DEPRECATED(5.6, "VirtualTextureStreaming is now deprecated and will be removed."),
 
 	/** Lumen Global Illumination. */
 	LumenGI,
diff --git a/Engine/Source/Runtime/Core/Public/Android/AndroidPlatform.h b/Engine/Source/Runtime/Core/Public/Android/AndroidPlatform.h
--- a/Engine/Source/Runtime/Core/Public/Android/AndroidPlatform.h
+++ b/Engine/Source/Runtime/Core/Public/Android/AndroidPlatform.h
@@ -60,7 +60,6 @@ typedef FAndroidTypes FPlatformTypes;
 #define PLATFORM_UI_NEEDS_FOCUS_OUTLINES			0
 #define PLATFORM_SUPPORTS_EARLY_MOVIE_PLAYBACK		1 // movies will start before engine is initialized
 #define PLATFORM_SUPPORTS_GEOMETRY_SHADERS			(RHI_RAYTRACING) // Currently required due to the way we offset ShaderStage::EStage
-#define PLATFORM_SUPPORTS_VIRTUAL_TEXTURE_STREAMING	1
 #define PLATFORM_USE_ANSI_POSIX_MALLOC				1
 #define PLATFORM_USE_MINIMAL_HANG_DETECTION			1
 #define PLATFORM_CODE_SECTION(Name)					__attribute__((section(Name)))
diff --git a/Engine/Source/Runtime/Core/Public/Android/AndroidPlatformProperties.h b/Engine/Source/Runtime/Core/Public/Android/AndroidPlatformProperties.h
--- a/Engine/Source/Runtime/Core/Public/Android/AndroidPlatformProperties.h
+++ b/Engine/Source/Runtime/Core/Public/Android/AndroidPlatformProperties.h
@@ -135,11 +135,6 @@ struct FAndroidPlatformProperties
 		// Data cooked for 16kb page size will work for 4kb page size firmware.
 		return 16384;
 	}
-
-	static FORCEINLINE bool SupportsVirtualTextureStreaming()
-	{
-		return true;
-	}
 };
 
 #ifdef PROPERTY_HEADER_SHOULD_DEFINE_TYPE
diff --git a/Engine/Source/Runtime/Core/Public/GenericPlatform/GenericPlatformProperties.h b/Engine/Source/Runtime/Core/Public/GenericPlatform/GenericPlatformProperties.h
--- a/Engine/Source/Runtime/Core/Public/GenericPlatform/GenericPlatformProperties.h
+++ b/Engine/Source/Runtime/Core/Public/GenericPlatform/GenericPlatformProperties.h
@@ -264,11 +264,6 @@ struct FGenericPlatformProperties
 		return 0;
 	}
 	
-	static FORCEINLINE bool SupportsVirtualTextureStreaming()
-	{
-		return false; // Currently VT is opt-in
-	}
-
 	static FORCEINLINE bool SupportsRayTracing()
 	{
 		return false;
diff --git a/Engine/Source/Runtime/Core/Public/HAL/Platform.h b/Engine/Source/Runtime/Core/Public/HAL/Platform.h
--- a/Engine/Source/Runtime/Core/Public/HAL/Platform.h
+++ b/Engine/Source/Runtime/Core/Public/HAL/Platform.h
@@ -315,9 +315,6 @@
 #ifndef PLATFORM_SUPPORTS_TEXTURE_STREAMING
 	#define PLATFORM_SUPPORTS_TEXTURE_STREAMING	1
 #endif
-#ifndef PLATFORM_SUPPORTS_VIRTUAL_TEXTURE_STREAMING
-	#define PLATFORM_SUPPORTS_VIRTUAL_TEXTURE_STREAMING	0
-#endif
 #ifndef PLATFORM_SUPPORTS_VIRTUAL_TEXTURES
 	#define PLATFORM_SUPPORTS_VIRTUAL_TEXTURES		0
 #endif
diff --git a/Engine/Source/Runtime/Core/Public/IOS/IOSPlatform.h b/Engine/Source/Runtime/Core/Public/IOS/IOSPlatform.h
--- a/Engine/Source/Runtime/Core/Public/IOS/IOSPlatform.h
+++ b/Engine/Source/Runtime/Core/Public/IOS/IOSPlatform.h
@@ -51,7 +51,6 @@ typedef FIOSPlatformTypes FPlatformTypes;
 
 #define PLATFORM_NEEDS_RHIRESOURCELIST					0
 #define PLATFORM_SUPPORTS_GEOMETRY_SHADERS				0
-#define PLATFORM_SUPPORTS_VIRTUAL_TEXTURE_STREAMING		1
 #define PLATFORM_SUPPORTS_BINDLESS_RENDERING			0
 
 #define PLATFORM_RETURN_ADDRESS_FOR_CALLSTACKTRACING    PLATFORM_RETURN_ADDRESS
diff --git a/Engine/Source/Runtime/Core/Public/IOS/IOSPlatformProperties.h b/Engine/Source/Runtime/Core/Public/IOS/IOSPlatformProperties.h
--- a/Engine/Source/Runtime/Core/Public/IOS/IOSPlatformProperties.h
+++ b/Engine/Source/Runtime/Core/Public/IOS/IOSPlatformProperties.h
@@ -104,11 +104,6 @@ struct FIOSPlatformProperties
 	{
 		return true;
 	}
-
-	static FORCEINLINE bool SupportsVirtualTextureStreaming()
-	{
-		return true;
-	}
 };
 
 struct FTVOSPlatformProperties : public FIOSPlatformProperties
diff --git a/Engine/Source/Runtime/Core/Public/Linux/LinuxPlatformProperties.h b/Engine/Source/Runtime/Core/Public/Linux/LinuxPlatformProperties.h
--- a/Engine/Source/Runtime/Core/Public/Linux/LinuxPlatformProperties.h
+++ b/Engine/Source/Runtime/Core/Public/Linux/LinuxPlatformProperties.h
@@ -159,12 +159,6 @@ struct FLinuxPlatformProperties
 	{
 		return !IsServerOnly();
 	}
-
-	static FORCEINLINE bool SupportsVirtualTextureStreaming()
-	{
-		return true;
-	}
-
 };
 
 #ifdef PROPERTY_HEADER_SHOULD_DEFINE_TYPE
diff --git a/Engine/Source/Runtime/Core/Public/Mac/MacPlatform.h b/Engine/Source/Runtime/Core/Public/Mac/MacPlatform.h
--- a/Engine/Source/Runtime/Core/Public/Mac/MacPlatform.h
+++ b/Engine/Source/Runtime/Core/Public/Mac/MacPlatform.h
@@ -62,7 +62,6 @@ typedef FMacPlatformTypes FPlatformTypes;
 #define PLATFORM_MAX_FILEPATH_LENGTH_DEPRECATED			MAC_MAX_PATH
 #define PLATFORM_SUPPORTS_TBB							1
 #define PLATFORM_SUPPORTS_MIMALLOC						PLATFORM_64BITS
-#define PLATFORM_SUPPORTS_VIRTUAL_TEXTURE_STREAMING		1
 #define PLATFORM_SUPPORTS_MESH_SHADERS                  1
 #define PLATFORM_SUPPORTS_BINDLESS_RENDERING            1
 #define PLATFORM_SUPPORTS_GEOMETRY_SHADERS              1
diff --git a/Engine/Source/Runtime/Core/Public/Mac/MacPlatformProperties.h b/Engine/Source/Runtime/Core/Public/Mac/MacPlatformProperties.h
--- a/Engine/Source/Runtime/Core/Public/Mac/MacPlatformProperties.h
+++ b/Engine/Source/Runtime/Core/Public/Mac/MacPlatformProperties.h
@@ -132,11 +132,6 @@ struct FMacPlatformProperties
 	{
 		return !IsServerOnly() && !HasEditorOnlyData();
 	}
-
-	static FORCEINLINE bool SupportsVirtualTextureStreaming()
-	{
-		return true;
-	}
 };
 
 #ifdef PROPERTY_HEADER_SHOULD_DEFINE_TYPE
diff --git a/Engine/Source/Runtime/Core/Public/Unix/UnixPlatform.h b/Engine/Source/Runtime/Core/Public/Unix/UnixPlatform.h
--- a/Engine/Source/Runtime/Core/Public/Unix/UnixPlatform.h
+++ b/Engine/Source/Runtime/Core/Public/Unix/UnixPlatform.h
@@ -66,7 +66,6 @@ typedef FUnixPlatformTypes FPlatformTypes;
 #define PLATFORM_SUPPORTS_STACK_SYMBOLS					1
 #define PLATFORM_IS_ANSI_MALLOC_THREADSAFE				1
 #define PLATFORM_ALLOW_ALLOCATIONS_IN_FASYNCWRITER_SERIALIZEBUFFERTOARCHIVE 0
-#define PLATFORM_SUPPORTS_VIRTUAL_TEXTURE_STREAMING		1
 #define PLATFORM_USE_ANSI_POSIX_MALLOC					1
 #define PLATFORM_SUPPORTS_BINDLESS_RENDERING			1
 #define PLATFORM_SUPPORTS_MESH_SHADERS					1
diff --git a/Engine/Source/Runtime/Core/Public/Windows/WindowsPlatform.h b/Engine/Source/Runtime/Core/Public/Windows/WindowsPlatform.h
--- a/Engine/Source/Runtime/Core/Public/Windows/WindowsPlatform.h
+++ b/Engine/Source/Runtime/Core/Public/Windows/WindowsPlatform.h
@@ -105,7 +105,6 @@ typedef FWindowsPlatformTypes FPlatformTypes;
 #define PLATFORM_COMPILER_HAS_TCHAR_WMAIN					1
 #define PLATFORM_SUPPORTS_EARLY_MOVIE_PLAYBACK				(!WITH_EDITOR) // movies will start before engine is initalized
 #define PLATFORM_USE_GENERIC_STRING_IMPLEMENTATION			0 // Can set this to 1 if you need to debug FGenericWidePlatformString::GetVarArgs on Windows
-#define PLATFORM_SUPPORTS_VIRTUAL_TEXTURE_STREAMING			1
 #define PLATFORM_SUPPORTS_VARIABLE_RATE_SHADING				1
 #define PLATFORM_SUPPORTS_MESH_SHADERS						1
 #define PLATFORM_SUPPORTS_WORKGRAPH_SHADERS					1
diff --git a/Engine/Source/Runtime/Core/Public/Windows/WindowsPlatformProperties.h b/Engine/Source/Runtime/Core/Public/Windows/WindowsPlatformProperties.h
--- a/Engine/Source/Runtime/Core/Public/Windows/WindowsPlatformProperties.h
+++ b/Engine/Source/Runtime/Core/Public/Windows/WindowsPlatformProperties.h
@@ -137,12 +137,6 @@ struct FWindowsPlatformProperties
 
 		return 1.0f;
 	}
-
-	static FORCEINLINE bool SupportsVirtualTextureStreaming()
-	{
-		return true;
-	}
-
 };
 
 #ifdef PROPERTY_HEADER_SHOULD_DEFINE_TYPE
diff --git a/Engine/Source/Runtime/Engine/Private/Texture.cpp b/Engine/Source/Runtime/Engine/Private/Texture.cpp
--- a/Engine/Source/Runtime/Engine/Private/Texture.cpp
+++ b/Engine/Source/Runtime/Engine/Private/Texture.cpp
@@ -482,23 +482,17 @@ bool UTexture::IsPostLoadThreadSafe() const
 
 /*static*/ bool UTexture::IsVirtualTexturingEnabled( const ITargetPlatformSettings * TargetPlatform /*= nullptr*/ )
 {
-	// check the project setting cvar for overall VT being on/off :
-	//	(this is in the host platform config)
-	if ( ! CVarVirtualTextures.GetValueOnAnyThread() )
+	// Check the host platform config.
+	if (!UseVirtualTexturing(GMaxRHIShaderPlatform))
 	{
 		return false;
 	}
-	
-	// note: rendering code (UseVirtualTexturing) also checks FPlatformProperties::SupportsVirtualTextureStreaming()
-	//	that has historically not been checked in the texture code, so we continue to not check it here
-
-	// optionally could do: if ! TargetPlatform , TargetPlatform = CurrentRunning ?
 
-	if ( TargetPlatform )
+	// Check the target platform config.	
+	// Optionally could do: if ! TargetPlatform , TargetPlatform = CurrentRunning ?
+	if (TargetPlatform != nullptr)
 	{
-		// this will wind up checking "r.Mobile.VirtualTextures" if TargetPlatform is a mobile platform
-		const bool bPlatformSupportsVirtualTextureStreaming = TargetPlatform->SupportsFeature(ETargetPlatformFeatures::VirtualTextureStreaming);
-		if ( ! bPlatformSupportsVirtualTextureStreaming )
+		if (!UseVirtualTexturing(TargetPlatform))
 		{
 			return false;
 		}
@@ -1135,7 +1129,7 @@ void UTexture::Serialize(FArchive& Ar)
 
 	if (Ar.IsCooking() && VirtualTextureStreaming)
 	{
-		if (UseVirtualTexturing(GMaxRHIShaderPlatform, Ar.CookingTarget()) == false)
+		if (IsVirtualTexturingEnabled(Ar.CookingTarget()) == false)
 		{
 			UE_LOG(LogTexture, Display, TEXT("%s is marked for virtual streaming but virtual texture streaming is not available."), *GetPathName());
 		}
@@ -1850,7 +1844,7 @@ bool UTexture::IsPossibleToStream() const
 // Based on target platform, returns whether texture is candidate to be streamed.
 bool UTexture::IsCandidateForTextureStreamingOnPlatformDuringCook(const ITargetPlatform* InTargetPlatform) const
 {
-	const bool bIsVirtualTextureStreaming = InTargetPlatform->SupportsFeature(ETargetPlatformFeatures::VirtualTextureStreaming) ? VirtualTextureStreaming : false;
+	const bool bIsVirtualTextureStreaming = IsVirtualTexturingEnabled(InTargetPlatform) ? VirtualTextureStreaming : false;
 	const bool bIsCandidateForTextureStreaming = InTargetPlatform->SupportsFeature(ETargetPlatformFeatures::TextureStreaming) && !bIsVirtualTextureStreaming;
 
 	if (bIsCandidateForTextureStreaming &&
diff --git a/Engine/Source/Runtime/RenderCore/Private/RenderUtils.cpp b/Engine/Source/Runtime/RenderCore/Private/RenderUtils.cpp
--- a/Engine/Source/Runtime/RenderCore/Private/RenderUtils.cpp
+++ b/Engine/Source/Runtime/RenderCore/Private/RenderUtils.cpp
@@ -1147,64 +1147,61 @@ RENDERCORE_API void QuantizeSceneBufferSize(const FIntPoint& InBufferSize, FIntP
 	OutBufferSize.Y = (InBufferSize.Y + DividableBy - 1) & Mask;
 }
 
-RENDERCORE_API bool UseVirtualTexturing(const FStaticShaderPlatform ShaderPlatform, const ITargetPlatform* TargetPlatform)
+RENDERCORE_API bool UseVirtualTexturing(const FStaticShaderPlatform InShaderPlatform)
 {
-#if PLATFORM_SUPPORTS_VIRTUAL_TEXTURE_STREAMING
-	if (!FPlatformProperties::SupportsVirtualTextureStreaming())
+	static FShaderPlatformCachedIniValue<bool> VirtualTexturesIniValue(TEXT("r.VirtualTextures"));
+	if (!VirtualTexturesIniValue.Get(InShaderPlatform))
 	{
 		return false;
 	}
 
-	// does the project has it enabled ?
-	static const auto CVarVirtualTexture = IConsoleManager::Get().FindTConsoleVariableDataInt(TEXT("r.VirtualTextures"));
-	check(CVarVirtualTexture);
-	if (CVarVirtualTexture->GetValueOnAnyThread() == 0)
+	if (IsMobilePlatform(InShaderPlatform))
 	{
-		return false;
-	}
-
-	bool bIsMobilePlatform = IsMobilePlatform(ShaderPlatform);
-#if WITH_EDITOR
-	// in case we are cooking for a platform that supports only mobile rendering
-	bIsMobilePlatform |= (TargetPlatform && !TargetPlatform->SupportsFeature(ETargetPlatformFeatures::DeferredRendering));
-#endif
-	// mobile needs an additional switch to enable VT		
-	if (bIsMobilePlatform)
-	{
-	#if WITH_EDITOR		
-		if (TargetPlatform)
+		// We treat r.Mobile.VirtualTextures as per project, and _not_ per platform.
+		// To change virtual texture support per platform use r.VirtualTextures.
+		// This makes management simpler and avoids a case where virtual texture support for platforms that support both desktop and mobile renderers could diverge.
+		// The values can't diverge per platform because we only cook one version of the texture resources.
+		static const auto CVarMobileVirtualTextures = IConsoleManager::Get().FindTConsoleVariableDataInt(TEXT("r.Mobile.VirtualTextures"));
+		static  bool bMobileVirtualTextures = CVarMobileVirtualTextures->GetValueOnAnyThread() != 0;
+		if (!bMobileVirtualTextures)
 		{
-			// VT can be toggled on and off for specific platform
-			return TargetPlatform->SupportsFeature(ETargetPlatformFeatures::VirtualTextureStreaming);
-		}
-		else
-	#endif
-		{ 
-			static FShaderPlatformCachedIniValue<bool> MobileVirtualTexturesIniValue(TEXT("r.Mobile.VirtualTextures"));
-			return (MobileVirtualTexturesIniValue.Get(ShaderPlatform) != false);
+			return false;
 		}
 	}
+
 	return true;
-#else
-	return false;
-#endif
 }
 
-RENDERCORE_API bool UseVirtualTexturing(const FStaticFeatureLevel InFeatureLevel, const ITargetPlatform* TargetPlatform)
+RENDERCORE_API bool UseVirtualTexturing(const ITargetPlatformSettings* InTargetPlatform)
 {
-	return UseVirtualTexturing(GetFeatureLevelShaderPlatform(InFeatureLevel), TargetPlatform);
+	if (InTargetPlatform != nullptr)
+	{	
+		TArray<FName> ShaderFormats;
+		InTargetPlatform->GetAllTargetedShaderFormats(ShaderFormats);
+		
+		for (FName ShaderFormat : ShaderFormats)
+		{
+			const EShaderPlatform ShaderPlatform = ShaderFormatToLegacyShaderPlatform(ShaderFormat);
+			if (!UseVirtualTexturing(ShaderPlatform))
+			{
+				return false;
+			}
+		}
+	}
+
+	return true;
 }
 
-RENDERCORE_API bool UseVirtualTextureLightmap(const FStaticShaderPlatform InShaderPlatform, const ITargetPlatform* TargetPlatform)
+RENDERCORE_API bool UseVirtualTexturing(const FStaticShaderPlatform InShaderPlatform, const ITargetPlatformSettings* InTargetPlatform)
 {
-	static const auto CVar = IConsoleManager::Get().FindTConsoleVariableDataInt(TEXT("r.VirtualTexturedLightmaps"));
-	const bool bUseVirtualTextureLightmap = (CVar->GetValueOnAnyThread() != 0) && UseVirtualTexturing(InShaderPlatform, TargetPlatform);
-	return bUseVirtualTextureLightmap;
+	return UseVirtualTexturing(InShaderPlatform) && UseVirtualTexturing(InTargetPlatform);
 }
 
-RENDERCORE_API bool UseVirtualTextureLightmap(const FStaticFeatureLevel InFeatureLevel, const ITargetPlatform* TargetPlatform)
+
+RENDERCORE_API bool UseVirtualTextureLightmap(const FStaticShaderPlatform InShaderPlatform, const ITargetPlatformSettings* TargetPlatform)
 {
-	return UseVirtualTextureLightmap(GetFeatureLevelShaderPlatform(InFeatureLevel), TargetPlatform);
+	static FShaderPlatformCachedIniValue<int32> PerPlatformCVar(TEXT("r.VirtualTexturedLightmaps"));
+	return (PerPlatformCVar.Get(InShaderPlatform) != 0) && UseVirtualTexturing(InShaderPlatform, TargetPlatform);
 }
 
 RENDERCORE_API bool UseNaniteLandscapeMesh(EShaderPlatform ShaderPlatform)
diff --git a/Engine/Source/Runtime/RenderCore/Private/Shader.cpp b/Engine/Source/Runtime/RenderCore/Private/Shader.cpp
--- a/Engine/Source/Runtime/RenderCore/Private/Shader.cpp
+++ b/Engine/Source/Runtime/RenderCore/Private/Shader.cpp
@@ -2624,7 +2624,7 @@ void ShaderMapAppendKey(EShaderPlatform Platform, FShaderKeyGenerator& KeyGen)
 			}
 		}
 
-		const bool VTSupported = TargetPlatform != nullptr && TargetPlatform->SupportsFeature(ETargetPlatformFeatures::VirtualTextureStreaming);
+		const bool VTSupported = UseVirtualTexturing(Platform, TargetPlatform);
 
 		KeyGen.AppendSeparator();
 		KeyGen.Append(TEXT("VT"));
diff --git a/Engine/Source/Runtime/RenderCore/Public/RenderUtils.h b/Engine/Source/Runtime/RenderCore/Public/RenderUtils.h
--- a/Engine/Source/Runtime/RenderCore/Public/RenderUtils.h
+++ b/Engine/Source/Runtime/RenderCore/Public/RenderUtils.h
@@ -9,6 +9,7 @@
 #include "ReadOnlyCVARCache.h"
 
 class ITargetPlatform;
+class ITargetPlatformSettings;
 struct IPooledRenderTarget;
 enum class ERayTracingMode : uint8;
 
@@ -389,10 +390,9 @@ RENDERCORE_API void QuantizeSceneBufferSize(const FIntPoint& InBufferSize, FIntP
 /**
 * Checks if virtual texturing enabled and supported
 */
-RENDERCORE_API bool UseVirtualTexturing(const FStaticShaderPlatform InShaderPlatform, const ITargetPlatform* TargetPlatform = nullptr);
-
-UE_DEPRECATED(5.4, "Use version that takes FStaticShaderPlatform instead")
-RENDERCORE_API bool UseVirtualTexturing(const FStaticFeatureLevel InFeatureLevel, const ITargetPlatform* TargetPlatform = nullptr);
+RENDERCORE_API bool UseVirtualTexturing(const FStaticShaderPlatform InShaderPlatform);
+RENDERCORE_API bool UseVirtualTexturing(const ITargetPlatformSettings* InTargetPlatform);
+RENDERCORE_API bool UseVirtualTexturing(const FStaticShaderPlatform InShaderPlatform, const ITargetPlatformSettings* InTargetPlatform);
 
 RENDERCORE_API bool NaniteAtomicsSupported();
 RENDERCORE_API bool NaniteWorkGraphMaterialsSupported();
@@ -476,10 +476,7 @@ RENDERCORE_API bool IsSingleLayerWaterDepthPrepassEnabled(const FStaticShaderPla
 /**
 *	Checks if virtual texturing lightmap enabled and supported
 */
-RENDERCORE_API bool UseVirtualTextureLightmap(const FStaticShaderPlatform Platform, const ITargetPlatform* TargetPlatform = nullptr);
-
-UE_DEPRECATED(5.4, "Use version that takes FStaticShaderPlatform instead")
-RENDERCORE_API bool UseVirtualTextureLightmap(const FStaticFeatureLevel InFeatureLevel, const ITargetPlatform* TargetPlatform = nullptr);
+RENDERCORE_API bool UseVirtualTextureLightmap(const FStaticShaderPlatform Platform, const ITargetPlatformSettings* TargetPlatform = nullptr);
 
 /**
 *	Checks if platform uses a Nanite landscape mesh
