Added r.RDG.ParallelSetup.TaskPriorityBias to control biasing all setup task priorities.

diff --git a/Engine/Source/Runtime/RenderCore/Private/RenderGraphBuilder.cpp b/Engine/Source/Runtime/RenderCore/Private/RenderGraphBuilder.cpp
--- a/Engine/Source/Runtime/RenderCore/Private/RenderGraphBuilder.cpp
+++ b/Engine/Source/Runtime/RenderCore/Private/RenderGraphBuilder.cpp
@@ -603,6 +603,8 @@ FRDGBuilder::FRDGBuilder(FRHICommandListImmediate& InRHICmdList, FRDGEventName I
 
 	const bool bParallelSetupEnabledForPlatform = ::IsParallelSetupEnabled(ShaderPlatform);
 	ParallelSetup.bEnabled   = bParallelSetupEnabledForPlatform && EnumHasAnyFlags(InFlags, ERDGBuilderFlags::ParallelSetup);
+	ParallelSetup.TaskPriorityBias = GRDGParallelSetupTaskPriorityBias;
+
 	bParallelCompileEnabled  = ::IsParallelSetupEnabled(ShaderPlatform) && EnumHasAnyFlags(InFlags, ERDGBuilderFlags::ParallelCompile);
 
 	if (TransientResourceAllocator)
diff --git a/Engine/Source/Runtime/RenderCore/Private/RenderGraphPrivate.cpp b/Engine/Source/Runtime/RenderCore/Private/RenderGraphPrivate.cpp
--- a/Engine/Source/Runtime/RenderCore/Private/RenderGraphPrivate.cpp
+++ b/Engine/Source/Runtime/RenderCore/Private/RenderGraphPrivate.cpp
@@ -307,6 +307,12 @@ FAutoConsoleVariableRef CVarRDGParallelSetup(
 	TEXT(" 1: pass setup is done asynchronously (default);"),
 	ECVF_RenderThreadSafe);
 
+int32 GRDGParallelSetupTaskPriorityBias = 1;
+FAutoConsoleVariableRef CVarRDGParallelSetupTaskPriorityBias(
+	TEXT("r.RDG.ParallelSetup.TaskPriorityBias"), GRDGParallelSetupTaskPriorityBias,
+	TEXT("Biases the task priority of all setup tasks. Useful as a tweak when contention from game thread tasks is high."),
+	ECVF_RenderThreadSafe);
+
 int32 GRDGParallelExecute = 2;
 FAutoConsoleVariableRef CVarRDGParallelExecute(
 	TEXT("r.RDG.ParallelExecute"), GRDGParallelExecute,
diff --git a/Engine/Source/Runtime/RenderCore/Private/RenderGraphPrivate.h b/Engine/Source/Runtime/RenderCore/Private/RenderGraphPrivate.h
--- a/Engine/Source/Runtime/RenderCore/Private/RenderGraphPrivate.h
+++ b/Engine/Source/Runtime/RenderCore/Private/RenderGraphPrivate.h
@@ -95,6 +95,7 @@ extern int32 GRDGTransientIndirectArgBuffers;
 
 extern int32 GRDGParallelDestruction;
 extern int32 GRDGParallelSetup;
+extern int32 GRDGParallelSetupTaskPriorityBias;
 extern int32 GRDGParallelExecute;
 extern int32 GRDGParallelExecutePassMin;
 extern int32 GRDGParallelExecutePassMax;
diff --git a/Engine/Source/Runtime/RenderCore/Public/RenderGraphBuilder.h b/Engine/Source/Runtime/RenderCore/Public/RenderGraphBuilder.h
--- a/Engine/Source/Runtime/RenderCore/Public/RenderGraphBuilder.h
+++ b/Engine/Source/Runtime/RenderCore/Public/RenderGraphBuilder.h
@@ -996,6 +996,12 @@ private:
 		TStaticArray<TArray<UE::Tasks::FTask, FRDGArrayAllocator>, (int32)ERDGSetupTaskWaitPoint::MAX> Tasks;
 
 		bool bEnabled = false;
+		int8 TaskPriorityBias = 0;
+
+		UE::Tasks::ETaskPriority GetTaskPriority(UE::Tasks::ETaskPriority TaskPriority) const
+		{
+			return (UE::Tasks::ETaskPriority)FMath::Clamp((int32)TaskPriority - TaskPriorityBias, 0, (int32)UE::Tasks::ETaskPriority::Count - 1);
+		}
 
 	} ParallelSetup;
 
diff --git a/Engine/Source/Runtime/RenderCore/Public/RenderGraphBuilder.inl b/Engine/Source/Runtime/RenderCore/Public/RenderGraphBuilder.inl
--- a/Engine/Source/Runtime/RenderCore/Public/RenderGraphBuilder.inl
+++ b/Engine/Source/Runtime/RenderCore/Public/RenderGraphBuilder.inl
@@ -611,11 +611,11 @@ UE::Tasks::FTask FRDGBuilder::AddSetupTask(
 	}
 	else if (Pipe)
 	{
-		Task = Pipe->Launch(TEXT("FRDGBuilder::AddSetupTask"), MoveTemp(OuterLambda), Forward<PrerequisitesCollectionType&&>(Prerequisites), Priority, ExtendedTaskPriority);
+		Task = Pipe->Launch(TEXT("FRDGBuilder::AddSetupTask"), MoveTemp(OuterLambda), Forward<PrerequisitesCollectionType&&>(Prerequisites), ParallelSetup.GetTaskPriority(Priority), ExtendedTaskPriority);
 	}
 	else
 	{
-		Task = UE::Tasks::Launch(TEXT("FRDGBuilder::AddSetupTask"), MoveTemp(OuterLambda), Forward<PrerequisitesCollectionType&&>(Prerequisites), Priority, ExtendedTaskPriority);
+		Task = UE::Tasks::Launch(TEXT("FRDGBuilder::AddSetupTask"), MoveTemp(OuterLambda), Forward<PrerequisitesCollectionType&&>(Prerequisites), ParallelSetup.GetTaskPriority(Priority), ExtendedTaskPriority);
 	}
 
 	if (Task.IsValid())
@@ -713,11 +713,11 @@ UE::Tasks::FTask FRDGBuilder::AddCommandListSetupTask(
 	}
 	else if (Pipe)
 	{
-		Task = Pipe->Launch(TEXT("FRDGBuilder::AddCommandListSetupTask"), MoveTemp(OuterLambda), Forward<PrerequisitesCollectionType&&>(Prerequisites), Priority, ExtendedTaskPriority);
+		Task = Pipe->Launch(TEXT("FRDGBuilder::AddCommandListSetupTask"), MoveTemp(OuterLambda), Forward<PrerequisitesCollectionType&&>(Prerequisites), ParallelSetup.GetTaskPriority(Priority), ExtendedTaskPriority);
 	}
 	else
 	{
-		Task = UE::Tasks::Launch(TEXT("FRDGBuilder::AddCommandListSetupTask"), MoveTemp(OuterLambda), Forward<PrerequisitesCollectionType&&>(Prerequisites), Priority, ExtendedTaskPriority);
+		Task = UE::Tasks::Launch(TEXT("FRDGBuilder::AddCommandListSetupTask"), MoveTemp(OuterLambda), Forward<PrerequisitesCollectionType&&>(Prerequisites), ParallelSetup.GetTaskPriority(Priority), ExtendedTaskPriority);
 	}
 
 	if (Task.IsValid())
