Fix compilation when RDG_ENABLE_PARALLEL_TASKS = 0
#rnx

diff --git a/Engine/Source/Runtime/RenderCore/Private/RenderGraphBuilder.cpp b/Engine/Source/Runtime/RenderCore/Private/RenderGraphBuilder.cpp
--- a/Engine/Source/Runtime/RenderCore/Private/RenderGraphBuilder.cpp
+++ b/Engine/Source/Runtime/RenderCore/Private/RenderGraphBuilder.cpp
@@ -603,7 +603,9 @@ FRDGBuilder::FRDGBuilder(FRHICommandListImmediate& InRHICmdList, FRDGEventName I
 
 	const bool bParallelSetupEnabledForPlatform = ::IsParallelSetupEnabled(ShaderPlatform);
 	ParallelSetup.bEnabled   = bParallelSetupEnabledForPlatform && EnumHasAnyFlags(InFlags, ERDGBuilderFlags::ParallelSetup);
+#if RDG_ENABLE_PARALLEL_TASKS
 	ParallelSetup.TaskPriorityBias = GRDGParallelSetupTaskPriorityBias;
+#endif
 
 	bParallelCompileEnabled  = ::IsParallelSetupEnabled(ShaderPlatform) && EnumHasAnyFlags(InFlags, ERDGBuilderFlags::ParallelCompile);
 
