[shaders]
- store shadertype associations for each shader in a shadermap in the editor only data
- output a debug artifact containing per-shadertype permutation/memory stats; note that this is not directly representative of final shader memory usage since it doesn't account for shader library deduplication or shader library chunk re-duplication; it is only intended to be used as a tool for tracking/identifying shader growth due to added/modified shadertypes
- remove the old "dumpshadercodestats" path from the shader library, as well as the stats tracking "unique" shaders and shader memory; the former is no longer used and the latter does not correctly account for library chunking
- bump shader version due to shadermap editoronly data change

diff --git a/Engine/Plugins/Compositing/OpenColorIO/Source/OpenColorIO/Private/OpenColorIOShaderMap.cpp b/Engine/Plugins/Compositing/OpenColorIO/Source/OpenColorIO/Private/OpenColorIOShaderMap.cpp
--- a/Engine/Plugins/Compositing/OpenColorIO/Source/OpenColorIO/Private/OpenColorIOShaderMap.cpp
+++ b/Engine/Plugins/Compositing/OpenColorIO/Source/OpenColorIO/Private/OpenColorIOShaderMap.cpp
@@ -442,7 +442,7 @@ FShader* FOpenColorIOShaderMap::ProcessCompilationResultsForSingleJob(const TRef
 	FShaderCompileJob* CurrentJob = SingleJob->GetSingleShaderJob();
 	check(CurrentJob->Id == CompilingId);
 
-	GetResourceCode()->AddShaderCompilerOutput(CurrentJob->Output, CurrentJob->Key.ToString(), CurrentJob->Input.GenerateDebugInfo());
+	GetResourceCode()->AddShaderCompilerOutput(CurrentJob->Output, CurrentJob->Key, CurrentJob->Input.GenerateDebugInfo());
 
 	FShader* Shader = nullptr;
 
diff --git a/Engine/Plugins/Experimental/NNERuntimeIREE/Source/NNERuntimeIREEShader/Private/NNERuntimeIREEShaderMap.cpp b/Engine/Plugins/Experimental/NNERuntimeIREE/Source/NNERuntimeIREEShader/Private/NNERuntimeIREEShaderMap.cpp
--- a/Engine/Plugins/Experimental/NNERuntimeIREE/Source/NNERuntimeIREEShader/Private/NNERuntimeIREEShaderMap.cpp
+++ b/Engine/Plugins/Experimental/NNERuntimeIREE/Source/NNERuntimeIREEShader/Private/NNERuntimeIREEShaderMap.cpp
@@ -340,7 +340,7 @@ FShader* FNNERuntimeIREEShaderMap::ProcessCompilationResultsForSingleJob(FShader
 {
 	check(CurrentJob.Id == CompilingId);
 
-	GetResourceCode()->AddShaderCompilerOutput(CurrentJob.Output, CurrentJob.Key.ToString(), CurrentJob.Input.GenerateDebugInfo());
+	GetResourceCode()->AddShaderCompilerOutput(CurrentJob.Output, CurrentJob.Key, CurrentJob.Input.GenerateDebugInfo());
 
 	FShader* Shader = nullptr;
 
diff --git a/Engine/Plugins/FX/Niagara/Source/NiagaraShader/Private/NiagaraShader.cpp b/Engine/Plugins/FX/Niagara/Source/NiagaraShader/Private/NiagaraShader.cpp
--- a/Engine/Plugins/FX/Niagara/Source/NiagaraShader/Private/NiagaraShader.cpp
+++ b/Engine/Plugins/FX/Niagara/Source/NiagaraShader/Private/NiagaraShader.cpp
@@ -779,7 +779,7 @@ FShader* FNiagaraShaderMap::ProcessCompilationResultsForSingleJob(const TRefCoun
 	auto CurrentJob = SingleJob->GetSingleShaderJob();
 	check(CurrentJob->Id == CompilingId);
 
-	GetResourceCode()->AddShaderCompilerOutput(CurrentJob->Output, CurrentJob->Key.ToString(), SingleJob->GetSingleShaderJob()->Input.GenerateDebugInfo());
+	GetResourceCode()->AddShaderCompilerOutput(CurrentJob->Output, CurrentJob->Key, SingleJob->GetSingleShaderJob()->Input.GenerateDebugInfo());
 
 	FShader* Shader = nullptr;
 
diff --git a/Engine/Plugins/Runtime/ComputeFramework/Source/ComputeFramework/Private/ComputeKernelShaderMap.cpp b/Engine/Plugins/Runtime/ComputeFramework/Source/ComputeFramework/Private/ComputeKernelShaderMap.cpp
--- a/Engine/Plugins/Runtime/ComputeFramework/Source/ComputeFramework/Private/ComputeKernelShaderMap.cpp
+++ b/Engine/Plugins/Runtime/ComputeFramework/Source/ComputeFramework/Private/ComputeKernelShaderMap.cpp
@@ -436,7 +436,7 @@ FShader* FComputeKernelShaderMap::ProcessCompilationResultsForSingleJob(FShaderC
 {
 	check(CurrentJob.Id == CompilingId);
 
-	GetResourceCode()->AddShaderCompilerOutput(CurrentJob.Output, CurrentJob.Key.ToString(), CurrentJob.Input.GenerateDebugInfo());
+	GetResourceCode()->AddShaderCompilerOutput(CurrentJob.Output, CurrentJob.Key, CurrentJob.Input.GenerateDebugInfo());
 
 	FShader* Shader = nullptr;
 
diff --git a/Engine/Shaders/Public/ShaderVersion.ush b/Engine/Shaders/Public/ShaderVersion.ush
--- a/Engine/Shaders/Public/ShaderVersion.ush
+++ b/Engine/Shaders/Public/ShaderVersion.ush
@@ -3,4 +3,4 @@
 // Each time the console command is executed it generates a new GUID. As this file is included
 // in Platform.ush (which should be included in any shader) it allows to invalidate the shader DDC.
 // 
-#pragma message("UESHADERMETADATA_VERSION 9FE56796-34AB-4080-9B31-5FD7214FAF3B")
+#pragma message("UESHADERMETADATA_VERSION 3F6B6BB8-21E5-42E7-BB92-4105AC3B7F02")
diff --git a/Engine/Source/Editor/UnrealEd/Private/CookOnTheFlyServer.cpp b/Engine/Source/Editor/UnrealEd/Private/CookOnTheFlyServer.cpp
--- a/Engine/Source/Editor/UnrealEd/Private/CookOnTheFlyServer.cpp
+++ b/Engine/Source/Editor/UnrealEd/Private/CookOnTheFlyServer.cpp
@@ -9814,6 +9814,7 @@ void UCookOnTheFlyServer::CookByTheBookFinishedInternal()
 			// Save shader code map
 			SaveShaderLibrary(TargetPlatform, LibraryName);
 			CreatePipelineCache(TargetPlatform, LibraryName);
+			FShaderLibraryCooker::DumpShaderTypeStats(GShaderCompilingManager->GetAbsoluteShaderDebugInfoDirectory(), ConvertToFullSandboxPath(GetMetadataDirectory(), true, PlatformNameString));
 		}
 
 		if (FParse::Param(FCommandLine::Get(), TEXT("fastcook")))
diff --git a/Engine/Source/Editor/UnrealEd/Private/Cooker/CookIncrementalVersion.cpp b/Engine/Source/Editor/UnrealEd/Private/Cooker/CookIncrementalVersion.cpp
--- a/Engine/Source/Editor/UnrealEd/Private/Cooker/CookIncrementalVersion.cpp
+++ b/Engine/Source/Editor/UnrealEd/Private/Cooker/CookIncrementalVersion.cpp
@@ -8,6 +8,5 @@ namespace UE::Cook
 {
 
 // Change CookIncrementalVersion to a new guid when all packages in an incremental cook need to be invalidated
-FGuid CookIncrementalVersion( 0x06A2576D, 0x2C654CE5, 0x86A9A40C, 0xB74471B7 );
-
+FGuid CookIncrementalVersion( 0xD6B3741E, 0x72CD49EF, 0xBEDEFBA9, 0xDB7C61CF );
 }
diff --git a/Engine/Source/Runtime/Engine/Private/Materials/MaterialShader.cpp b/Engine/Source/Runtime/Engine/Private/Materials/MaterialShader.cpp
--- a/Engine/Source/Runtime/Engine/Private/Materials/MaterialShader.cpp
+++ b/Engine/Source/Runtime/Engine/Private/Materials/MaterialShader.cpp
@@ -2746,7 +2746,7 @@ FShader* FMaterialShaderMap::ProcessCompilationResultsForSingleJob(FShaderCompil
 		*LexToString(CurrentJob.Input.Target.GetPlatform()),
 		*LexToString(GetShaderPlatform()));
 
-	GetResourceCode()->AddShaderCompilerOutput(CurrentJob.Output, CurrentJob.Key.ToString(), SingleJob->Input.GenerateDebugInfo());
+	GetResourceCode()->AddShaderCompilerOutput(CurrentJob.Output, CurrentJob.Key, SingleJob->Input.GenerateDebugInfo());
 
 #if ALLOW_SHADERMAP_DEBUG_DATA
 	CompileTime += SingleJob->Output.CompileTime;
diff --git a/Engine/Source/Runtime/Engine/Private/ShaderCompiler/ShaderCompilerEditor.cpp b/Engine/Source/Runtime/Engine/Private/ShaderCompiler/ShaderCompilerEditor.cpp
--- a/Engine/Source/Runtime/Engine/Private/ShaderCompiler/ShaderCompilerEditor.cpp
+++ b/Engine/Source/Runtime/Engine/Private/ShaderCompiler/ShaderCompilerEditor.cpp
@@ -149,7 +149,7 @@ FShader* FGlobalShaderTypeCompiler::FinishCompileShader(const FGlobalShaderType*
 		EShaderPlatform Platform = CurrentJob.Input.Target.GetPlatform();
 		FGlobalShaderMapSection* Section = GGlobalShaderMap[Platform]->FindOrAddSection(ShaderType);
 
-		Section->GetResourceCode()->AddShaderCompilerOutput(CurrentJob.Output, CurrentJob.Key.ToString(), CurrentJob.Input.GenerateDebugInfo());
+		Section->GetResourceCode()->AddShaderCompilerOutput(CurrentJob.Output, CurrentJob.Key, CurrentJob.Input.GenerateDebugInfo());
 
 		if (ShaderPipelineType && !ShaderPipelineType->ShouldOptimizeUnusedOutputs(CurrentJob.Input.Target.GetPlatform()))
 		{
diff --git a/Engine/Source/Runtime/RenderCore/Private/ShaderCodeArchive.cpp b/Engine/Source/Runtime/RenderCore/Private/ShaderCodeArchive.cpp
--- a/Engine/Source/Runtime/RenderCore/Private/ShaderCodeArchive.cpp
+++ b/Engine/Source/Runtime/RenderCore/Private/ShaderCodeArchive.cpp
@@ -191,6 +191,10 @@ bool FSerializedShaderArchive::FindOrAddShader(const FSHAHash& Hash, int32& OutI
 		ShaderEntries.AddDefaulted();
 		check(ShaderEntries.Num() == ShaderHashes.Num());
 		ShaderHashTable.Add(Key, OutIndex);
+#if WITH_EDITOR
+		ShaderTypes.AddDefaulted();
+		checkf(ShaderTypes.Num() == ShaderEntries.Num(), TEXT("Mismatch in number of shader (%d) and shader type (%d) entries when adding shader to archive"), ShaderEntries.Num(), ShaderTypes.Num());
+#endif
 		return true;
 	}
 
@@ -205,6 +209,10 @@ void FSerializedShaderArchive::RemoveLastAddedShader()
 	ShaderHashTable.Remove(Key, ShaderIndex);
 	ShaderHashes.RemoveAt(ShaderIndex);
 	ShaderEntries.RemoveAt(ShaderIndex);
+#if WITH_EDITOR
+	ShaderTypes.RemoveAt(ShaderIndex);
+	checkf(ShaderTypes.Num() == ShaderEntries.Num(), TEXT("Mismatch in number of shader (%d) and shader type (%d) entries when removing last added shader from archive"), ShaderEntries.Num(), ShaderTypes.Num());
+#endif
 }
 
 #if WITH_EDITOR
@@ -216,6 +224,8 @@ FCbWriter& operator<<(FCbWriter& Writer, const FSerializedShaderArchive& Archive
 		const_cast<FSerializedShaderArchive&>(Archive).Serialize(SerializeArchive);
 	}
 
+	checkf(Archive.ShaderTypes.Num() == Archive.ShaderEntries.Num(), TEXT("Mismatch in number of shader (%d) and shader type (%d) entries when serializing shader archive"), Archive.ShaderEntries.Num(), Archive.ShaderTypes.Num());
+
 	Writer.BeginObject();
 	{
 		Writer << "SerializedBytes";
@@ -235,6 +245,18 @@ FCbWriter& operator<<(FCbWriter& Writer, const FSerializedShaderArchive& Archive
 			Writer.EndArray();
 		}
 		Writer.EndArray();
+
+		Writer.BeginArray("ShaderTypes");
+		for (const FSerializedShaderArchive::FShaderTypeHashes& SingleShaderTypes : Archive.ShaderTypes)
+		{
+			Writer.BeginArray();
+			for (uint64 ShaderTypeHash : SingleShaderTypes.Data)
+			{
+				Writer << ShaderTypeHash;
+			}
+			Writer.EndArray();
+		}
+		Writer.EndArray();
 	}
 	Writer.EndObject();
 	return Writer;
@@ -253,30 +275,58 @@ bool LoadFromCompactBinary(FCbFieldView Field, FSerializedShaderArchive& OutArch
 		}
 	}
 
-	FCbFieldView ShaderCodeToAssetsField = Field["ShaderCodeToAssets"];
-	// Map size is array size divided by two because pairs are written as successive elements
-	int32 NumShaderCodeToAssets = ShaderCodeToAssetsField.AsArrayView().Num()/2;
-	bool bOk = !ShaderCodeToAssetsField.HasError();
-	OutArchive.ShaderCodeToAssets.Empty(NumShaderCodeToAssets);
-	FCbFieldViewIterator It = ShaderCodeToAssetsField.CreateViewIterator();
-	while (It)
+	bool bOk = true;
+
 	{
-		FSHAHash ShaderMapHash;
-		if (!LoadFromCompactBinary(*It++, ShaderMapHash))
+		FCbFieldView ShaderCodeToAssetsField = Field["ShaderCodeToAssets"];
+		// Map size is array size divided by two because pairs are written as successive elements
+		int32 NumShaderCodeToAssets = ShaderCodeToAssetsField.AsArrayView().Num() / 2;
+		bOk = !ShaderCodeToAssetsField.HasError();
+		OutArchive.ShaderCodeToAssets.Empty(NumShaderCodeToAssets);
+		FCbFieldViewIterator It = ShaderCodeToAssetsField.CreateViewIterator();
+		while (It)
 		{
-			bOk = false;
-			continue;
+			FSHAHash ShaderMapHash;
+			if (!LoadFromCompactBinary(*It++, ShaderMapHash))
+			{
+				bOk = false;
+				continue;
+			}
+			FShaderMapAssetPaths& Paths = OutArchive.ShaderCodeToAssets.FindOrAdd(ShaderMapHash);
+			FCbFieldView AssetNameArrayField = *It++;
+			Paths.Reserve((*It).AsArrayView().Num());
+			bOk = (!AssetNameArrayField.HasError()) && bOk;
+			for (FCbFieldView AssetNameField : AssetNameArrayField)
+			{
+				FName AssetName;
+				bOk = LoadFromCompactBinary(AssetNameField, AssetName) && bOk;
+				Paths.Add(AssetName);
+			}
 		}
-		FShaderMapAssetPaths& Paths = OutArchive.ShaderCodeToAssets.FindOrAdd(ShaderMapHash);
-		FCbFieldView AssetNameArrayField = *It++;
-		Paths.Reserve((*It).AsArrayView().Num());
-		bOk = (!AssetNameArrayField.HasError()) & bOk;
-		for (FCbFieldView AssetNameField : AssetNameArrayField)
+	}
+
+	{
+		FCbFieldView ShaderTypesField = Field["ShaderTypes"];
+		bOk = !ShaderTypesField.HasError() && bOk;
+		OutArchive.ShaderTypes.Empty(ShaderTypesField.AsArrayView().Num());
+		FCbFieldViewIterator It = ShaderTypesField.CreateViewIterator();
+		while (It)
 		{
-			FName AssetName;
-			bOk = LoadFromCompactBinary(AssetNameField, AssetName) & bOk;
-			Paths.Add(AssetName);
+			bOk = !It->HasError() && bOk;
+			FSerializedShaderArchive::FShaderTypeHashes& ShaderTypeHashes = OutArchive.ShaderTypes.AddDefaulted_GetRef();
+			FCbArrayView ArrayView = It.AsArrayView();
+			ShaderTypeHashes.Data.Reserve(ArrayView.Num());
+			for (FCbFieldViewIterator ArrayIt = It.AsArrayView().CreateViewIterator(); ArrayIt; ++ArrayIt)
+			{
+				ShaderTypeHashes.Data.Add(ArrayIt.AsUInt64());
+			}
+			It++;
 		}
+
+		checkf(OutArchive.ShaderTypes.Num() == OutArchive.ShaderEntries.Num(), 
+			TEXT("Mismatch in number of shader (%d) and shader type (%d) entries when loading shader archive from compact binary"),
+			OutArchive.ShaderEntries.Num(), 
+			OutArchive.ShaderTypes.Num());
 	}
 
 	return bOk;
@@ -813,6 +863,9 @@ void FSerializedShaderArchive::CreateAsChunkFrom(const FSerializedShaderArchive&
 
 						// copy the entry as is
 						ShaderEntries[ShaderIndex] = Parent.ShaderEntries[ParentShaderIndex];
+#if WITH_EDITOR
+						ShaderTypes[ShaderIndex] = Parent.ShaderTypes[ParentShaderIndex];
+#endif
 					}
 					ShaderIndices[ShaderMapDescriptor.ShaderIndicesOffset + ShaderIdx] = ShaderIndex;
 				}
diff --git a/Engine/Source/Runtime/RenderCore/Private/ShaderCodeLibrary.cpp b/Engine/Source/Runtime/RenderCore/Private/ShaderCodeLibrary.cpp
--- a/Engine/Source/Runtime/RenderCore/Private/ShaderCodeLibrary.cpp
+++ b/Engine/Source/Runtime/RenderCore/Private/ShaderCodeLibrary.cpp
@@ -32,6 +32,7 @@ ShaderCodeLibrary.cpp: Bound shader state cache implementation.
 #include "PipelineFileCache.h"
 #include "ProfilingDebugging/AssetMetadataTrace.h"
 #include "ProfilingDebugging/LoadTimeTracker.h"
+#include "ProfilingDebugging/DiagnosticTable.h"
 #include "RenderingThread.h"
 #include "Shader.h"
 #include "ShaderCodeArchive.h"
@@ -71,6 +72,7 @@ static uint32 GShaderPipelineArchiveVersion = 1;
 
 static FString ShaderExtension = TEXT(".ushaderbytecode");
 static FString ShaderAssetInfoExtension = TEXT(".assetinfo.json");
+static FString ShaderTypeInfoExtension = TEXT(".stinfo");
 static FString StableExtension = TEXT(".shk");
 static FString PipelineExtension = TEXT(".ushaderpipelines");
 
@@ -374,6 +376,11 @@ static FString GetShaderAssetInfoFilename(const FString& BaseDir, const FString&
 	return BaseDir / FString::Printf(TEXT("ShaderAssetInfo-%s-"), *LibraryName) + Platform.ToString() + ShaderAssetInfoExtension;
 }
 
+static FString GetShaderTypeInfoFilename(const FString& BaseDir, const FString& LibraryName, FName Platform)
+{
+	return BaseDir / FString::Printf(TEXT("ShaderTypeInfo-%s-"), *LibraryName) + Platform.ToString() + ShaderTypeInfoExtension;
+}
+
 static FString GetShaderDebugFolder(const FString& BaseDir, const FString& LibraryName, FName Platform)
 {
 	return BaseDir / FString::Printf(TEXT("ShaderDebug-%s-"), *LibraryName) + Platform.ToString();
@@ -1528,11 +1535,14 @@ int32 FShaderMapResource_SharedCode::GetLibraryId() const
 #if WITH_EDITOR
 struct FShaderCodeStats
 {
-	int64 ShadersSize;
-	int64 ShadersUniqueSize;
-	int32 NumShaders;
-	int32 NumUniqueShaders;
+	struct FPerTypeStats
+	{
+		int64 ShadersSize = 0;
+		int32 NumShaders = 0;
+	};
 	int32 NumShaderMaps;
+
+	TMap<uint64, FPerTypeStats> ShaderTypeStats;
 };
 
 struct FEditorShaderCodeArchive
@@ -1596,9 +1606,10 @@ struct FEditorShaderCodeArchive
 			UE_LOG(LogShaderLibrary, Warning, TEXT("Shadermap %s does not have assets associated with it, library layout may be inconsistent between builds"), *Code->ResourceHash.ToString());
 		}
 
+		const int32 NumShaders = Code->ShaderCodeResources.Num();
 		if (SerializedShaders.FindOrAddShaderMap(Code->ResourceHash, ShaderMapIndex, &AssociatedAssets))
 		{
-			const int32 NumShaders = Code->ShaderCodeResources.Num();
+			CodeStats.NumShaderMaps++;
 			FShaderMapEntry& ShaderMapEntry = SerializedShaders.ShaderMapEntries[ShaderMapIndex];
 			ShaderMapEntry.NumShaders = NumShaders;
 			ShaderMapEntry.ShaderIndicesOffset = SerializedShaders.ShaderIndices.AddZeroed(NumShaders);
@@ -1617,19 +1628,30 @@ struct FEditorShaderCodeArchive
 					check(SerializedShaderEntry.Size > 0);
 					ShaderCode.Add(CodeBuffer);
 					check(ShaderCode.Num() == SerializedShaders.ShaderEntries.Num());
-
-					CodeStats.NumUniqueShaders++;
-					CodeStats.ShadersUniqueSize += CodeBuffer.GetSize();
 				}
-				CodeStats.ShadersSize += Code->ShaderCodeResources[i].GetCodeView().NumBytes();
 				SerializedShaders.ShaderIndices[ShaderMapEntry.ShaderIndicesOffset + i] = ShaderIndex;
 			}
+		}
 
-			// for total shaders, only count shaders when we're adding a new shadermap. AddShaderCode() for the same shadermap can be called several times during
-			// the cook because of serialization path being reused for other purposes than actual saving, so counting them every time artificially inflates number of shaders.
-			CodeStats.NumShaders += Code->ShaderCodeResources.Num();
-			CodeStats.NumShaderMaps++;
+		for (int32 i = 0; i < NumShaders; ++i)
+		{
+			for (uint64 ShaderTypeHash : Code->ShaderEditorOnlyDataEntries[i].ShaderTypeHashes)
+			{
+				int32 ShaderIndex = SerializedShaders.FindShader(Code->ShaderHashes[i]);
+			
+				TArray<uint64>& SerializedShaderTypes = SerializedShaders.ShaderTypes[ShaderIndex].Data;
+				const int32 HashIndex = Algo::LowerBound(SerializedShaderTypes, ShaderTypeHash);
+				if (HashIndex >= SerializedShaderTypes.Num() ||SerializedShaderTypes[HashIndex] != ShaderTypeHash)
+				{
+					SerializedShaderTypes.Insert(ShaderTypeHash, HashIndex);
+					FShaderCodeStats::FPerTypeStats& TypeCodeStats = CodeStats.ShaderTypeStats.FindOrAdd(ShaderTypeHash);
+					int32 CodeSize = SerializedShaders.ShaderEntries[ShaderIndex].Size;
+					TypeCodeStats.NumShaders++;
+					TypeCodeStats.ShadersSize += CodeSize;
+				}
+			}
 		}
+
 		// always mark the shadermap dirty, because it might have gotten new asset associations
 		MarkShaderMapDirty(ShaderMapIndex);
 		return ShaderMapIndex;
@@ -1773,6 +1795,16 @@ struct FEditorShaderCodeArchive
 						TargetShaderEntry.Offset = INDEX_NONE;
 					}
 				}
+
+				for (uint64 ShaderTypeHash : SourceArchive.ShaderTypes[SourceShaderIndex].Data)
+				{
+					TArray<uint64>& TargetShaderTypes = TargetArchive.ShaderTypes[TargetShaderIndex].Data;
+					const int32 HashIndex = Algo::LowerBound(TargetShaderTypes, ShaderTypeHash);
+					if (HashIndex >= TargetShaderTypes.Num() || TargetShaderTypes[HashIndex] != ShaderTypeHash)
+					{
+						TargetShaderTypes.Insert(ShaderTypeHash, HashIndex);
+					}
+				}
 			}
 		}
 
@@ -1793,21 +1825,32 @@ struct FEditorShaderCodeArchive
 			const FSHAHash& SourceShaderHash = SourceArchive.ShaderHashes[SourceShaderIndex];
 			int32 TargetShaderIndex = INDEX_NONE;
 			const bool bShaderIsNew = TargetArchive.FindOrAddShader(SourceShaderHash, TargetShaderIndex);
+			const TArray<uint64>& ShaderTypeHashes = SourceArchive.ShaderTypes[SourceShaderIndex].Data;
+			const FShaderCodeEntry& SourceShaderEntry = SourceArchive.ShaderEntries[SourceShaderIndex];
+
+			TArray<uint64>& ExistingShaderTypeHashes = TargetArchive.ShaderTypes[TargetShaderIndex].Data;
+			for (uint64 ShaderTypeHash : ShaderTypeHashes)
+			{
+				const int32 HashIndex = Algo::LowerBound(ExistingShaderTypeHashes, ShaderTypeHash);
+				if (HashIndex >= ExistingShaderTypeHashes.Num() || ExistingShaderTypeHashes[HashIndex] != ShaderTypeHash)
+				{
+					ExistingShaderTypeHashes.Insert(ShaderTypeHash, HashIndex);
+					FShaderCodeStats::FPerTypeStats& TypeCodeStats = CodeStats.ShaderTypeStats.FindOrAdd(ShaderTypeHash);
+					TypeCodeStats.NumShaders++;
+					TypeCodeStats.ShadersSize += SourceShaderEntry.Size;
+				}
+			}
+
 			if (!bShaderIsNew)
 			{
 				continue;
 			}
-			check(TargetShaderIndex == TargetArchive.ShaderEntries.Num() - 1 &&
-				TargetShaderCodes.Num() == TargetArchive.ShaderEntries.Num() - 1);
+			check(TargetShaderIndex == TargetArchive.ShaderEntries.Num() - 1 && TargetShaderCodes.Num() == TargetArchive.ShaderEntries.Num() - 1);
 
-			const FShaderCodeEntry& SourceShaderEntry = SourceArchive.ShaderEntries[SourceShaderIndex];
 			FShaderCodeEntry& TargetShaderEntry = TargetArchive.ShaderEntries[TargetShaderIndex];
 			TargetShaderEntry = SourceShaderEntry;
 			TargetShaderEntry.Offset = 0;
 
-			CodeStats.NumUniqueShaders++;
-			CodeStats.ShadersUniqueSize += SourceShaderEntry.Size;
-
 			if (SourceShaderEntry.Offset == INDEX_NONE)
 			{
 				UE_LOG(LogShaderLibrary, Error, TEXT("ShaderMapLibrary transfer received from a remote machine has incomplete record for Shader %s. ")
@@ -1866,8 +1909,6 @@ struct FEditorShaderCodeArchive
 				// Every shader in the SourceArchive should have already been added by the loop above over SourceArchive.ShaderHashes
 				check(!bShaderIsNew);
 			}
-
-			CodeStats.NumShaders += TargetEntry.NumShaders; // Sum of shader counts used by each ShaderMap, without removing duplicates
 			CodeStats.NumShaderMaps++;
 		}
 		return bOk;
@@ -1921,6 +1962,7 @@ struct FEditorShaderCodeArchive
 				{
 					const FShaderCodeEntry& OtherShaderEntry = OtherArchive.SerializedShaders.ShaderEntries[OtherShaderIndex];
 					SerializedShaders.ShaderEntries[ShaderIndex] = OtherShaderEntry;
+					SerializedShaders.ShaderTypes[ShaderIndex] = OtherArchive.SerializedShaders.ShaderTypes[OtherShaderIndex];
 
 					const FSharedBuffer& OtherShaderCodeEntry = OtherArchive.ShaderCode[OtherShaderIndex];
 					check(OtherShaderCodeEntry.GetSize() > 0);
@@ -1935,7 +1977,8 @@ struct FEditorShaderCodeArchive
 		return ShaderMapIndex;
 	}
 
-	int32 AddShaderCode(int32 OtherShaderMapIndex,
+	int32 AddShaderCode(FShaderCodeStats& CodeStats,
+		int32 OtherShaderMapIndex,
 		const FSerializedShaderArchive& OtherShaders,
 		int64 OtherShaderCodeOffset,
 		FArchive& Ar)
@@ -1946,6 +1989,8 @@ struct FEditorShaderCodeArchive
 		bool bIsNew = SerializedShaders.FindOrAddShaderMap(OtherShaderMapHash, ShaderMapIndex, OtherAssets);
 		if (bIsNew)
 		{
+			CodeStats.NumShaderMaps++;
+
 			const FShaderMapEntry& OtherEntry = OtherShaders.ShaderMapEntries[OtherShaderMapIndex];
 			FShaderMapEntry& ShaderMapEntry = SerializedShaders.ShaderMapEntries[ShaderMapIndex];
 			ShaderMapEntry.NumShaders = OtherEntry.NumShaders;
@@ -1971,6 +2016,20 @@ struct FEditorShaderCodeArchive
 					Ar.Serialize(Code.GetData(), ReadSize);
 					ShaderCode.Add(Code.MoveToShared());
 				}
+
+				TArray<uint64>& SerializedShaderTypes = SerializedShaders.ShaderTypes[ShaderIndex].Data;
+				for (uint64 OtherShaderTypeHash : OtherShaders.ShaderTypes[OtherShaderIndex].Data)
+				{
+					const int32 HashIndex = Algo::LowerBound(SerializedShaderTypes, OtherShaderTypeHash);
+					if (HashIndex >= SerializedShaderTypes.Num() || SerializedShaderTypes[HashIndex] != OtherShaderTypeHash)
+					{
+						SerializedShaderTypes.Insert(OtherShaderTypeHash, HashIndex);
+						FShaderCodeStats::FPerTypeStats& TypeCodeStats = CodeStats.ShaderTypeStats.FindOrAdd(OtherShaderTypeHash);
+						int32 CodeSize = SerializedShaders.ShaderEntries[ShaderIndex].Size;
+						TypeCodeStats.NumShaders++;
+						TypeCodeStats.ShadersSize += CodeSize;
+					}
+				}
 				SerializedShaders.ShaderIndices[ShaderMapEntry.ShaderIndicesOffset + i] = ShaderIndex;
 			}
 		}
@@ -2030,9 +2089,11 @@ struct FEditorShaderCodeArchive
 		return bOK;
 	}
 
-	void AddShaderCodeLibraryByName(const FString& BaseDir, const FString& InLibraryName)
+	void AddShaderCodeLibraryByName(const FString& BaseDir, const FString& InLibraryName, FShaderCodeStats& CodeStats)
 	{
-		if (FArchive* PrevCookedAr = CreateShaderFileReader(*GetCodeArchiveFilename(BaseDir, InLibraryName, FormatName)))
+		EShaderPlatform Platform = ShaderFormatToLegacyShaderPlatform(FormatName);
+		FString CodeArchiveFileName = GetCodeArchiveFilename(BaseDir, InLibraryName, FormatName);
+		if (FArchive* PrevCookedAr = CreateShaderFileReader(*CodeArchiveFileName))
 		{
 			uint32 Version = 0;
 			*PrevCookedAr << Version;
@@ -2044,27 +2105,49 @@ struct FEditorShaderCodeArchive
 				*PrevCookedAr << PrevCookedShaders;
 
 				// check if it also contains the asset info file
-				TUniquePtr<FArchive> ShaderAssetInfoReader(CreateShaderFileReader(*GetShaderAssetInfoFilename(BaseDir, InLibraryName, FormatName)));
+				FString AssetInfoFileName = GetShaderAssetInfoFilename(BaseDir, InLibraryName, FormatName);
+				TUniquePtr<FArchive> ShaderAssetInfoReader(CreateShaderFileReader(*AssetInfoFileName));
 				if (PrevCookedShaders.LoadAssetInfo(ShaderAssetInfoReader.Get()))
 				{
 					UE_LOG(LogShaderLibrary, Display, TEXT("Loaded asset info %s for the shader library %s: %d entries"),
-						*GetShaderAssetInfoFilename(BaseDir, InLibraryName, FormatName),
-						*GetCodeArchiveFilename(BaseDir, InLibraryName, FormatName),
+						*AssetInfoFileName,
+						*CodeArchiveFileName,
 						PrevCookedShaders.ShaderCodeToAssets.Num()
 					);
 				}
 				else
 				{
 					UE_LOG(LogShaderLibrary, Warning, TEXT("Could not find or load asset info %s for the shader library %s"),
-						*GetShaderAssetInfoFilename(BaseDir, InLibraryName, FormatName),
-						*GetCodeArchiveFilename(BaseDir, InLibraryName, FormatName)
+						*AssetInfoFileName,
+						*CodeArchiveFileName
+					);
+				}
+
+				// check if it also contains the shader type info file
+				FString TypeInfoFileName = GetShaderTypeInfoFilename(BaseDir, InLibraryName, FormatName);
+				TUniquePtr<FArchive> ShaderTypeInfoReader(CreateShaderFileReader(*TypeInfoFileName));
+				if (ShaderTypeInfoReader)
+				{
+					*ShaderTypeInfoReader << PrevCookedShaders.ShaderTypes;
+					check(PrevCookedShaders.ShaderTypes.Num() == PrevCookedShaders.ShaderEntries.Num());
+					UE_LOG(LogShaderLibrary, Display, TEXT("Loaded shader type info %s for the shader library %s: %d entries"),
+						*TypeInfoFileName,
+						*CodeArchiveFileName,
+						PrevCookedShaders.ShaderTypes.Num()
+					);
+				}
+				else
+				{
+					UE_LOG(LogShaderLibrary, Warning, TEXT("Could not load shader type info %s for the shader library %s"),
+						*TypeInfoFileName,
+						*CodeArchiveFileName
 					);
 				}
 
 				int64 PrevCookedShadersCodeStart = PrevCookedAr->Tell();
 				for (int32 PrevShaderMapIndex = 0; PrevShaderMapIndex < PrevCookedShaders.ShaderMapEntries.Num(); ++PrevShaderMapIndex)
 				{
-					AddShaderCode(PrevShaderMapIndex, PrevCookedShaders, PrevCookedShadersCodeStart, *PrevCookedAr);
+					AddShaderCode(CodeStats, PrevShaderMapIndex, PrevCookedShaders, PrevCookedShadersCodeStart, *PrevCookedAr);
 				}
 			}
 
@@ -2073,12 +2156,12 @@ struct FEditorShaderCodeArchive
 		}
 	}
 
-	void AddShaderCodeLibraryFromDirectory(const FString& BaseDir)
+	void AddShaderCodeLibraryFromDirectory(const FString& BaseDir, FShaderCodeStats& CodeStats)
 	{
-		AddShaderCodeLibraryByName(BaseDir, LibraryName);
+		AddShaderCodeLibraryByName(BaseDir, LibraryName, CodeStats);
 	}
 
-	void AddExistingShaderCodeLibrary(FString const& OutputDir)
+	void AddExistingShaderCodeLibrary(FString const& OutputDir, FShaderCodeStats& CodeStats)
 	{
 		check(LibraryName.Len() > 0);
 		static const FRegexPattern ShaderCodePattern(ShaderCodePatternStr);
@@ -2099,15 +2182,15 @@ struct FEditorShaderCodeArchive
 
 				if (FoundLibraryName.StartsWith(LibraryName) && FoundFormatName.Equals(FormatNameStr, ESearchCase::IgnoreCase))
 				{
-					AddShaderCodeLibraryByName(OutputDir, FoundLibraryName);
+					AddShaderCodeLibraryByName(OutputDir, FoundLibraryName, CodeStats);
 				}
 			}
 		}
 	}
 
-	void FinishPopulate(const FString& OutputDir)
+	void FinishPopulate(const FString& OutputDir, FShaderCodeStats& CodeStats)
 	{
-		AddExistingShaderCodeLibrary(OutputDir);
+		AddExistingShaderCodeLibrary(OutputDir, CodeStats);
 	}
 
 	bool SaveToDisk(const FString& OutputDir, const FString& MetaOutputDir, bool bSaveOnlyAssetInfo = false, TArray<FString>* OutputFilenames = nullptr)
@@ -2140,6 +2223,7 @@ struct FEditorShaderCodeArchive
 			// Write to a intermediate file
 			FString IntermediateFormatPath = GetShaderCodeFilename(FPaths::ProjectSavedDir() / TEXT("Shaders") / FormatName.ToString(), LibraryName, FormatName);
 			FString AssetInfoIntermediatePath = GetShaderAssetInfoFilename(FPaths::ProjectSavedDir() / TEXT("Shaders") / FormatName.ToString(), LibraryName, FormatName);
+			FString TypeInfoIntermediatePath = GetShaderTypeInfoFilename(FPaths::ProjectSavedDir() / TEXT("Shaders") / FormatName.ToString(), LibraryName, FormatName);
 
 			// save the actual shader code code
 			if (!bSaveOnlyAssetInfo)
@@ -2177,6 +2261,27 @@ struct FEditorShaderCodeArchive
 						}
 					}
 				}
+
+				FArchive* ShaderTypeFileWriter = IFileManager::Get().CreateFileWriter(*TypeInfoIntermediatePath, FILEWRITE_NoFail);
+				if (ShaderTypeFileWriter)
+				{
+					*ShaderTypeFileWriter << SerializedShaders.ShaderTypes;
+					ShaderTypeFileWriter->Close();
+					delete ShaderTypeFileWriter;
+					
+					// Copy to output location - support for iterative native library cooking
+					if (!CopyFile(GetShaderTypeInfoFilename(OutputDir, LibraryName, FormatName), TypeInfoIntermediatePath, OutputFilenames))
+					{
+						bSuccess = false;
+					}
+					else if (MetaOutputDir.Len())
+					{
+						if (!CopyFile(GetShaderTypeInfoFilename(MetaOutputDir / TEXT("../ShaderLibrarySource"), LibraryName, FormatName), TypeInfoIntermediatePath, nullptr))
+						{
+							bSuccess = false;
+						}
+					}
+				}
 			}
 
 			// save asset info
@@ -2288,13 +2393,14 @@ struct FEditorShaderCodeArchive
 		if (bOK)
 		{
 			FEditorShaderCodeArchive OutLibrary(FormatName, bNeedsDeterministicOrder);
+			FShaderCodeStats OutStats;
 			OutLibrary.OpenLibrary(LibraryName);
 			OutLibrary.MakePatchLibrary(OldLibraries, NewLibrary);
 			bOK = OutLibrary.SerializedShaders.GetNumShaderMaps() > 0;
 			if (bOK)
 			{
 				FString Empty;
-				OutLibrary.FinishPopulate(OutDir);
+				OutLibrary.FinishPopulate(OutDir, OutStats);
 				bOK = OutLibrary.SaveToDisk(OutDir, Empty);
 				UE_CLOG(!bOK, LogShaderLibrary, Error, TEXT("Failed to save %s shader patch library %s, %s, %s"), bNativeFormat ? TEXT("native") : TEXT(""), *FormatName.ToString(), *LibraryName, *OutDir);
 				
@@ -3502,10 +3608,12 @@ public:
 		{
 			EShaderPlatform SPlatform = ShaderFormatToLegacyShaderPlatform(ShaderFormatName);
 
+			FShaderCodeStats& CodeStats = EditorShaderCodeStats[SPlatform];
+
 			FEditorShaderCodeArchive* CodeArchive = EditorShaderCodeArchive[SPlatform];
 			if (CodeArchive)
 			{
-				CodeArchive->FinishPopulate(ShaderCodeDir);
+				CodeArchive->FinishPopulate(ShaderCodeDir, CodeStats);
 			}
 
 			FEditorShaderStableInfo* StableArchive = EditorShaderStableInfo[SPlatform];
@@ -3645,26 +3753,68 @@ public:
 		return bOK;
 	}
 
-	void DumpShaderCodeStats()
+	void DumpShaderTypeStats(const FString& DebugInfoDir, const FString& MetaDataDir)
 	{
 		int32 PlatformId = 0;
-		for (const FShaderCodeStats& CodeStats : EditorShaderCodeStats)
+		for (const FShaderCodeStats& PlatformCodeStats : EditorShaderCodeStats)
 		{
-			if (CodeStats.NumShaders > 0)
+			if (PlatformCodeStats.NumShaderMaps > 0)
 			{
-				float UniqueSize = CodeStats.ShadersUniqueSize;
-				float UniqueSizeMB = FUnitConversion::Convert(UniqueSize, EUnit::Bytes, EUnit::Megabytes);
-				float TotalSize = CodeStats.ShadersSize;
-				float TotalSizeMB = FUnitConversion::Convert(TotalSize, EUnit::Bytes, EUnit::Megabytes);
-
-				UE_LOG(LogShaderLibrary, Display, TEXT(""));
-				UE_LOG(LogShaderLibrary, Display, TEXT("Shader Code Stats: %s"), *LegacyShaderPlatformToShaderFormat((EShaderPlatform)PlatformId).ToString());
-				UE_LOG(LogShaderLibrary, Display, TEXT("================="));
-				UE_LOG(LogShaderLibrary, Display, TEXT("Unique Shaders: %d, Total Shaders: %d, Unique Shadermaps: %d"), CodeStats.NumUniqueShaders, CodeStats.NumShaders, CodeStats.NumShaderMaps);
-				UE_LOG(LogShaderLibrary, Display, TEXT("Unique Shaders Size: %.2fmb, Total Shader Size: %.2fmb"), UniqueSizeMB, TotalSizeMB);
-				UE_LOG(LogShaderLibrary, Display, TEXT("================="));
-			}
+				FString PlatformName = FGenericDataDrivenShaderPlatformInfo::GetName((EShaderPlatform)PlatformId).ToString();
+				FString Filename = FPaths::Combine(*DebugInfoDir, *PlatformName, TEXT("ShaderTypeStats.csv"));
+				TMap<const FShaderType*, FShaderCodeStats::FPerTypeStats> CountersByTypeName;
+				for (const TPair<uint64, FShaderCodeStats::FPerTypeStats>& Pair : PlatformCodeStats.ShaderTypeStats)
+				{
+					const FShaderType* ShaderType = FindShaderTypeByName(FHashedName(Pair.Key));
+					CountersByTypeName.Add(ShaderType, Pair.Value);
+				}
+				
+				// sort first by the type of shadertype, then by the name (so all global, material, etc. shaders are grouped together in the CSV)
+				struct FShaderTypeSort
+				{
+					bool operator()(const FShaderType& A, const FShaderType& B) const
+					{
+						if (A.GetTypeForDynamicCast() == B.GetTypeForDynamicCast())
+						{
+							return FCString::Strcmp(A.GetName(), B.GetName()) < 0;
+						}
+						return A.GetTypeForDynamicCast() < B.GetTypeForDynamicCast();
+					}
+				};
+				CountersByTypeName.KeySort(FShaderTypeSort());
 
+				{
+					TUniquePtr<FArchive> ShaderTypeStatsCsv(IFileManager::Get().CreateFileWriter(*Filename));
+					if (ShaderTypeStatsCsv)
+					{
+						FDiagnosticTableWriterCSV StatWriter(ShaderTypeStatsCsv.Get());
+						StatWriter.AddColumn(TEXT("ShaderTypeName"));
+						StatWriter.AddColumn(TEXT("ShaderType"));
+						StatWriter.AddColumn(TEXT("Frequency"));
+						StatWriter.AddColumn(TEXT("NumShaders"));
+						StatWriter.AddColumn(TEXT("ShadersSize"));
+						StatWriter.CycleRow();
+
+						for (const TPair<const FShaderType*, FShaderCodeStats::FPerTypeStats>& Pair : CountersByTypeName)
+						{
+							StatWriter.AddColumn(Pair.Key->GetName());
+							StatWriter.AddColumn(*LexToString(Pair.Key->GetTypeForDynamicCast()));
+							StatWriter.AddColumn(GetShaderFrequencyString(Pair.Key->GetFrequency()));
+							StatWriter.AddColumn(TEXT("%d"), Pair.Value.NumShaders);
+							StatWriter.AddColumn(TEXT("%lld"), Pair.Value.ShadersSize);
+							StatWriter.CycleRow();
+						}
+					}
+					else
+					{
+						UE_LOG(LogShaderLibrary, Warning, TEXT("Failed to create shader type stats file '%s'"), *Filename);
+					}
+				}
+
+				// Add a copy of the type stats to the build metadata as well
+				FString MetaFilename = FString::Printf(TEXT("%s/ShaderTypeStats-%s.csv"), *MetaDataDir, *PlatformName);
+				IFileManager::Get().Copy(*MetaFilename, *Filename);
+			}
 			PlatformId++;
 		}
 	}
@@ -4089,8 +4239,6 @@ void FShaderLibraryCooker::Shutdown()
 {
 	if (FShaderLibrariesCollection::Impl)
 	{
-		//DumpShaderCodeStats();
-
 		delete FShaderLibrariesCollection::Impl;
 		FShaderLibrariesCollection::Impl = nullptr;
 	}
@@ -4185,11 +4333,11 @@ void FShaderLibraryCooker::AddShaderStableKeyValue(EShaderPlatform ShaderPlatfor
 	}
 }
 
-void FShaderLibraryCooker::DumpShaderCodeStats()
+void FShaderLibraryCooker::DumpShaderTypeStats(const FString& DebugInfoDir, const FString& MetaDataDir)
 {
 	if (FShaderLibrariesCollection::Impl)
 	{
-		FShaderLibrariesCollection::Impl->DumpShaderCodeStats();
+		FShaderLibrariesCollection::Impl->DumpShaderTypeStats(DebugInfoDir, MetaDataDir);
 	}
 }
 
@@ -4257,7 +4405,7 @@ bool FShaderLibraryCooker::MergeShaderCodeArchive(const TArray<FString>& CookedM
 	static const FRegexPattern StableInfoPattern(StableInfoPatternStr);
 
 	// Map of file name to the archive we are unioning
-	TMap<FString, FEditorShaderCodeArchive> ShaderCodeArchives;
+	TMap<FString, TPair<FEditorShaderCodeArchive, FShaderCodeStats>> ShaderCodeArchives;
 	TMap<FString, FEditorShaderStableInfo> ShaderStableInfos;
 
 	for (const FString& MetadataDir : CookedMetadataDirs)
@@ -4271,7 +4419,7 @@ bool FShaderLibraryCooker::MergeShaderCodeArchive(const TArray<FString>& CookedM
 		{
 			if (ShaderCodeArchives.Contains(ByteCodeFile))
 			{
-				ShaderCodeArchives[ByteCodeFile].AddShaderCodeLibraryFromDirectory(ShaderCodeDir);
+				ShaderCodeArchives[ByteCodeFile].Key.AddShaderCodeLibraryFromDirectory(ShaderCodeDir, ShaderCodeArchives[ByteCodeFile].Value);
 			}
 			else
 			{
@@ -4280,10 +4428,11 @@ bool FShaderLibraryCooker::MergeShaderCodeArchive(const TArray<FString>& CookedM
 				{
 					const FName ShaderFormat(*FindShaderFormat.GetCaptureGroup(2));
 					FEditorShaderCodeArchive NewCodeArchive(ShaderFormat, bNeedsDeterministicOrder);
+					FShaderCodeStats NewCodeStats;
 					NewCodeArchive.OpenLibrary(FindShaderFormat.GetCaptureGroup(1));
-					NewCodeArchive.AddShaderCodeLibraryFromDirectory(ShaderCodeDir);
+					NewCodeArchive.AddShaderCodeLibraryFromDirectory(ShaderCodeDir, NewCodeStats);
 
-					ShaderCodeArchives.FindOrAdd(ByteCodeFile, MoveTemp(NewCodeArchive));
+					ShaderCodeArchives.FindOrAdd(ByteCodeFile, { MoveTemp(NewCodeArchive), MoveTemp(NewCodeStats) });
 				}
 			}
 		}
@@ -4319,10 +4468,10 @@ bool FShaderLibraryCooker::MergeShaderCodeArchive(const TArray<FString>& CookedM
 	const FString OutShaderStableInfoDir = OutputDir / TEXT("PipelineCaches");
 
 	bool bSuccess = true;
-	for (TPair<FString, FEditorShaderCodeArchive>& ShaderCodeArchivePair : ShaderCodeArchives)
+	for (TPair<FString, TPair<FEditorShaderCodeArchive, FShaderCodeStats>>& ShaderCodeArchivePair : ShaderCodeArchives)
 	{
 		TArray<FString> CreatedFiles;
-		if (ShaderCodeArchivePair.Value.SaveToDisk(OutShaderCodeDir, TEXT(""), false, &CreatedFiles))
+		if (ShaderCodeArchivePair.Value.Key.SaveToDisk(OutShaderCodeDir, TEXT(""), false, &CreatedFiles))
 		{
 			OutWrittenFiles.Append(MoveTemp(CreatedFiles));
 		}
@@ -4332,7 +4481,7 @@ bool FShaderLibraryCooker::MergeShaderCodeArchive(const TArray<FString>& CookedM
 			bSuccess = false;
 		}
 		CreatedFiles.Reset();
-		if (ShaderCodeArchivePair.Value.SaveToDisk(OutContentDir, TEXT(""), false, &CreatedFiles))
+		if (ShaderCodeArchivePair.Value.Key.SaveToDisk(OutContentDir, TEXT(""), false, &CreatedFiles))
 		{
 			OutWrittenFiles.Append(MoveTemp(CreatedFiles));
 		}
diff --git a/Engine/Source/Runtime/RenderCore/Private/ShaderResource.cpp b/Engine/Source/Runtime/RenderCore/Private/ShaderResource.cpp
--- a/Engine/Source/Runtime/RenderCore/Private/ShaderResource.cpp
+++ b/Engine/Source/Runtime/RenderCore/Private/ShaderResource.cpp
@@ -21,6 +21,7 @@
 #include "RHI.h"
 #include "RHIResources.h"	// Access to FRHIRayTracingShader::RayTracingPayloadType requires this
 #include "ShaderCompilerCore.h"
+#include "ShaderCompilerJobTypes.h"
 #include "ShaderCore.h"
 #include "ShaderSerialization.h"
 #include "UObject/RenderingObjectVersion.h"
@@ -270,11 +271,13 @@ void FShaderMapResourceCode::FShaderEditorOnlyDataEntry::ConditionalSetSymbolBuf
 }
 #endif
 
-void FShaderMapResourceCode::AddShaderCompilerOutput(const FShaderCompilerOutput& Output, const FString& DebugName, FString DebugInfo)
+void FShaderMapResourceCode::AddShaderCompilerOutput(const FShaderCompilerOutput& Output, const FShaderCompileJobKey& Key, FString DebugInfo)
 {
 	TRACE_CPUPROFILER_EVENT_SCOPE(FShaderMapResourceCode::AddShaderCode);
 
 	const FSHAHash& InHash = Output.OutputHash;
+	FString DebugName = Key.ToString();
+	uint64 ShaderTypeHash = Key.ShaderType->GetHashedName().GetHash();
 	const int32 Index = Algo::LowerBound(ShaderHashes, InHash);
 	const FShaderCodeResource& CodeResource = Output.GetFinalizedCodeResource();
 	if (Index >= ShaderHashes.Num() || ShaderHashes[Index] != InHash)
@@ -283,7 +286,7 @@ void FShaderMapResourceCode::AddShaderCompilerOutput(const FShaderCompilerOutput
 
 #if WITH_EDITORONLY_DATA
 		// Output.Errors contains warnings in the case any exist (no errors since if there were the job would have failed)
-		AddEditorOnlyData(Index, DebugName, CodeResource.GetSymbolsBuffer(), Output.Errors, Output.ShaderStatistics, DebugInfo);
+		AddEditorOnlyData(Index, DebugName, ShaderTypeHash, CodeResource.GetSymbolsBuffer(), Output.Errors, Output.ShaderStatistics, DebugInfo);
 #endif
 
 		ShaderCodeResources.Insert(CodeResource, Index);
@@ -294,14 +297,14 @@ void FShaderMapResourceCode::AddShaderCompilerOutput(const FShaderCompilerOutput
 		// Output.Errors contains warnings in the case any exist (no errors since if there were the job would have failed)
 		// We append the warnings and deduplicate other data like DebugInfo for any additional jobs which resulted in the
 		// same bytecode for the sake of determinism in the results saved to DDC.
-		UpdateEditorOnlyData(Index, DebugName, CodeResource.GetSymbolsBuffer(), Output.Errors, DebugInfo);
+		UpdateEditorOnlyData(Index, DebugName, ShaderTypeHash, CodeResource.GetSymbolsBuffer(), Output.Errors, DebugInfo);
 		ValidateShaderStatisticsEditorOnlyData(Index, Output.ShaderStatistics);
 	}
 #endif
 }
 
 #if WITH_EDITORONLY_DATA
-void FShaderMapResourceCode::AddEditorOnlyData(int32 Index, const FString& DebugName, FCompressedBuffer InSymbols, TConstArrayView<FShaderCompilerError> InCompilerWarnings, const TArray<FGenericShaderStat>& ShaderStatistics, const FString& DebugInfo)
+void FShaderMapResourceCode::AddEditorOnlyData(int32 Index, const FString& DebugName, uint64 ShaderTypeHash, FCompressedBuffer InSymbols, TConstArrayView<FShaderCompilerError> InCompilerWarnings, const TArray<FGenericShaderStat>& ShaderStatistics, const FString& DebugInfo)
 {
 	FShaderEditorOnlyDataEntry& Entry = ShaderEditorOnlyDataEntries.InsertDefaulted_GetRef(Index);
 
@@ -309,10 +312,10 @@ void FShaderMapResourceCode::AddEditorOnlyData(int32 Index, const FString& Debug
 	check(Entry.ShaderStatistics.Num() == 0);
 	Entry.ShaderStatistics = ShaderStatistics;
 
-	UpdateEditorOnlyData(Index, DebugName, InSymbols, InCompilerWarnings, DebugInfo);
+	UpdateEditorOnlyData(Index, DebugName, ShaderTypeHash, InSymbols, InCompilerWarnings, DebugInfo);
 }
 
-void FShaderMapResourceCode::UpdateEditorOnlyData(int32 Index, const FString& DebugName, FCompressedBuffer InSymbols, TConstArrayView<FShaderCompilerError> InCompilerWarnings, const FString& DebugInfo)
+void FShaderMapResourceCode::UpdateEditorOnlyData(int32 Index, const FString& DebugName, uint64 ShaderTypeHash, FCompressedBuffer InSymbols, TConstArrayView<FShaderCompilerError> InCompilerWarnings, const FString& DebugInfo)
 {
 	FShaderEditorOnlyDataEntry& Entry = ShaderEditorOnlyDataEntries[Index];
 
@@ -334,6 +337,12 @@ void FShaderMapResourceCode::UpdateEditorOnlyData(int32 Index, const FString& De
 			Entry.CompilerWarnings.Insert(ModifiedWarning, WarningIndex);
 		}
 	}
+	
+	const int32 HashIndex = Algo::LowerBound(Entry.ShaderTypeHashes, ShaderTypeHash);
+	if (HashIndex >= Entry.ShaderTypeHashes.Num() || Entry.ShaderTypeHashes[HashIndex] != ShaderTypeHash)
+	{
+		Entry.ShaderTypeHashes.Insert(ShaderTypeHash, HashIndex);
+	}
 }
 
 void FShaderMapResourceCode::ValidateShaderStatisticsEditorOnlyData(int32 Index, const TArray<FGenericShaderStat>& ShaderStatistics)
diff --git a/Engine/Source/Runtime/RenderCore/Public/Shader.h b/Engine/Source/Runtime/RenderCore/Public/Shader.h
--- a/Engine/Source/Runtime/RenderCore/Public/Shader.h
+++ b/Engine/Source/Runtime/RenderCore/Public/Shader.h
@@ -82,6 +82,7 @@ class FShaderType;
 class FVertexFactoryType;
 class FRHIShaderBindingLayout;
 struct FShaderCompiledShaderInitializerType;
+struct FShaderCompileJobKey;
 struct FShaderCompilerOutput;
 struct FShaderSerializeContext;
 using FShaderMapAssetPaths = TSet<FName>; // Copied from ShaderCodeLibrary.h
@@ -551,13 +552,17 @@ public:
 		 *  the same code, all warnings for each unique source will be reported).
 		 *  Does not contain errors since if there were any errors, this object wouldn't exist. */
 		TArray<FString> CompilerWarnings;
+		
+		/** An array of the hash of all FShaderType names which generate this particular shader (i.e. the return value
+		 * of FShaderType::GetHashedName().GetHash()). */
+		TArray<uint64> ShaderTypeHashes;
 
 		/** Generic, data-driven key/value pairs of statistics. */
 		TArray<FGenericShaderStat> ShaderStatistics;
 
 		friend FArchive& operator<<(FArchive& Ar, FShaderEditorOnlyDataEntry& Entry)
 		{
-			return Ar << Entry.DebugInfo << Entry.CompilerWarnings << Entry.ShaderStatistics;
+			return Ar << Entry.DebugInfo << Entry.CompilerWarnings << Entry.ShaderTypeHashes << Entry.ShaderStatistics;
 		}
 
 		PRAGMA_DISABLE_DEPRECATION_WARNINGS
@@ -588,13 +593,16 @@ public:
 
 	RENDERCORE_API uint32 GetSizeBytes() const;
 
-	RENDERCORE_API void AddShaderCompilerOutput(const FShaderCompilerOutput& Output, const FString& DebugName = FString(), FString DebugInfo = FString());
+	UE_DEPRECATED(5.6, "AddShaderCompilerOutput now accepts an FShaderCompileJobKey instead of DebugName")
+	inline void AddShaderCompilerOutput(const FShaderCompilerOutput& Output, const FString& DebugName = FString(), FString DebugInfo = FString()) {}
+
+	RENDERCORE_API void AddShaderCompilerOutput(const FShaderCompilerOutput& Output, const FShaderCompileJobKey& Key, FString DebugInfo = FString());
 
 	int32 FindShaderIndex(const FSHAHash& InHash) const;
 
 #if WITH_EDITORONLY_DATA
-	void AddEditorOnlyData(int32 Index, const FString& DebugName, FCompressedBuffer InSymbols, TConstArrayView<FShaderCompilerError> InCompilerWarnings, const TArray<FGenericShaderStat>& ShaderStatistics, const FString& DebugInfo = FString());
-	void UpdateEditorOnlyData(int32 Index, const FString& DebugName, FCompressedBuffer InSymbols, TConstArrayView<FShaderCompilerError> InCompilerWarnings, const FString& DebugInfo);
+	void AddEditorOnlyData(int32 Index, const FString& DebugName, uint64 ShaderTypeHash, FCompressedBuffer InSymbols, TConstArrayView<FShaderCompilerError> InCompilerWarnings, const TArray<FGenericShaderStat>& ShaderStatistics, const FString& DebugInfo = FString());
+	void UpdateEditorOnlyData(int32 Index, const FString& DebugName, uint64 ShaderTypeHash, FCompressedBuffer InSymbols, TConstArrayView<FShaderCompilerError> InCompilerWarnings, const FString& DebugInfo);
 
 	RENDERCORE_API void LogShaderCompilerWarnings();
 #endif
@@ -1234,6 +1242,15 @@ private:
 template<typename ShaderType>
 using TShaderRef = TShaderRefBase<ShaderType, FShaderMapPointerTable>;
 
+#define SHADER_TYPE_LIST \
+	SHADER_TYPE_ENTRY(Global), \
+	SHADER_TYPE_ENTRY(Material), \
+	SHADER_TYPE_ENTRY(MeshMaterial), \
+	SHADER_TYPE_ENTRY(Niagara), \
+	SHADER_TYPE_ENTRY(OCIO), \
+	SHADER_TYPE_ENTRY(ComputeKernel), \
+	SHADER_TYPE_ENTRY(NNERuntimeIREE),
+
 /**
  * An object which is used to serialize/deserialize, compile, and cache a particular shader class.
  *
@@ -1243,17 +1260,13 @@ using TShaderRef = TShaderRefBase<ShaderType, FShaderMapPointerTable>;
 class FShaderType
 {
 public:
+#define SHADER_TYPE_ENTRY(Name) Name
 	enum class EShaderTypeForDynamicCast : uint32
 	{
-		Global,
-		Material,
-		MeshMaterial,
-		Niagara,
-		OCIO,
-		ComputeKernel,
-		NNERuntimeIREE,
+		SHADER_TYPE_LIST
 		NumShaderTypes,
 	};
+#undef SHADER_TYPE_ENTRY
 
 	/**
 	 * Derived FShaderTypes should derive from this class to pass params to FShader constructor
@@ -1585,6 +1598,19 @@ protected:
 #endif // WITH_EDITOR
 };
 
+inline FString LexToString(FShaderType::EShaderTypeForDynamicCast ShaderType)
+{
+#define SHADER_TYPE_ENTRY(Name) TEXT(#Name)
+	static const TCHAR* TypeStrings[] =
+	{
+		SHADER_TYPE_LIST
+	};
+#undef SHADER_TYPE_ENTRY
+	checkf(ShaderType < FShaderType::EShaderTypeForDynamicCast::NumShaderTypes, TEXT("Invalid shader type"));
+	static_assert((sizeof(TypeStrings) / sizeof(TCHAR*)) == (int)FShaderType::EShaderTypeForDynamicCast::NumShaderTypes);
+	return TypeStrings[static_cast<uint32>(ShaderType)];
+}
+
 /**
  * Registers a shader type in various systems. Should be created as a static field/global.
  * 
diff --git a/Engine/Source/Runtime/RenderCore/Public/ShaderCodeArchive.h b/Engine/Source/Runtime/RenderCore/Public/ShaderCodeArchive.h
--- a/Engine/Source/Runtime/RenderCore/Public/ShaderCodeArchive.h
+++ b/Engine/Source/Runtime/RenderCore/Public/ShaderCodeArchive.h
@@ -131,6 +131,20 @@ public:
 	/** Mapping from shadermap hashes to an array of asset names - this is used for on-disk storage as it is shorter. */
 	TMap<FSHAHash, FShaderMapAssetPaths> ShaderCodeToAssets;
 
+	struct FShaderTypeHashes
+	{
+		TArray<uint64> Data;
+
+		friend FArchive& operator<<(FArchive& Ar, FShaderTypeHashes& Ref)
+		{
+			Ar << Ref.Data;
+			return Ar;
+		}
+	};
+	
+	/** Array of all shader type hashes which deduplicated to each individual shader; indexed as ShaderEntries array */
+	TArray<FShaderTypeHashes> ShaderTypes;
+
 	enum class EAssetInfoVersion : uint8
 	{
 		CurrentVersion = 2
@@ -197,6 +211,7 @@ public:
 #endif
 #if WITH_EDITOR
 			+ ShaderCodeToAssets.GetAllocatedSize()
+			+ ShaderTypes.GetAllocatedSize()
 #endif // WITH_EDITOR
 			;
 	}
@@ -209,6 +224,9 @@ public:
 		ShaderEntries.Empty();
 #endif
 		ShaderHashTable.Clear();
+#if WITH_EDITOR
+		ShaderTypes.Empty();
+#endif
 	}
 
 	void EmptyShaderMaps()
@@ -240,6 +258,7 @@ public:
 		return ShaderMapEntries.IsEmpty() && ShaderEntries.IsEmpty() && PreloadEntries.IsEmpty()
 #if WITH_EDITOR
 			&& ShaderCodeToAssets.IsEmpty()
+			&& ShaderTypes.IsEmpty()
 #endif
 			;
 	}
diff --git a/Engine/Source/Runtime/RenderCore/Public/ShaderCodeLibrary.h b/Engine/Source/Runtime/RenderCore/Public/ShaderCodeLibrary.h
--- a/Engine/Source/Runtime/RenderCore/Public/ShaderCodeLibrary.h
+++ b/Engine/Source/Runtime/RenderCore/Public/ShaderCodeLibrary.h
@@ -373,8 +373,11 @@ struct FShaderLibraryCooker
 		const FString& SandboxDestinationPath, const FString& SandboxMetadataPath, TArray<FString>& OutChunkFilenames, bool& bOutHasData);
 
 	// Dump collected stats for each shader platform
+	UE_DEPRECATED(5.6, "DumpShaderCodeStats is deprecated due to being incorrect w.r.t shader duplication due to chunking.")
 	static RENDERCORE_API void DumpShaderCodeStats();
 
+	static RENDERCORE_API void DumpShaderTypeStats(const FString& DebugInfoPath, const FString& MetaDataPath);
+
 	// Create a smaller 'patch' library that only contains data from 'NewMetaDataDir' not contained in any of 'OldMetaDataDirs'
 	static RENDERCORE_API bool CreatePatchLibrary(TArray<FString> const& OldMetaDataDirs, FString const& NewMetaDataDir, FString const& OutDir, bool bNativeFormat, bool bNeedsDeterministicOrder);
 };
