[shaders] stat fixes/improvements
- improve mismatch warnings to provide more details
- remove a potential but unlikely source of stat nondeterminism by sorting prior to storing in the shadermap
- don't allow multiple stats with the same name; some backends compile iteratively and set the same stats multiple times, so modify this such that each time a stat is added it overwrites the previous one

diff --git a/Engine/Source/Runtime/RenderCore/Private/ShaderCore.cpp b/Engine/Source/Runtime/RenderCore/Private/ShaderCore.cpp
--- a/Engine/Source/Runtime/RenderCore/Private/ShaderCore.cpp
+++ b/Engine/Source/Runtime/RenderCore/Private/ShaderCore.cpp
@@ -59,7 +59,7 @@ namespace FShaderStatTagNames
 	extern RENDERCORE_API const FName AnalysisArtifactsName = TEXT("ShaderStatTags.AnalysisArtifacts");
 }
 
-bool operator==(const FShaderStatVariant LHS, const FShaderStatVariant RHS)
+bool operator==(const FShaderStatVariant& LHS, const FShaderStatVariant& RHS)
 {
 	if (LHS.IsType<bool>() && RHS.IsType<bool>())
 	{
@@ -85,6 +85,57 @@ bool operator==(const FShaderStatVariant LHS, const FShaderStatVariant RHS)
 	return false;
 }
 
+bool operator<(const FShaderStatVariant& LHS, const FShaderStatVariant& RHS)
+{
+	if (LHS.IsType<bool>() && RHS.IsType<bool>())
+	{
+		return LHS.Get<bool>() < RHS.Get<bool>();
+	}
+	else if (LHS.IsType<float>() && RHS.IsType<float>()) 
+	{
+		return LHS.Get<float>() < RHS.Get<float>();
+	}
+	else if (LHS.IsType<int32>() && RHS.IsType<int32>())
+	{
+		return LHS.Get<int32>() < RHS.Get<int32>();
+	}
+	else if (LHS.IsType<uint32>() && RHS.IsType<uint32>())
+	{
+		return LHS.Get<uint32>() < RHS.Get<uint32>();
+	}
+	else if (LHS.IsType<FString>() && RHS.IsType<FString>())
+	{
+		return LHS.Get<FString>() < RHS.Get<FString>();
+	}
+
+	return false;
+}
+
+FStringBuilderBase& operator<<(FStringBuilderBase& Builder, const FShaderStatVariant& Variant)
+{
+	if (Variant.IsType<bool>())
+	{
+		Builder << Variant.Get<bool>();
+	}
+	else if (Variant.IsType<float>())
+	{
+		Builder << Variant.Get<float>();
+	}
+	else if (Variant.IsType<int32>())
+	{
+		Builder << Variant.Get<int32>();
+	}
+	else if (Variant.IsType<uint32>())
+	{
+		Builder << Variant.Get<uint32>();
+	}
+	else if (Variant.IsType<FString>())
+	{
+		Builder << Variant.Get<FString>();
+	}
+	return Builder;
+}
+
 void FGenericShaderStat::StreamArchive(FArchive& Ar)
 {
 	FGenericShaderStat& Stat = *this;
@@ -121,6 +172,41 @@ bool FGenericShaderStat::operator==(const FGenericShaderStat& RHS) const
 	return (StatName == RHS.StatName) && (Value == RHS.Value) && (Flags == RHS.Flags) && (TagName == RHS.TagName);
 }
 
+bool FGenericShaderStat::operator<(const FGenericShaderStat& RHS) const
+{
+	int32 NameComp = StatName.Compare(RHS.StatName);
+	if (NameComp != 0)
+	{
+		return NameComp < 0;
+	}
+	
+	if (!(Value == RHS.Value))
+	{
+		return Value < RHS.Value;
+	}
+	
+	if (Flags != RHS.Flags)
+	{
+		return Flags < RHS.Flags;
+	}
+
+	return TagName.Compare(RHS.TagName) < 0;
+}
+
+FStringBuilderBase& operator<<(FStringBuilderBase& Builder, const FGenericShaderStat& Stat)
+{
+	Builder << Stat.StatName << TEXT("=") << Stat.Value;
+	if (!Stat.TagName.IsNone())
+	{
+		Builder.Appendf(TEXT(" (Tag=%s)"), *Stat.TagName.ToString());
+	}
+	if (Stat.Flags == FGenericShaderStat::EFlags::Hidden)
+	{
+		Builder << TEXT(" (hidden)");
+	}
+	return Builder;
+}
+
 static TAutoConsoleVariable<bool> CVarDumpDebugInfoForCacheHits(
 	TEXT("r.ShaderCompiler.DumpDebugInfoForCacheHits"),
 	true,
diff --git a/Engine/Source/Runtime/RenderCore/Private/ShaderResource.cpp b/Engine/Source/Runtime/RenderCore/Private/ShaderResource.cpp
--- a/Engine/Source/Runtime/RenderCore/Private/ShaderResource.cpp
+++ b/Engine/Source/Runtime/RenderCore/Private/ShaderResource.cpp
@@ -280,13 +280,19 @@ void FShaderMapResourceCode::AddShaderCompilerOutput(const FShaderCompilerOutput
 	uint64 ShaderTypeHash = Key.ShaderType->GetHashedName().GetHash();
 	const int32 Index = Algo::LowerBound(ShaderHashes, InHash);
 	const FShaderCodeResource& CodeResource = Output.GetFinalizedCodeResource();
+
+#if WITH_EDITORONLY_DATA
+	TArray<FGenericShaderStat> ShaderStatsSorted(Output.ShaderStatistics);
+	ShaderStatsSorted.StableSort();
+#endif
+
 	if (Index >= ShaderHashes.Num() || ShaderHashes[Index] != InHash)
 	{
 		ShaderHashes.Insert(InHash, Index);
 
 #if WITH_EDITORONLY_DATA
 		// Output.Errors contains warnings in the case any exist (no errors since if there were the job would have failed)
-		AddEditorOnlyData(Index, DebugName, ShaderTypeHash, CodeResource.GetSymbolsBuffer(), Output.Errors, Output.ShaderStatistics, DebugInfo);
+		AddEditorOnlyData(Index, DebugName, ShaderTypeHash, CodeResource.GetSymbolsBuffer(), Output.Errors, MoveTemp(ShaderStatsSorted), DebugInfo);
 #endif
 
 		ShaderCodeResources.Insert(CodeResource, Index);
@@ -298,19 +304,19 @@ void FShaderMapResourceCode::AddShaderCompilerOutput(const FShaderCompilerOutput
 		// We append the warnings and deduplicate other data like DebugInfo for any additional jobs which resulted in the
 		// same bytecode for the sake of determinism in the results saved to DDC.
 		UpdateEditorOnlyData(Index, DebugName, ShaderTypeHash, CodeResource.GetSymbolsBuffer(), Output.Errors, DebugInfo);
-		ValidateShaderStatisticsEditorOnlyData(Index, Output.ShaderStatistics);
+		ValidateShaderStatisticsEditorOnlyData(Index, ShaderStatsSorted);
 	}
 #endif
 }
 
 #if WITH_EDITORONLY_DATA
-void FShaderMapResourceCode::AddEditorOnlyData(int32 Index, const FString& DebugName, uint64 ShaderTypeHash, FCompressedBuffer InSymbols, TConstArrayView<FShaderCompilerError> InCompilerWarnings, const TArray<FGenericShaderStat>& ShaderStatistics, const FString& DebugInfo)
+void FShaderMapResourceCode::AddEditorOnlyData(int32 Index, const FString& DebugName, uint64 ShaderTypeHash, FCompressedBuffer InSymbols, TConstArrayView<FShaderCompilerError> InCompilerWarnings, TArray<FGenericShaderStat>&& ShaderStatistics, const FString& DebugInfo)
 {
 	FShaderEditorOnlyDataEntry& Entry = ShaderEditorOnlyDataEntries.InsertDefaulted_GetRef(Index);
 
 	// This should be a newly created shader entry.
 	check(Entry.ShaderStatistics.Num() == 0);
-	Entry.ShaderStatistics = ShaderStatistics;
+	Entry.ShaderStatistics = MoveTemp(ShaderStatistics);
 
 	UpdateEditorOnlyData(Index, DebugName, ShaderTypeHash, InSymbols, InCompilerWarnings, DebugInfo);
 }
@@ -350,9 +356,20 @@ void FShaderMapResourceCode::ValidateShaderStatisticsEditorOnlyData(int32 Index,
 	check(ShaderEditorOnlyDataEntries.IsValidIndex(Index));
 	FShaderEditorOnlyDataEntry& Entry = ShaderEditorOnlyDataEntries[Index];
 
+	auto PrintMismatchWarning = [](FShaderEditorOnlyDataEntry& Entry, const TArray<FGenericShaderStat>& ShaderStatistics)
+		{
+			TStringBuilder<512> WarningMessage;
+			WarningMessage << TEXT("Non-determinism detected in shader statistics; existing entry contains:\n");
+			WarningMessage.Join(Entry.ShaderStatistics, TEXT(", "));
+			WarningMessage << TEXT("\nNew entry contains:\n");
+			WarningMessage.Join(ShaderStatistics, TEXT(", "));
+			WarningMessage << TEXT("\nShader statistics are assumed to be a property of the bytecode and so must match when bytecode matches.");
+			UE_LOG(LogShaders, Warning, TEXT("%s"), WarningMessage.ToString());
+		};
+
 	if (Entry.ShaderStatistics.Num() != ShaderStatistics.Num())
 	{
-		UE_LOG(LogShaders, Warning, TEXT("Non-determinism detected in shader statistics.  Multiple duplicate shaders have the same shader statistics."));
+		PrintMismatchWarning(Entry, ShaderStatistics);
 		return;
 	}
 
@@ -362,7 +379,7 @@ void FShaderMapResourceCode::ValidateShaderStatisticsEditorOnlyData(int32 Index,
 		const FGenericShaderStat& StatB = ShaderStatistics[i];
 		if (!(StatA == StatB))
 		{
-			UE_LOG(LogShaders, Warning, TEXT("Non-determinism detected in shader statistics.  Multiple duplicate shaders have the same shader statistics."));
+			PrintMismatchWarning(Entry, ShaderStatistics);
 			return;
 		}
 	}
diff --git a/Engine/Source/Runtime/RenderCore/Public/Shader.h b/Engine/Source/Runtime/RenderCore/Public/Shader.h
--- a/Engine/Source/Runtime/RenderCore/Public/Shader.h
+++ b/Engine/Source/Runtime/RenderCore/Public/Shader.h
@@ -601,7 +601,7 @@ public:
 	int32 FindShaderIndex(const FSHAHash& InHash) const;
 
 #if WITH_EDITORONLY_DATA
-	void AddEditorOnlyData(int32 Index, const FString& DebugName, uint64 ShaderTypeHash, FCompressedBuffer InSymbols, TConstArrayView<FShaderCompilerError> InCompilerWarnings, const TArray<FGenericShaderStat>& ShaderStatistics, const FString& DebugInfo = FString());
+	void AddEditorOnlyData(int32 Index, const FString& DebugName, uint64 ShaderTypeHash, FCompressedBuffer InSymbols, TConstArrayView<FShaderCompilerError> InCompilerWarnings, TArray<FGenericShaderStat>&& ShaderStatistics, const FString& DebugInfo = FString());
 	void UpdateEditorOnlyData(int32 Index, const FString& DebugName, uint64 ShaderTypeHash, FCompressedBuffer InSymbols, TConstArrayView<FShaderCompilerError> InCompilerWarnings, const FString& DebugInfo);
 
 	RENDERCORE_API void LogShaderCompilerWarnings();
diff --git a/Engine/Source/Runtime/RenderCore/Public/ShaderCompilerCore.h b/Engine/Source/Runtime/RenderCore/Public/ShaderCompilerCore.h
--- a/Engine/Source/Runtime/RenderCore/Public/ShaderCompilerCore.h
+++ b/Engine/Source/Runtime/RenderCore/Public/ShaderCompilerCore.h
@@ -563,11 +563,29 @@ struct FShaderCompilerOutput
 	template<typename TValue>
 	void AddStatistic(const TCHAR* Name, TValue Value, FGenericShaderStat::EFlags Flags = FGenericShaderStat::EFlags::None, FName TagName = NAME_None)
 	{
-		FGenericShaderStat& Stat = ShaderStatistics.AddZeroed_GetRef();
-		Stat.StatName = FName(Name);
-		Stat.Value = FShaderStatVariant(TInPlaceType<TValue>(), Value);
-		Stat.Flags = Flags;
-		Stat.TagName = TagName;
+		// handle the case where a stat with the given name already exists; some backends compile iteratively and this is a simpler way to handle that
+		// than requiring each backend to manually avoid setting until the "final" compilation. linear search is in practice fine here; there are low 
+		// single digit numbers of stats for all existing shader formats at the time of writing.
+		FName StatName(Name);
+		FGenericShaderStat* Stat = nullptr;
+		for (FGenericShaderStat& ExistingStat : ShaderStatistics)
+		{
+			if (ExistingStat.StatName == StatName)
+			{
+				Stat = &ExistingStat;
+				break;
+			}
+		}
+
+		if (!Stat)
+		{
+			Stat = &ShaderStatistics.AddZeroed_GetRef();
+			Stat->StatName = StatName;
+		}
+
+		Stat->Value = FShaderStatVariant(TInPlaceType<TValue>(), Value);
+		Stat->Flags = Flags;
+		Stat->TagName = TagName;
 	}
 
 	const FShaderCodeResource& GetFinalizedCodeResource() const
diff --git a/Engine/Source/Runtime/RenderCore/Public/ShaderCore.h b/Engine/Source/Runtime/RenderCore/Public/ShaderCore.h
--- a/Engine/Source/Runtime/RenderCore/Public/ShaderCore.h
+++ b/Engine/Source/Runtime/RenderCore/Public/ShaderCore.h
@@ -57,7 +57,8 @@ class ITargetPlatform;
 
 using FShaderStatVariant = TVariant<bool, float, int32, uint32, FString>;
 
-bool operator==(const FShaderStatVariant LHS, const FShaderStatVariant RHS);
+bool operator==(const FShaderStatVariant& LHS, const FShaderStatVariant& RHS);
+bool operator<(const FShaderStatVariant& LHS, const FShaderStatVariant& RHS);
 
 namespace FShaderStatTagNames
 {
@@ -82,6 +83,8 @@ public:
 
 	friend FArchive& operator<<(FArchive& Ar, FGenericShaderStat& Stat) { Stat.StreamArchive(Ar); return Ar; }
 	bool operator==(const FGenericShaderStat& RHS) const;
+	bool operator<(const FGenericShaderStat& RHS) const;
+	friend FStringBuilderBase& operator<<(FStringBuilderBase& Builder, const FGenericShaderStat& Stat);
 
 	RENDERCORE_API void StreamArchive(FArchive& Ar);
 };
