Remove incorrect check().
Copying from UINT to BC surfaces is valid but was triggering the assert.
#rnx

diff --git a/Engine/Source/Runtime/RenderCore/Private/RenderGraphUtils.cpp b/Engine/Source/Runtime/RenderCore/Private/RenderGraphUtils.cpp
--- a/Engine/Source/Runtime/RenderCore/Private/RenderGraphUtils.cpp
+++ b/Engine/Source/Runtime/RenderCore/Private/RenderGraphUtils.cpp
@@ -269,7 +269,6 @@ void AddCopyTexturePass(
 
 	const FRDGTextureDesc& InputDesc = InputTexture->Desc;
 	const FRDGTextureDesc& OutputDesc = OutputTexture->Desc;
-	checkf(InputDesc.Format == OutputDesc.Format, TEXT("This method does not support format conversion."));
 
 	FCopyTextureParameters* Parameters = GraphBuilder.AllocParameters<FCopyTextureParameters>();
 	Parameters->Input = InputTexture;
