Deprecate use of low level RHI Texture target when requesting virtual texture production.
Virtual texture producers should work with the pooled render target to ensure that texture transition logic can stay robust.
#rnx

diff --git a/Engine/Source/Runtime/RenderCore/Public/VirtualTexturing.h b/Engine/Source/Runtime/RenderCore/Public/VirtualTexturing.h
--- a/Engine/Source/Runtime/RenderCore/Public/VirtualTexturing.h
+++ b/Engine/Source/Runtime/RenderCore/Public/VirtualTexturing.h
@@ -24,6 +24,7 @@ class FRHICommandListImmediate;
 class FRHIShaderResourceView;
 class FRHITexture;
 class FRHIUnorderedAccessView;
+struct IPooledRenderTarget;
 class ISceneRenderer;
 namespace ERHIFeatureLevel { enum Type : int; }
 enum class EVTProducerPriority : uint8;
@@ -261,21 +262,20 @@ struct FVTRequestPageResult
 /** Describes a location to write a single layer of a VT tile */
 struct FVTProduceTargetLayer
 {
+	/** Pooled render target to write to. */
+	IPooledRenderTarget* PooledRenderTarget = nullptr;
+	/** Location within the target texture to write. */
+	FIntVector pPageLocation = FIntVector::ZeroValue;
+
+	UE_DEPRECATED(5.6, "Use PooledRenderTarget instead.")
+	FRHITexture* TextureRHI = nullptr;
+
+	// Disable some TextureRHI deprecation warnings here.
+	PRAGMA_DISABLE_DEPRECATION_WARNINGS
 	FVTProduceTargetLayer() = default;
 	FVTProduceTargetLayer(const FVTProduceTargetLayer&) = default;
 	FVTProduceTargetLayer& operator=(const FVTProduceTargetLayer&) = default;
-
-	/** The texture to write to. */
-	FRHITexture* TextureRHI = nullptr;
-	
-	/**
-	 * Pooled render target. For FRDGBuilder::RegisterExternalTexture() which only accepts pooled render targets.
-	 * To avoid cost of manipulating ref counting pointers a raw pointer is used instead - it is valid until returning from your Finalize().
-	 * So do not try to store the pointer
-	 */
-	struct IPooledRenderTarget* PooledRenderTarget = nullptr;
-	/** Location within the texture to write */
-	FIntVector pPageLocation = FIntVector::ZeroValue;
+	PRAGMA_ENABLE_DEPRECATION_WARNINGS
 };
 
 /**
