Removing deprecated binding functions.

diff --git a/Engine/Source/Runtime/RenderCore/Public/ShaderParameterUtils.h b/Engine/Source/Runtime/RenderCore/Public/ShaderParameterUtils.h
--- a/Engine/Source/Runtime/RenderCore/Public/ShaderParameterUtils.h
+++ b/Engine/Source/Runtime/RenderCore/Public/ShaderParameterUtils.h
@@ -285,18 +285,6 @@ inline void SetBatchedShaderParametersMixed(
 	InShader->SetParameters(BatchedParameters, Forward<TArguments>(InArguments)...);
 }
 
-template<typename TShaderType, typename TShaderTypeRHI, typename... TArguments>
-UE_DEPRECATED(5.5, "SetBatchedShaderParametersMixed does not need an RHI shader")
-inline void SetBatchedShaderParametersMixed(
-	FRHIBatchedShaderParameters& BatchedParameters,
-	const TShaderRef<TShaderType>& InShader,
-	TShaderTypeRHI* InShaderRHI,
-	const typename TShaderType::FParameters& Parameters,
-	TArguments&&... InArguments)
-{
-	SetBatchedShaderParametersMixed(BatchedParameters, InShader, Parameters, Forward<TArguments>(InArguments)...);
-}
-
 /// Utility to set all legacy and non-legacy parameters for a shader. Requires the shader type to implement SetParameters(FRHIBatchedShaderParameters& BatchedParameters, ...)
 template<typename TRHICmdList, typename TShaderType, typename TShaderTypeRHI, typename... TArguments>
 inline void SetShaderParametersMixed(
@@ -311,62 +299,6 @@ inline void SetShaderParametersMixed(
 	RHICmdList.SetBatchedShaderParameters(InShaderRHI, BatchedParameters);
 }
 
-/// Utility to set all legacy and non-legacy parameters for a Vertex shader. Requires the shader type to implement SetParameters(FRHIBatchedShaderParameters& BatchedParameters, ...)
-template<typename TShaderType, typename... TArguments>
-UE_DEPRECATED(5.5, "Please use SetBatchedShaderParametersMixed directly.")
-inline void SetShaderParametersMixedVS(FRHIBatchedShaderParameters& BatchedParameters, const TShaderRef<TShaderType>& InShader, const typename TShaderType::FParameters& Parameters, TArguments&&... InArguments)
-{
-	SetBatchedShaderParametersMixed(BatchedParameters, InShader, Parameters, Forward<TArguments>(InArguments)...);
-}
-
-/// Utility to set all legacy and non-legacy parameters for a Mesh shader. Requires the shader type to implement SetParameters(FRHIBatchedShaderParameters& BatchedParameters, ...)
-template<typename TShaderType, typename... TArguments>
-UE_DEPRECATED(5.5, "Please use SetBatchedShaderParametersMixed directly.")
-inline void SetShaderParametersMixedMS(FRHIBatchedShaderParameters& BatchedParameters, const TShaderRef<TShaderType>& InShader, const typename TShaderType::FParameters& Parameters, TArguments&&... InArguments)
-{
-	SetBatchedShaderParametersMixed(BatchedParameters, InShader, Parameters, Forward<TArguments>(InArguments)...);
-}
-
-/// Utility to set all legacy and non-legacy parameters for an Amplification shader. Requires the shader type to implement SetParameters(FRHIBatchedShaderParameters& BatchedParameters, ...)
-template<typename TShaderType, typename... TArguments>
-UE_DEPRECATED(5.5, "Please use SetBatchedShaderParametersMixed directly.")
-inline void SetShaderParametersMixedAS(FRHIBatchedShaderParameters& BatchedParameters, const TShaderRef<TShaderType>& InShader, const typename TShaderType::FParameters& Parameters, TArguments&&... InArguments)
-{
-	SetBatchedShaderParametersMixed(BatchedParameters, InShader, Parameters, Forward<TArguments>(InArguments)...);
-}
-
-/// Utility to set all legacy and non-legacy parameters for a Pixel shader. Requires the shader type to implement SetParameters(FRHIBatchedShaderParameters& BatchedParameters, ...)
-template<typename TShaderType, typename... TArguments>
-UE_DEPRECATED(5.5, "Please use SetBatchedShaderParametersMixed directly.")
-inline void SetShaderParametersMixedPS(FRHIBatchedShaderParameters& BatchedParameters, const TShaderRef<TShaderType>& InShader, const typename TShaderType::FParameters& Parameters, TArguments&&... InArguments)
-{
-	SetBatchedShaderParametersMixed(BatchedParameters, InShader, Parameters, Forward<TArguments>(InArguments)...);
-}
-
-/// Utility to set all legacy and non-legacy parameters for a Geometry shader. Requires the shader type to implement SetParameters(FRHIBatchedShaderParameters& BatchedParameters, ...)
-template<typename TShaderType, typename... TArguments>
-UE_DEPRECATED(5.5, "Please use SetBatchedShaderParametersMixed directly.")
-inline void SetShaderParametersMixedGS(FRHIBatchedShaderParameters& BatchedParameters, const TShaderRef<TShaderType>& InShader, const typename TShaderType::FParameters& Parameters, TArguments&&... InArguments)
-{
-	SetBatchedShaderParametersMixed(BatchedParameters, InShader, Parameters, Forward<TArguments>(InArguments)...);
-}
-
-/// Utility to set all legacy and non-legacy parameters for a Compute shader. Requires the shader type to implement SetParameters(FRHIBatchedShaderParameters& BatchedParameters, ...)
-template<typename TShaderType, typename... TArguments>
-UE_DEPRECATED(5.5, "Please use SetBatchedShaderParametersMixed directly.")
-inline void SetShaderParametersMixedCS(FRHIBatchedShaderParameters& BatchedParameters, const TShaderRef<TShaderType>& InShader, const typename TShaderType::FParameters& Parameters, TArguments&&... InArguments)
-{
-	SetBatchedShaderParametersMixed(BatchedParameters, InShader, Parameters, Forward<TArguments>(InArguments)...);
-}
-
-/// Utility to set all legacy and non-legacy parameters for a WorkGraph shader. Requires the shader type to implement SetParameters(FRHIBatchedShaderParameters& BatchedParameters, ...)
-template<typename TShaderType, typename... TArguments>
-UE_DEPRECATED(5.5, "Please use SetBatchedShaderParametersMixed directly.")
-inline void SetShaderParametersMixedWS(FRHIBatchedShaderParameters& BatchedParameters, const TShaderRef<TShaderType>& InShader, const typename TShaderType::FParameters& Parameters, TArguments&&... InArguments)
-{
-	SetBatchedShaderParametersMixed(BatchedParameters, InShader, Parameters, Forward<TArguments>(InArguments)...);
-}
-
 /// Utility to set all legacy and non-legacy parameters for a Vertex shader. Requires the shader type to implement SetParameters(FRHIBatchedShaderParameters& BatchedParameters, ...)
 template<typename TRHICmdList, typename TShaderType, typename... TArguments>
 inline void SetShaderParametersMixedVS(TRHICmdList& RHICmdList, const TShaderRef<TShaderType>& InShader, const typename TShaderType::FParameters& Parameters, TArguments&&... InArguments)
diff --git a/Engine/Source/Runtime/Renderer/Public/MeshPassProcessor.h b/Engine/Source/Runtime/Renderer/Public/MeshPassProcessor.h
--- a/Engine/Source/Runtime/Renderer/Public/MeshPassProcessor.h
+++ b/Engine/Source/Runtime/Renderer/Public/MeshPassProcessor.h
@@ -1002,12 +1002,6 @@ public:
 	RENDERER_API void SetParameters(FRHIBatchedShaderParameters& BatchedParameters, class FShaderBindingState* StateCacheShaderBindings = nullptr) const;
 	RENDERER_API void SetOnCommandList(FRHIComputeCommandList& RHICmdList, FRHIComputeShader* Shader, class FShaderBindingState* StateCacheShaderBindings = nullptr) const;
 
-	UE_DEPRECATED(5.5, "Use the version of SetParameters that doesn't take an FRHIComputeShader.")
-	void SetParameters(FRHIBatchedShaderParameters& BatchedParameters, FRHIComputeShader* Shader, class FShaderBindingState* StateCacheShaderBindings = nullptr) const
-	{
-		SetParameters(BatchedParameters, StateCacheShaderBindings);
-	}
-
 #if RHI_RAYTRACING
 	RENDERER_API FRayTracingLocalShaderBindings* SetRayTracingShaderBindings(FRayTracingLocalShaderBindingWriter* BindingWriter, uint32 ShaderIndexInPipeline, uint32 RecordIndex, uint32 UserData, ERayTracingLocalShaderBindingType BindingType) const;
 	UE_DEPRECATED(5.5, "Use version with explicit FRHIRayTracingGeometry.")
