Fixed incorrect command list being used for disabling extra transitions.

diff --git a/Engine/Source/Runtime/RenderCore/Private/RenderGraphBuilder.cpp b/Engine/Source/Runtime/RenderCore/Private/RenderGraphBuilder.cpp
--- a/Engine/Source/Runtime/RenderCore/Private/RenderGraphBuilder.cpp
+++ b/Engine/Source/Runtime/RenderCore/Private/RenderGraphBuilder.cpp
@@ -2682,7 +2682,7 @@ void FRDGBuilder::SubmitBufferUploads(FRHICommandListBase& RHICmdListUpload, UE:
 		SCOPED_NAMED_EVENT_TEXT("Upload_Multiple", FColor::Magenta);
 		
 		// This is here because we are explicitly batching a series of transition for all the buffers and we don't the individual extra transitions in Lock/Unlock
-		FRHICommandListScopedAllowExtraTransitions ScopedExtraTransitions(RHICmdList, false);
+		FRHICommandListScopedAllowExtraTransitions ScopedExtraTransitions(RHICmdListUpload, false);
 
 		TSet<FRHIBuffer*, DefaultKeyFuncs<void*>, FConcurrentLinearSetAllocator> BuffersSet;
 		TArray<FRHITransitionInfo, FConcurrentLinearArrayAllocator> CopyDestTransitionInfo;
