[shaders] deprecated code cleanup for 5.7

diff --git a/Engine/Source/Developer/ShaderCompilerCommon/Private/ShaderCompilerCommon.cpp b/Engine/Source/Developer/ShaderCompilerCommon/Private/ShaderCompilerCommon.cpp
--- a/Engine/Source/Developer/ShaderCompilerCommon/Private/ShaderCompilerCommon.cpp
+++ b/Engine/Source/Developer/ShaderCompilerCommon/Private/ShaderCompilerCommon.cpp
@@ -1636,11 +1636,6 @@ static FString OfflineCompiler_ExtractErrors(const FString& CompilerOutput)
 	return ReturnedErrors;
 }
 
-void CompileOfflineMali(const FShaderCompilerInput& Input, FShaderCompilerOutput& ShaderOutput, const ANSICHAR* ShaderSource, const int32 SourceSize, bool bVulkanSpirV, const ANSICHAR* VulkanSpirVEntryPoint)
-{
-	return CompileShaderOffline(Input, ShaderOutput, ShaderSource, SourceSize, bVulkanSpirV, VulkanSpirVEntryPoint);
-}
-
 /**
 * OfflineShaderCompiler's compilation command line options.
 * Each OfflineShaderCompiler should specify its own FOfflineShaderCompilerOptions. If one option is not supported by this OfflineShaderCompiler, leave it empty.
diff --git a/Engine/Source/Developer/ShaderCompilerCommon/Private/ShaderConductorContext.cpp b/Engine/Source/Developer/ShaderCompilerCommon/Private/ShaderConductorContext.cpp
--- a/Engine/Source/Developer/ShaderCompilerCommon/Private/ShaderConductorContext.cpp
+++ b/Engine/Source/Developer/ShaderCompilerCommon/Private/ShaderConductorContext.cpp
@@ -635,12 +635,6 @@ namespace CrossCompiler
 		return true;
 	}
 
-
-	bool FShaderConductorContext::LoadSource(const FString& ShaderSource, const FString& Filename, const FString& EntryPoint, EShaderFrequency ShaderStage, const FShaderCompilerDefinitions* Definitions, const TArray<FString>* ExtraDxcArgs)
-	{
-		return LoadSource(FStringView(ShaderSource), Filename, EntryPoint, ShaderStage, Definitions, ExtraDxcArgs);
-	}
-
 	bool FShaderConductorContext::LoadSource(const ANSICHAR* ShaderSource, const ANSICHAR* Filename, const ANSICHAR* EntryPoint, EShaderFrequency ShaderStage, const FShaderCompilerDefinitions* Definitions, const TArray<FString>* ExtraDxcArgs)
 	{
 		Intermediates->ShaderSource = ShaderSource;
@@ -664,11 +658,6 @@ namespace CrossCompiler
 		return true;
 	}
 
-	bool FShaderConductorContext::RewriteHlsl(const FShaderConductorOptions& Options, FString* OutSource)
-	{
-		return false; // DXC rewriter is no longer supported and this function has been deprecated since UE5.5
-	}
-
     bool FShaderConductorContext::CompileHlslToDxil(const FShaderConductorOptions& Options, TArray<uint32>& OutDxil)
     {
 		constexpr bool bGenerateSpirv = false;
diff --git a/Engine/Source/Developer/ShaderCompilerCommon/Public/ShaderCompilerCommon.h b/Engine/Source/Developer/ShaderCompilerCommon/Public/ShaderCompilerCommon.h
--- a/Engine/Source/Developer/ShaderCompilerCommon/Public/ShaderCompilerCommon.h
+++ b/Engine/Source/Developer/ShaderCompilerCommon/Public/ShaderCompilerCommon.h
@@ -446,9 +446,6 @@ extern SHADERCOMPILERCOMMON_API void DumpDebugShaderBinary(const FShaderCompiler
 extern SHADERCOMPILERCOMMON_API void DumpDebugShaderDisassembledSpirv(const FShaderCompilerInput& Input, void* InData, int32 InDataByteSize, const FString& FileExtension);
 extern SHADERCOMPILERCOMMON_API void DumpDebugShaderDisassembledDxil(const FShaderCompilerInput& Input, void* InData, int32 InDataByteSize, const FString& FileExtension);
 
-UE_DEPRECATED(5.5, "CompileOfflineMali is no longer used; CompileShaderOffline is used instead to support different compilers instead of Mali itself.")
-extern SHADERCOMPILERCOMMON_API void CompileOfflineMali(const FShaderCompilerInput &Input, FShaderCompilerOutput& ShaderOutput, const ANSICHAR* ShaderSource, const int32 SourceSize, bool bVulkanSpirV, const ANSICHAR* VulkanSpirVEntryPoint = nullptr);
-
 // calls 'Offline Compiler' to compile the source code and extract the stats
 extern SHADERCOMPILERCOMMON_API void CompileShaderOffline(const FShaderCompilerInput& Input,
 	FShaderCompilerOutput& ShaderOutput,
diff --git a/Engine/Source/Developer/ShaderCompilerCommon/Public/ShaderConductorContext.h b/Engine/Source/Developer/ShaderCompilerCommon/Public/ShaderConductorContext.h
--- a/Engine/Source/Developer/ShaderCompilerCommon/Public/ShaderConductorContext.h
+++ b/Engine/Source/Developer/ShaderCompilerCommon/Public/ShaderConductorContext.h
@@ -199,17 +199,9 @@ namespace CrossCompiler
 		FShaderConductorContext& operator = (const FShaderConductorContext&) = delete;
 
 		/** Loads the shader source and converts the input descriptor to a format suitable for ShaderConductor. If 'Definitions' is null, the previously loaded definitions are not modified. */
-		UE_API PRAGMA_DISABLE_DEPRECATION_WARNINGS		// FShaderCompilerDefinitions will be made internal in the future, marked deprecated until then
-		bool LoadSource(const FString& ShaderSource, const FString& Filename, const FString& EntryPoint, EShaderFrequency ShaderStage, const FShaderCompilerDefinitions* Definitions = nullptr, const TArray<FString>* ExtraDxcArgs = nullptr);
 		UE_API bool LoadSource(FStringView ShaderSource, const FString& Filename, const FString& EntryPoint, EShaderFrequency ShaderStage, const FShaderCompilerDefinitions* Definitions = nullptr, const TArray<FString>* ExtraDxcArgs = nullptr);
 		UE_API bool LoadSource(FAnsiStringView ShaderSource, const FString& Filename, const FString& EntryPoint, EShaderFrequency ShaderStage, const FShaderCompilerDefinitions* Definitions = nullptr, const TArray<FString>* ExtraDxcArgs = nullptr);
 		UE_API bool LoadSource(const ANSICHAR* ShaderSource, const ANSICHAR* Filename, const ANSICHAR* EntryPoint, EShaderFrequency ShaderStage, const FShaderCompilerDefinitions* Definitions = nullptr, const TArray<FString>* ExtraDxcArgs = nullptr);
-		PRAGMA_ENABLE_DEPRECATION_WARNINGS
-
-		/** Rewrites the specified HLSL shader source code. This allows to reduce the HLSL code by removing unused global resources for instance.
-		This will update the internally loaded source (see 'LoadSource'), so the output parameter 'OutSource' is optional. */
-		UE_DEPRECATED(5.5, "DXC rewriter has been deprecated since UE5.5. FShaderConductorContext::RewriteHlsl will be removed in future versions.")
-		UE_API bool RewriteHlsl(const FShaderConductorOptions& Options, FString* OutSource = nullptr);
 
         /** Compiles the specified HLSL shader source code to DXIL. */
         UE_API bool CompileHlslToDxil(const FShaderConductorOptions& Options, TArray<uint32>& OutDxil);
diff --git a/Engine/Source/Developer/TargetPlatform/Public/Interfaces/IShaderFormat.h b/Engine/Source/Developer/TargetPlatform/Public/Interfaces/IShaderFormat.h
--- a/Engine/Source/Developer/TargetPlatform/Public/Interfaces/IShaderFormat.h
+++ b/Engine/Source/Developer/TargetPlatform/Public/Interfaces/IShaderFormat.h
@@ -62,19 +62,6 @@ public:
 	*/
 	virtual bool SupportsShaderArchives() const { return false; }
 	
-	UE_DEPRECATED(5.5, "Use version accepting a TArray of FSharedBuffer code buffers.")
-	virtual bool CreateShaderArchive(FString const& LibraryName,
-		FName ShaderFormatAndShaderPlatformName,
-		const FString& WorkingDirectory,
-		const FString& OutputDir,
-		const FString& DebugOutputDir,
-		const FSerializedShaderArchive& SerializedShaders,
-		const TArray<TArray<uint8>>& ShaderCode,
-		TArray<FString>* OutputFiles) const
-	{ 
-		return false; 
-	}
-
 	/**
      * Create a format specific archive for precompiled shader code.
      *
diff --git a/Engine/Source/Runtime/Engine/Private/Materials/MaterialShader.cpp b/Engine/Source/Runtime/Engine/Private/Materials/MaterialShader.cpp
--- a/Engine/Source/Runtime/Engine/Private/Materials/MaterialShader.cpp
+++ b/Engine/Source/Runtime/Engine/Private/Materials/MaterialShader.cpp
@@ -1457,12 +1457,6 @@ void FMaterialShaderMapId::UpdateFromParameterSet(const FStaticParameterSet& Sta
 
 uint32 GetTypeHash(FPlatformTypeLayoutParameters Params) { return HashCombine(Params.Flags, Params.MaxFieldAlignment); }
 
-void FMaterialShaderMapId::AppendStaticParametersString(FString& ParamsString) const
-{
-	FShaderKeyGenerator KeyGen(ParamsString);
-	AppendStaticParameters(KeyGen);
-}
-
 void FMaterialShaderMapId::AppendStaticParameters(FShaderKeyGenerator& KeyGen) const
 {
 	for (const FStaticSwitchParameter& StaticSwitchParameter : StaticSwitchParameters)
@@ -3679,14 +3673,6 @@ FMaterialShaderMap* FMaterialShaderMap::GetFinalizedClone() const
 }
 #endif // WITH_EDITOR
 
-bool FMaterialShaderMap::Serialize(FArchive& Ar, bool bInlineShaderResources, bool bLoadingCooked, bool bInlineShaderCode, const FName& SerializingAsset)
-{
-	FShaderSerializeContext Ctx(Ar);
-	Ctx.bLoadingCooked = bLoadingCooked;
-	Ctx.SerializingAsset = SerializingAsset;
-	return Serialize(Ctx);
-}
-
 bool FMaterialShaderMap::Serialize(FShaderSerializeContext& Ctx)
 {
 	SCOPED_LOADTIMER(FMaterialShaderMap_Serialize);
diff --git a/Engine/Source/Runtime/Engine/Private/ShaderCompiler/ShaderCompiler.cpp b/Engine/Source/Runtime/Engine/Private/ShaderCompiler/ShaderCompiler.cpp
--- a/Engine/Source/Runtime/Engine/Private/ShaderCompiler/ShaderCompiler.cpp
+++ b/Engine/Source/Runtime/Engine/Private/ShaderCompiler/ShaderCompiler.cpp
@@ -236,24 +236,12 @@ namespace GlobalShaderCookStats
 }
 #endif
 
-const FString& GetGlobalShaderMapDDCKey()
-{
-	static FString GlobalShaderMapDDCKey = GetGlobalShaderMapDDCGuid().ToString();
-	return GlobalShaderMapDDCKey;
-}
-
 const FGuid& GetGlobalShaderMapDDCGuid()
 {
 	static FGuid GlobalShaderMapDDCGuid = FDevSystemGuids::GetSystemGuid(FDevSystemGuids::Get().GLOBALSHADERMAP_DERIVEDDATA_VER);
 	return GlobalShaderMapDDCGuid;
 }
 
-const FString& GetMaterialShaderMapDDCKey()
-{
-	static FString MaterialShaderMapDDCKey = GetMaterialShaderMapDDCGuid().ToString();
-	return MaterialShaderMapDDCKey;
-}
-
 const FGuid& GetMaterialShaderMapDDCGuid()
 {
 	static FGuid MaterialShaderMapDDCGuid = FDevSystemGuids::GetSystemGuid(FDevSystemGuids::Get().MATERIALSHADERMAP_DERIVEDDATA_VER);
diff --git a/Engine/Source/Runtime/Engine/Public/MaterialShaderType.h b/Engine/Source/Runtime/Engine/Public/MaterialShaderType.h
--- a/Engine/Source/Runtime/Engine/Public/MaterialShaderType.h
+++ b/Engine/Source/Runtime/Engine/Public/MaterialShaderType.h
@@ -52,13 +52,6 @@ extern ENGINE_API FString GetShadingModelFieldString(FMaterialShadingModelField
 extern ENGINE_API FString GetBlendModeString(EBlendMode BlendMode);
 
 #if WITH_EDITOR
-/** Creates a string key for the derived data cache given a shader map id. */
-UE_DEPRECATED(5.5, "GetMaterialShaderMapKeyString now requires passing an FMaterialShaderParameters instance as input.")
-inline FString GetMaterialShaderMapKeyString(const FMaterialShaderMapId& ShaderMapId, EShaderPlatform Platform, bool bIncludeKeyStringShaderDependencies = true)
-{
-	return TEXT("INVALID_KEY");
-}
-
 extern ENGINE_API FString GetMaterialShaderMapKeyString(
 	const FMaterialShaderMapId& ShaderMapId,
 	const FMaterialShaderParameters& ShaderParameters,
diff --git a/Engine/Source/Runtime/Engine/Public/MaterialShared.h b/Engine/Source/Runtime/Engine/Public/MaterialShared.h
--- a/Engine/Source/Runtime/Engine/Public/MaterialShared.h
+++ b/Engine/Source/Runtime/Engine/Public/MaterialShared.h
@@ -1309,9 +1309,6 @@ public:
 		bool bIncludeKeyStringShaderDependencies = true) const;
 	void RecordAndEmit(FMaterialKeyGeneratorContext& Context);
 
-	UE_DEPRECATED(5.5, "Internal function only, called through AppendKeyString.")
-	void AppendStaticParametersString(FString& ParamsString) const;
-
 	const TArray<FStaticSwitchParameter> &GetStaticSwitchParameters() const 					{ return StaticSwitchParameters; }
 	const TArray<FStaticComponentMaskParameter> &GetStaticComponentMaskParameters() const 		{ return StaticComponentMaskParameters; }
 	const TArray<FStaticTerrainLayerWeightParameter> &GetTerrainLayerWeightParameters() const 	{ return TerrainLayerWeightParameters; }
@@ -1581,9 +1578,6 @@ public:
 	ENGINE_API void AddRef();
 	ENGINE_API void Release();
 
-	UE_DEPRECATED(5.5, "Please use overload accepting an FShaderSerializeContext. Note that inlining of shader code into runtime shadermaps is no longer supported. If you rely on this please reach out to the UE rendering team.")
-	bool Serialize(FArchive& Ar, bool bInlineShaderResources = true, bool bLoadingCooked = false, bool bInlineShaderCode = false, const FName& SerializingAsset = NAME_None);
-
 	/** Serializes the shader map. */
 	bool Serialize(FShaderSerializeContext& Ctx);
 
@@ -2606,8 +2600,6 @@ public:
 	/** Allows to associate the shader resources with the asset for load order. */
 	virtual FName GetAssetPath() const { return NAME_None; };
 
-	UE_DEPRECATED(5.5, "ShouldInlineShaderCode is no longer used when packaging shader maps. If you rely on this for some reason please reach out to the UE rendering team.")
-	virtual bool ShouldInlineShaderCode() const { return false; }
 #endif // WITH_EDITOR
 
 	virtual FString GetFullPath() const { return TEXT(""); };
@@ -3071,8 +3063,6 @@ public:
 	ENGINE_API virtual void NotifyCompilationFinished() override;
 	/** Allows to associate the shader resources with the asset for load order. */
 	ENGINE_API virtual FName GetAssetPath() const override;
-	UE_DEPRECATED(5.5, "ShouldInlineShaderCode is no longer used when packaging shader maps. If you rely on this for some reason please reach out to the UE rendering team.")
-	virtual bool ShouldInlineShaderCode() const override { return false; }
 	ENGINE_API virtual bool IsUsingNewHLSLGenerator() const override;
 	ENGINE_API virtual bool CheckInValidStateForCompilation(FMaterialCompiler* Compiler) const override;
 #endif // WITH_EDITOR
diff --git a/Engine/Source/Runtime/Engine/Public/ShaderCompiler.h b/Engine/Source/Runtime/Engine/Public/ShaderCompiler.h
--- a/Engine/Source/Runtime/Engine/Public/ShaderCompiler.h
+++ b/Engine/Source/Runtime/Engine/Public/ShaderCompiler.h
@@ -1526,12 +1526,7 @@ extern ENGINE_API void CompileGlobalShaderMap(EShaderPlatform Platform, bool bRe
 extern ENGINE_API void CompileGlobalShaderMap(EShaderPlatform Platform, const ITargetPlatform* TargetPlatform, bool bRefreshShaderMap);
 extern ENGINE_API void ShutdownGlobalShaderMap();
 
-UE_DEPRECATED(5.5, "Use GetGlobalShaderMapDDCGuid")
-extern ENGINE_API const FString& GetGlobalShaderMapDDCKey();
 extern ENGINE_API const FGuid& GetGlobalShaderMapDDCGuid();
-
-UE_DEPRECATED(5.5, "Use GetMaterialShaderMapDDCGuid")
-extern ENGINE_API const FString& GetMaterialShaderMapDDCKey();
 extern ENGINE_API const FGuid& GetMaterialShaderMapDDCGuid();
 
 extern ENGINE_API bool ShouldDumpShaderDDCKeys();
diff --git a/Engine/Source/Runtime/RenderCore/Private/ShaderCore.cpp b/Engine/Source/Runtime/RenderCore/Private/ShaderCore.cpp
--- a/Engine/Source/Runtime/RenderCore/Private/ShaderCore.cpp
+++ b/Engine/Source/Runtime/RenderCore/Private/ShaderCore.cpp
@@ -1415,18 +1415,6 @@ void FShaderCompilerEnvironment::SetDefine(FShaderCompilerDefineNameCache& Name,
 	}
 }
 
-int32 FShaderCompilerEnvironment::GetIntegerValue(FName Name) const
-{
-	checkf(Definitions.IsValid(), TEXT("GetIntegerValue is not supported on FShaderCompilerEnvironment in hashing mode"));
-	return Definitions->GetIntegerValue(Name);
-}
-
-int32 FShaderCompilerEnvironment::GetIntegerValue(FShaderCompilerDefineNameCache& NameCache, int32 ResultIfNotFound) const
-{
-	checkf(Definitions.IsValid(), TEXT("GetIntegerValue is not supported on FShaderCompilerEnvironment in hashing mode"));
-	return Definitions->GetIntegerValue(NameCache, ResultIfNotFound);
-}
-
 bool FShaderCompilerEnvironment::ContainsDefinition(FName Name) const
 {
 	if (Definitions.IsValid())
diff --git a/Engine/Source/Runtime/RenderCore/Private/ShaderMap.cpp b/Engine/Source/Runtime/RenderCore/Private/ShaderMap.cpp
--- a/Engine/Source/Runtime/RenderCore/Private/ShaderMap.cpp
+++ b/Engine/Source/Runtime/RenderCore/Private/ShaderMap.cpp
@@ -271,14 +271,6 @@ void FShaderMapBase::UnfreezeContent()
 
 #define CHECK_SHADERMAP_DEPENDENCIES (WITH_EDITOR || !(UE_BUILD_SHIPPING || UE_BUILD_TEST))
 
-bool FShaderMapBase::Serialize(FArchive& Ar, bool bInlineShaderResources, bool bLoadingCooked, bool bInlineShaderCode, const FName& SerializingAsset)
-{
-	FShaderSerializeContext Ctx(Ar);
-	Ctx.bLoadingCooked = bLoadingCooked;
-	Ctx.SerializingAsset = SerializingAsset;
-	return Serialize(Ctx);
-}
-
 bool FShaderMapBase::Serialize(FShaderSerializeContext& Ctx)
 {
 	FArchive& Ar = Ctx.GetMainArchive();
diff --git a/Engine/Source/Runtime/RenderCore/Private/ShaderResource.cpp b/Engine/Source/Runtime/RenderCore/Private/ShaderResource.cpp
--- a/Engine/Source/Runtime/RenderCore/Private/ShaderResource.cpp
+++ b/Engine/Source/Runtime/RenderCore/Private/ShaderResource.cpp
@@ -414,13 +414,6 @@ void FShaderMapResourceCode::ToString(FStringBuilderBase& OutString) const
 	}
 }
 
-void FShaderMapResourceCode::Serialize(FArchive& Ar, bool bLoadingCooked)
-{
-	FShaderSerializeContext Ctx(Ar);
-	Ctx.bLoadingCooked = bLoadingCooked;
-	Serialize(Ctx);
-}
-
 void FShaderMapResourceCode::Serialize(FShaderSerializeContext& Ctx)
 {
 	FArchive& Ar = Ctx.GetMainArchive();
diff --git a/Engine/Source/Runtime/RenderCore/Public/GlobalShader.h b/Engine/Source/Runtime/RenderCore/Public/GlobalShader.h
--- a/Engine/Source/Runtime/RenderCore/Public/GlobalShader.h
+++ b/Engine/Source/Runtime/RenderCore/Public/GlobalShader.h
@@ -148,12 +148,6 @@ public:
 	RENDERCORE_API static FGlobalShaderMapSection* CreateFromCache(FShaderCacheLoadContext& Ctx);
 #endif
 
-	UE_DEPRECATED(5.5, "Use overload accepting a FShaderSerializeContext instead")
-	bool Serialize(FArchive& Ar)
-	{
-		return false;
-	}
-
 	RENDERCORE_API bool Serialize(FShaderSerializeContext& Ctx);
 	RENDERCORE_API virtual void GetShaderList(TMap<FHashedName, TShaderRef<FShader>>& OutShaders) const override;
 	RENDERCORE_API virtual void GetShaderPipelineList(TArray<FShaderPipelineRef>& OutShaderPipelines) const override;
diff --git a/Engine/Source/Runtime/RenderCore/Public/Shader.h b/Engine/Source/Runtime/RenderCore/Public/Shader.h
--- a/Engine/Source/Runtime/RenderCore/Public/Shader.h
+++ b/Engine/Source/Runtime/RenderCore/Public/Shader.h
@@ -514,21 +514,6 @@ private:
 class FShaderMapResourceCode : public FThreadSafeRefCountedObject
 {
 public:
-	struct UE_DEPRECATED(5.5, "Use FShaderCodeResource") FShaderEntry
-	{
-		TArray<uint8> Code;
-		int32 UncompressedSize;
-		EShaderFrequency Frequency;
-
-		friend FArchive& operator<<(FArchive& Ar, FShaderEntry& Entry)
-		{
-			uint8 Freq = Entry.Frequency;
-			Ar << Entry.Code << Entry.UncompressedSize << Freq;
-			Entry.Frequency = (EShaderFrequency)Freq;
-			return Ar;
-		}
-	};
-
 #if WITH_EDITORONLY_DATA
 	struct FShaderEditorOnlyDataEntry
 	{
@@ -583,9 +568,6 @@ public:
 
 	RENDERCORE_API void Finalize();
 
-	UE_DEPRECATED(5.5, "Use overload accepting an FShaderSerializeContext object")
-	RENDERCORE_API void Serialize(FArchive& Ar, bool bLoadingCooked);
-
 	RENDERCORE_API void Serialize(FShaderSerializeContext& Ctx);
 #if WITH_EDITORONLY_DATA
 	RENDERCORE_API void NotifyShadersCompiled(FName FormatName);
@@ -613,11 +595,6 @@ public:
 	FSHAHash ResourceHash;
 	TArray<FSHAHash> ShaderHashes;
 
-	UE_DEPRECATED(5.5, "Use ShaderCodeResources") 
-	PRAGMA_DISABLE_DEPRECATION_WARNINGS
-	TArray<FShaderEntry> ShaderEntries;
-	PRAGMA_ENABLE_DEPRECATION_WARNINGS
-
 	TArray<FShaderCodeResource> ShaderCodeResources;
 #if WITH_EDITORONLY_DATA
 	// Optional array of editor-only data indexed in the same order as ShaderEntries (sorted by the shader hash)
@@ -918,12 +895,6 @@ public:
 	inline uint32 GetNumTextureSamplers() const { return NumTextureSamplers; }
 	inline uint32 GetCodeSize() const { return CodeSize; }
 	inline void SetNumInstructions(uint32 Value) { NumInstructions = Value; }
-	UE_DEPRECATED(5.5, "GetShaderStatistics which returns a reference to the stat map is deprecated. Shader Stats can be accessed on the editor only data on the FShaderMapResourceCode object.")
-	inline const FShaderStatisticMap& GetShaderStatistics() const
-	{
-		static FShaderStatisticMap DummyStatMap;
-		return DummyStatMap;
-	}
 #else
 	inline uint32 GetNumInstructions() const { return 0u; }
 	inline uint32 GetNumTextureSamplers() const { return 0u; }
@@ -1531,13 +1502,6 @@ public:
 	}
 
 #if WITH_EDITOR
-	UE_DEPRECATED(5.5, "GetReferencedUniformBufferNames is deprecated, use GetReferencedUniformBuffers instead")
-	inline const TSet<const TCHAR*, TStringPointerSetKeyFuncs_DEPRECATED<const TCHAR*>>& GetReferencedUniformBufferNames() const
-	{
-		static TSet<const TCHAR*, TStringPointerSetKeyFuncs_DEPRECATED<const TCHAR*>> EmptySet;
-		return EmptySet;
-	}
-
 	const TSet<const FShaderParametersMetadata*>& GetReferencedUniformBuffers() const
 	{
 		return ReferencedUniformBuffers;
@@ -2444,13 +2408,6 @@ public:
 	RENDERCORE_API void GetOutdatedTypes(const FShaderMapBase& InShaderMap, TArray<const FShaderType*>& OutdatedShaderTypes, TArray<const FShaderPipelineType*>& OutdatedShaderPipelineTypes, TArray<const FVertexFactoryType*>& OutdatedFactoryTypes) const;
 
 	RENDERCORE_API void SaveShaderStableKeys(const FShaderMapBase& InShaderMap, EShaderPlatform TargetShaderPlatform, const struct FStableShaderKeyAndValue& SaveKeyVal);
-
-	UE_DEPRECATED(5.5, "GetShaderStatistics which returns a reference to the stat map is deprecated. Please use the function that returns a non-ref.")
-	const FShader::FShaderStatisticMap GetShaderStatisticsMapForShader(const FShaderMapBase& InShaderMap, FShaderType* ShaderType) const
-	{
-		static FShader::FShaderStatisticMap DummyStatMap;
-		return DummyStatMap;
-	}
 #endif // WITH_EDITOR
 
 #if WITH_EDITORONLY_DATA
@@ -2543,9 +2500,6 @@ public:
 
 	RENDERCORE_API void FinalizeContent();
 	RENDERCORE_API void UnfreezeContent();
-	UE_DEPRECATED(5.5, "Please use overload accepting an FShaderSerializeContext. Note that inlining of shader code into runtime shadermaps is no longer supported. If you rely on this please reach out to the UE rendering team.")
-	RENDERCORE_API bool Serialize(FArchive& Ar, bool bInlineShaderResources, bool bLoadingCooked, bool bInlineShaderCode = false, const FName& SerializingAsset = NAME_None);
-
 	RENDERCORE_API bool Serialize(FShaderSerializeContext& Ctx);
 
 	EShaderPermutationFlags GetPermutationFlags() const
diff --git a/Engine/Source/Runtime/RenderCore/Public/ShaderCodeArchive.h b/Engine/Source/Runtime/RenderCore/Public/ShaderCodeArchive.h
--- a/Engine/Source/Runtime/RenderCore/Public/ShaderCodeArchive.h
+++ b/Engine/Source/Runtime/RenderCore/Public/ShaderCodeArchive.h
@@ -275,11 +275,6 @@ public:
 	RENDERCORE_API void Finalize();
 #endif
 
-	UE_DEPRECATED(5.5, "Use overload accepting a TArray<FSharedBuffer> for shader code")
-	void DecompressShader(int32 Index, const TArray<TArray<uint8>>& ShaderCode, TArray<uint8>& OutDecompressedShader) const
-	{
-	}
-
 	RENDERCORE_API void DecompressShader(int32 Index, const TArray<FSharedBuffer>& ShaderCode, TArray<uint8>& OutDecompressedShader) const;
 
 	RENDERCORE_API void Serialize(FArchive& Ar);
diff --git a/Engine/Source/Runtime/RenderCore/Public/ShaderCompilerCore.h b/Engine/Source/Runtime/RenderCore/Public/ShaderCompilerCore.h
--- a/Engine/Source/Runtime/RenderCore/Public/ShaderCompilerCore.h
+++ b/Engine/Source/Runtime/RenderCore/Public/ShaderCompilerCore.h
@@ -202,13 +202,6 @@ struct FShaderCompilerInput
 	// Indicates which additional debug outputs should be written for this compile job.
 	EShaderDebugInfoFlags DebugInfoFlags = EShaderDebugInfoFlags::Default;
 
-	UE_DEPRECATED(5.5, "bCachePreprocessed member no longer used; preprocessed job cache is now the only path for in-memory caching of shader jobs")
-	// True if the cache key for this job should be based on preprocessed source. If so,
-	// preprocessing will be executed in the cook process independent of compilation (and
-	// as such this will only ever be set for jobs whose shader format supports independent
-	// preprocessing)
-	bool bCachePreprocessed = true;
-
 	// Array of symbols that should be maintained when deadstripping. If this is empty, entry
 	// point name alone will be used.
 	TArray<FString> RequiredSymbols;
@@ -252,16 +245,6 @@ struct FShaderCompilerInput
 	/** Oodle-specific compression level - used if CompressionFormat is set to NAME_Oodle. */
 	FOodleDataCompression::ECompressionLevel OodleLevel;
 
-	PRAGMA_DISABLE_DEPRECATION_WARNINGS
-	// Explicitly-defaulted ctors & assignment operators are needed temporarily due to deprecation
-	// of bCachePreprocessed field. Can be removed once the deprecation window for said field ends.
-	FShaderCompilerInput() = default;
-	FShaderCompilerInput(FShaderCompilerInput&&) = default;
-	FShaderCompilerInput(const FShaderCompilerInput&) = default;
-	FShaderCompilerInput& operator=(FShaderCompilerInput&&) = default;
-	FShaderCompilerInput& operator=(const FShaderCompilerInput&) = default;
-	PRAGMA_ENABLE_DEPRECATION_WARNINGS
-
 	bool DumpDebugInfoEnabled() const 
 	{
 		return DumpDebugInfoPath != TEXT("") && IFileManager::Get().DirectoryExists(*DumpDebugInfoPath);
@@ -306,18 +289,6 @@ struct FShaderCompilerInput
 		return Environment.CompilerFlags.Contains(CFLAG_GenerateSymbolsInfo) ? DebugGroupName : FString();
 	}
 
-	UE_DEPRECATED(5.5, "GatherSharedInputs no longer processes external includes, as these are not needed by SCW anymore")
-	void GatherSharedInputs(
-		TMap<FString, FString>& ExternalIncludes,
-		TArray<TRefCountPtr<FSharedShaderCompilerEnvironment>>& SharedEnvironments,
-		TArray<const FShaderParametersMetadata*>& ParametersStructures) {}
-
-	UE_DEPRECATED(5.5, "GatherSharedInputsAnsi is no longer needed since external includes are no longer serialized; use GetSharedInputs")
-	void GatherSharedInputsAnsi(
-		TMap<FString, TArray<ANSICHAR>>& ExternalIncludes,
-		TArray<TRefCountPtr<FSharedShaderCompilerEnvironment>>& SharedEnvironments,
-		TArray<const FShaderParametersMetadata*>& ParametersStructures) {}
-
 	void GatherSharedInputs(
 		TArray<TRefCountPtr<FSharedShaderCompilerEnvironment>>& SharedEnvironments,
 		TArray<const FShaderParametersMetadata*>& ParametersStructures)
@@ -498,8 +469,6 @@ struct FShaderCompilerOutput
 
 	FShaderParameterMap ParameterMap;
 	TArray<FShaderCompilerError> Errors;
-	UE_DEPRECATED(5.5, "PragmaDirectives are now only accessible via the FShaderPreprocessOutput struct")
-	TArray<FString> PragmaDirectives;
 	FShaderTarget Target;
 	FShaderCode ShaderCode;
 	FSHAHash OutputHash;
@@ -517,7 +486,7 @@ struct FShaderCompilerOutput
 
 	PRAGMA_DISABLE_DEPRECATION_WARNINGS
 	// Explicitly-defaulted copy/move ctors & assignment operators are needed temporarily due to deprecation
-	// of PragmaDirectives and bSerializeModifiedSource fields. These can be removed once the deprecation window for said fields ends.
+	// of bSerializeModifiedSource field. These can be removed once the deprecation window for this field ends.
 	FShaderCompilerOutput(FShaderCompilerOutput&&) = default;
 	FShaderCompilerOutput(const FShaderCompilerOutput&) = default;
 	FShaderCompilerOutput& operator=(FShaderCompilerOutput&&) = default;
@@ -661,8 +630,6 @@ extern RENDERCORE_API int HandleShaderCompileException(Windows::LPEXCEPTION_POIN
 extern RENDERCORE_API const IShaderFormat* FindShaderFormat(FName Format, const TArray<const IShaderFormat*>& ShaderFormats);
 
 // Executes preprocessing for the given job, if the job is marked to be preprocessed independently prior to compilation.
-UE_DEPRECATED(5.5, "ConditionalPreprocessShader is now just PreprocessShader (no longer conditional, always executes in the main process at job submission time)")
-extern RENDERCORE_API bool ConditionalPreprocessShader(FShaderCommonCompileJob* Job);
 extern RENDERCORE_API bool PreprocessShader(FShaderCommonCompileJob* Job);
 extern RENDERCORE_API void CompileShader(const TArray<const IShaderFormat*>& ShaderFormats, FShaderCompileJob& Job, const FString& WorkingDirectory, int32* CompileCount = nullptr);
 extern RENDERCORE_API void CompileShaderPipeline(const TArray<const IShaderFormat*>& ShaderFormats, FShaderPipelineCompileJob* PipelineJob, const FString& WorkingDirectory, int32* CompileCount = nullptr);
diff --git a/Engine/Source/Runtime/RenderCore/Public/ShaderCompilerJobTypes.h b/Engine/Source/Runtime/RenderCore/Public/ShaderCompilerJobTypes.h
--- a/Engine/Source/Runtime/RenderCore/Public/ShaderCompilerJobTypes.h
+++ b/Engine/Source/Runtime/RenderCore/Public/ShaderCompilerJobTypes.h
@@ -145,8 +145,6 @@ public:
 	FShaderCommonCompileJob* NextLink = nullptr;
 	FShaderCommonCompileJob** PrevLink = nullptr;
 
-	using FInputHash UE_DEPRECATED(5.5, "Use FShaderCompilerInputHash instead of FShaderCommonCompileJob::FInputHash") = FShaderCompilerInputHash;
-
 	FPendingShaderMapCompileResultsPtr PendingShaderMap;
 
 	mutable FThreadSafeCounter NumRefs;
@@ -231,9 +229,6 @@ public:
 	/** Returns hash of all inputs for this job (needed for caching). */
 	virtual FShaderCompilerInputHash GetInputHash() { return FShaderCompilerInputHash(); }
 
-	UE_DEPRECATED(5.5, "Use overload accepting an FShaderCacheSerializeContext struct")
-	virtual void SerializeOutput(FArchive& Ar) {}
-
 	/** Serializes (and deserializes) the output for caching purposes. */
 	virtual void SerializeOutput(FShaderCacheSerializeContext& Ctx) {}
 
@@ -344,11 +339,6 @@ public:
 
 	virtual RENDERCORE_API FShaderCompilerInputHash GetInputHash() override;
 
-	UE_DEPRECATED(5.5, "Use overload accepting an FShaderCacheSerializeContext")
-	virtual void SerializeOutput(FArchive& Ar) override
-	{
-	}
-
 	RENDERCORE_API void SerializeOutput(FShaderCacheSerializeContext& Ctx, int32 CodeIndex);
 	virtual void SerializeOutput(FShaderCacheSerializeContext& Ctx)
 	{
@@ -417,11 +407,6 @@ public:
 
 	virtual RENDERCORE_API FShaderCompilerInputHash GetInputHash() override;
 
-	UE_DEPRECATED(5.5, "Use overload accepting an FShaderCommonCompileJob::FSerializationContext")
-	virtual void SerializeOutput(FArchive& Ar) override
-	{
-	}
-
 	virtual RENDERCORE_API void SerializeOutput(FShaderCacheSerializeContext& Ctx) override;
 
 	virtual RENDERCORE_API void OnComplete() override;
diff --git a/Engine/Source/Runtime/RenderCore/Public/ShaderCore.h b/Engine/Source/Runtime/RenderCore/Public/ShaderCore.h
--- a/Engine/Source/Runtime/RenderCore/Public/ShaderCore.h
+++ b/Engine/Source/Runtime/RenderCore/Public/ShaderCore.h
@@ -602,22 +602,10 @@ struct FShaderCompilerEnvironment
 	RENDERCORE_API void SetDefine(FShaderCompilerDefineNameCache& Name, bool Value);
 	RENDERCORE_API void SetDefine(FShaderCompilerDefineNameCache& Name, float Value);
 
-	UE_DEPRECATED(5.5,	"Accessing definition values via FShaderCompilerEnvironment is no longer allowed outside of the core shader system.\n"
-						"Please use Get/SetCompileArgument mechanism for additional properties required by compilation.")
-	RENDERCORE_API int32 GetIntegerValue(FName Name) const;
-	UE_DEPRECATED(5.5,	"Accessing definition values via FShaderCompilerEnvironment is no longer allowed outside of the core shader system.\n"
-						"Please use Get/SetCompileArgument mechanism for additional properties required by compilation.")
-	RENDERCORE_API int32 GetIntegerValue(FShaderCompilerDefineNameCache& NameCache, int32 ResultIfNotFound = 0) const;
-
-	UE_DEPRECATED(5.5, "ContainsDefinition will be made private in the future and should not be called by downstream code.")
-	RENDERCORE_API bool ContainsDefinition(FName Name) const;
-
 	template <typename ValueType> void SetDefineIfUnset(const TCHAR* Name, ValueType Value)
 	{
 		FName NameKey(Name);
-		PRAGMA_DISABLE_DEPRECATION_WARNINGS
 		if (!ContainsDefinition(NameKey))
-		PRAGMA_ENABLE_DEPRECATION_WARNINGS
 		{
 			SetDefine(NameKey, Value);
 		}
@@ -708,6 +696,7 @@ struct FShaderCompilerEnvironment
 	RENDERCORE_API FString GetDefinitionsAsCommentedCode() const;
 
 private:
+	RENDERCORE_API bool ContainsDefinition(FName Name) const;
 
 	friend class FShaderCompileUtilities;
 	friend class FShaderPreprocessorUtilities;
@@ -1340,13 +1329,6 @@ public:
 		}
 	}
 
-	UE_DEPRECATED(5.5, "Use GetReadView")
-	const TArray<uint8>& GetReadAccess() const
-	{
-		static TArray<uint8> Dummy;
-		return Dummy;
-	}
-
 	// for read access, can have additional data attached to the end. Can also be compressed
 	TConstArrayView<uint8> GetReadView() const
 	{
@@ -1549,9 +1531,6 @@ extern RENDERCORE_API const class FSHAHash& GetShaderFileHash(const TCHAR* Virtu
 extern RENDERCORE_API const FSHAHash* TryGetShaderFileHash(const TCHAR* VirtualFilePath,
 	EShaderPlatform ShaderPlatform, FString* OutErrorMessage = nullptr);
 
-UE_DEPRECATED(5.5, "HashShaderFileWithIncludes is no longer used by the shader compilation system and so no longer supported.")
-extern RENDERCORE_API void HashShaderFileWithIncludes(FArchive& HashingArchive, const TCHAR* VirtualFilePath, const FString& FileContents, EShaderPlatform ShaderPlatform, bool bOnlyHashIncludedFiles);
-
 /**
  * Calculates a Hash for the list of filenames if it does not already exist in the Hash cache.
  */
diff --git a/Engine/Source/Runtime/RenderCore/Public/ShaderParameters.h b/Engine/Source/Runtime/RenderCore/Public/ShaderParameters.h
--- a/Engine/Source/Runtime/RenderCore/Public/ShaderParameters.h
+++ b/Engine/Source/Runtime/RenderCore/Public/ShaderParameters.h
@@ -39,10 +39,6 @@ namespace UE::ShaderParameters
 {
 	/** Creates a shader code declaration of this struct for the given shader platform. */
 	RENDERCORE_API FString CreateUniformBufferShaderDeclaration(const TCHAR* Name, const FShaderParametersMetadata& UniformBufferStruct, const FRHIUniformBufferShaderBindingLayout* UniformBufferSBLayout);
-
-	UE_DEPRECATED(5.5, "AddUniformBufferIncludesToEnviroment now takes a set of FShaderParametersMetadata pointers")
-	inline void AddUniformBufferIncludesToEnvironment(FShaderCompilerEnvironment& OutEnvironment, const TSet<const TCHAR*, TStringPointerSetKeyFuncs_DEPRECATED<const TCHAR*>>& InUniformBufferNames) {}
-
 	RENDERCORE_API void AddUniformBufferIncludesToEnvironment(FShaderCompilerEnvironment& OutEnvironment, const TSet<const FShaderParametersMetadata*>& InUniformBuffers);
 }
 #endif
