CVarRDGAsyncComputeSink doesn't need to check IsAsyncComputeSupported to set GRDGAsyncCompute to zero since the function checks GRDGAsyncCompute itself. Additional uses of GRDGAsyncCompute should also check IsAsyncComputeSupported, which is cached in the RDG builder.

This prevents using GMaxRHIFeatureLevel before RHIInit.

diff --git a/Engine/Source/Runtime/RenderCore/Private/RenderGraphBuilder.cpp b/Engine/Source/Runtime/RenderCore/Private/RenderGraphBuilder.cpp
--- a/Engine/Source/Runtime/RenderCore/Private/RenderGraphBuilder.cpp
+++ b/Engine/Source/Runtime/RenderCore/Private/RenderGraphBuilder.cpp
@@ -2123,7 +2123,7 @@ void FRDGBuilder::Execute()
 					RHICmdList.ImmediateFlush(EImmediateFlushType::DispatchToRHIThread);
 				}
 
-				if (GRDGDebugFlushGPU && !GRDGAsyncCompute)
+				if (GRDGDebugFlushGPU && !bSupportsAsyncCompute)
 				{
 					RHICmdList.SubmitCommandsAndFlushGPU();
 					RHICmdList.BlockUntilGPUIdle();
diff --git a/Engine/Source/Runtime/RenderCore/Private/RenderGraphPrivate.cpp b/Engine/Source/Runtime/RenderCore/Private/RenderGraphPrivate.cpp
--- a/Engine/Source/Runtime/RenderCore/Private/RenderGraphPrivate.cpp
+++ b/Engine/Source/Runtime/RenderCore/Private/RenderGraphPrivate.cpp
@@ -226,11 +226,6 @@ FAutoConsoleVariableSink CVarRDGAsyncComputeSink(FConsoleCommandDelegate::Create
 	{
 		GRDGAsyncCompute = 0;
 	}
-
-	if (!IsAsyncComputeSupported(GMaxRHIShaderPlatform))
-	{
-		GRDGAsyncCompute = 0;
-	}
 }));
 
 int32 GRDGCullPasses = 1;
