[Shaders] Don't log error on invalid virtual shader path when error is meant to be returned as output parameter.

This allows validation failures to be handled gracefully when they are only a side effect of another compile error that might be reported as warning only.

#rnx

diff --git a/Engine/Source/Runtime/RenderCore/Private/ShaderCore.cpp b/Engine/Source/Runtime/RenderCore/Private/ShaderCore.cpp
--- a/Engine/Source/Runtime/RenderCore/Private/ShaderCore.cpp
+++ b/Engine/Source/Runtime/RenderCore/Private/ShaderCore.cpp
@@ -1539,10 +1539,12 @@ static void ReportVirtualShaderFilePathError(TArray<FShaderCompilerError>* Compi
 {
 	if (CompileErrors)
 	{
-		CompileErrors->Add(FShaderCompilerError(*ErrorString));
+		CompileErrors->Add(FShaderCompilerError(MoveTemp(ErrorString)));
+	}
+	else
+	{
+		UE_LOG(LogShaders, Error, TEXT("%s"), *ErrorString);
 	}
-
-	UE_LOG(LogShaders, Error, TEXT("%s"), *ErrorString);
 }
 
 
@@ -3894,7 +3896,7 @@ void FShaderCompilerError::ExtractSourceLocations(TArray<FShaderCompilerError>&
 	}
 }
 
-FString FShaderCompilerError::GetShaderSourceFilePath() const
+FString FShaderCompilerError::GetShaderSourceFilePath(TArray<FShaderCompilerError>* InOutErrors) const
 {
 	// Always return error file path as-is if it doesn't denote a virtual path.
 	// We don't wont to report errors when accessing a compile error's message.
@@ -3904,7 +3906,7 @@ FString FShaderCompilerError::GetShaderSourceFilePath() const
 	}
 	else
 	{
-		return ::GetShaderSourceFilePath(ErrorVirtualFilePath, nullptr);
+		return ::GetShaderSourceFilePath(ErrorVirtualFilePath, InOutErrors);
 	}
 }
 
diff --git a/Engine/Source/Runtime/RenderCore/Private/ShaderDiagnostics.cpp b/Engine/Source/Runtime/RenderCore/Private/ShaderDiagnostics.cpp
--- a/Engine/Source/Runtime/RenderCore/Private/ShaderDiagnostics.cpp
+++ b/Engine/Source/Runtime/RenderCore/Private/ShaderDiagnostics.cpp
@@ -97,7 +97,8 @@ int32 FShaderDiagnosticInfo::AddAndProcessErrorsForFailedJobFiltered(FShaderComp
 					bReportedDebugInfo = true;
 				}
 
-				const FString ShaderFilePath = ConvertToNativePlatformAbsolutePath(*CurrentError.GetShaderSourceFilePath());
+				TArray<FShaderCompilerError> SecondaryErrorsFromFilePath;
+				const FString ShaderFilePath = ConvertToNativePlatformAbsolutePath(*CurrentError.GetShaderSourceFilePath(&SecondaryErrorsFromFilePath));
 				const FString ShaderErrorLineString = CurrentError.ErrorLineString.IsEmpty() ? TEXT("0") : *CurrentError.ErrorLineString;
 				if (CurrentJob.Key.ShaderType)
 				{
@@ -116,8 +117,15 @@ int32 FShaderDiagnosticInfo::AddAndProcessErrorsForFailedJobFiltered(FShaderComp
 						*ShaderErrorLineString);
 				}
 
-				UniqueErrorHashes.Add(ErrorHash);
+				// Append secondary errors resulting from invalid file path
+				for (const FShaderCompilerError& SecondaryError : SecondaryErrorsFromFilePath)
+				{
+					FString SecondaryErrorString = SecondaryError.GetErrorString();
+					CurrentErrorString.AppendChar(TEXT('\n'));
+					CurrentErrorString.Append(SecondaryErrorString);
+				}
 
+				UniqueErrorHashes.Add(ErrorHash);
 				UniqueErrors.Add(UniqueErrorPrefix + CurrentErrorString);
 				ErrorJobs.AddUnique(&CurrentJob);
 			}
diff --git a/Engine/Source/Runtime/RenderCore/Public/ShaderCompilerCore.h b/Engine/Source/Runtime/RenderCore/Public/ShaderCompilerCore.h
--- a/Engine/Source/Runtime/RenderCore/Public/ShaderCompilerCore.h
+++ b/Engine/Source/Runtime/RenderCore/Public/ShaderCompilerCore.h
@@ -442,7 +442,7 @@ struct FShaderCompilerError
 	static void RENDERCORE_API ExtractSourceLocations(TArray<FShaderCompilerError>& InOutErrors);
 
 	/** Returns the path of the underlying source file relative to the process base dir. */
-	FString RENDERCORE_API GetShaderSourceFilePath() const;
+	FString RENDERCORE_API GetShaderSourceFilePath(TArray<FShaderCompilerError>* InOutErrors = nullptr) const;
 
 	friend FArchive& operator<<(FArchive& Ar,FShaderCompilerError& Error)
 	{
@@ -637,7 +637,7 @@ extern RENDERCORE_API void CompileShaderPipeline(const TArray<const IShaderForma
 /**
  * Validates the format of a virtual shader file path.
  * Meant to be use as such: check(CheckVirtualShaderFilePath(VirtualFilePath));
- * CompileErrors output array is optional.
+ * CompileErrors output array is optional. If this is non-null, all validation errors are returned to this array instead of logging them to LogShaders.
  */
 extern RENDERCORE_API bool CheckVirtualShaderFilePath(FStringView VirtualPath, TArray<FShaderCompilerError>* CompileErrors = nullptr);
 
