When texture declare with UBMT_RDG_TEXTURE_NON_PIXEL_SRV, remove the pixel bit to be able to track the proper barrier to apply in the platform low level implementation of the barrier transitions. Wrapper around a cvar to have the possibility to disable the optimization. Proper fix will be to add an UBMT_RDG_TEXTURE_PIXEL_SRV in addition to the non-pixel version, then when we mark up specific shaders that need those barriers we'd use the pixel / non-pixel versions instead of the _TEXTURE ones. This will allow to keep SRVGraphics = both pixel and non pixel barrier, and progressively change the places where the _PIXEL_SRV could speed up the GPU.

#tests +local replayruns

diff --git a/Engine/Source/Runtime/RenderCore/Private/RenderGraphBuilder.cpp b/Engine/Source/Runtime/RenderCore/Private/RenderGraphBuilder.cpp
--- a/Engine/Source/Runtime/RenderCore/Private/RenderGraphBuilder.cpp
+++ b/Engine/Source/Runtime/RenderCore/Private/RenderGraphBuilder.cpp
@@ -135,8 +135,12 @@ void EnumerateTextureAccess(FRDGParameterStruct PassParameters, ERDGPassFlags Pa
 		case UBMT_RDG_TEXTURE_NON_PIXEL_SRV:
 			if (FRDGTextureSRVRef SRV = Parameter.GetAsTextureSRV())
 			{
-				const ERHIAccess SRVGraphicsNonPixel = Parameter.GetType() == UBMT_RDG_TEXTURE_NON_PIXEL_SRV ? ERHIAccess::SRVGraphicsNonPixel : ERHIAccess::None;
-				AccessFunction(SRV, SRV->GetParent(), SRVAccess | SRVGraphicsNonPixel, NoneFlags, SRV->GetSubresourceRange());
+				ERHIAccess CurrentSRVAccess = SRVAccess;
+				if (Parameter.GetType() == UBMT_RDG_TEXTURE_NON_PIXEL_SRV)
+				{
+					EnumRemoveFlags(CurrentSRVAccess, ERHIAccess::SRVGraphicsPixel);
+				}
+				AccessFunction(SRV, SRV->GetParent(), CurrentSRVAccess, NoneFlags, SRV->GetSubresourceRange());
 			}
 		break;
 		case UBMT_RDG_TEXTURE_UAV:
