Fixed use-after-free in FRenderCommandList.

diff --git a/Engine/Source/Runtime/RenderCore/Private/RenderingThread.cpp b/Engine/Source/Runtime/RenderCore/Private/RenderingThread.cpp
--- a/Engine/Source/Runtime/RenderCore/Private/RenderingThread.cpp
+++ b/Engine/Source/Runtime/RenderCore/Private/RenderingThread.cpp
@@ -1713,9 +1713,12 @@ FRenderCommandList::FRenderCommandList(ERenderCommandListFlags InFlags, EPageSiz
 
 FRenderCommandList::~FRenderCommandList()
 {
-	for (FRenderCommandList* Child = Children.Head; Child; Child = Child->NextSibling)
+	FRenderCommandList* Child = Children.Head;
+	while (Child)
 	{
+		FRenderCommandList* NextChild = Child->NextSibling;
 		delete Child;
+		Child = NextChild;
 	}
 	Children = {};
 
