Fixed static analysis warning around dereferencing a null pointer.

diff --git a/Engine/Source/Runtime/RenderCore/Private/RenderingThread.cpp b/Engine/Source/Runtime/RenderCore/Private/RenderingThread.cpp
--- a/Engine/Source/Runtime/RenderCore/Private/RenderingThread.cpp
+++ b/Engine/Source/Runtime/RenderCore/Private/RenderingThread.cpp
@@ -1577,7 +1577,7 @@ void FCommandList::Release()
 
 //////////////////////////////////////////////////////////////////////////
 
-void ExecuteCommand(FRHICommandList& RHICmdList, const FRenderCommandFunctionVariant& Function, const FRenderCommandTag& Tag)
+static void ExecuteCommand(FRHICommandList* RHICmdList, const FRenderCommandFunctionVariant& Function, const FRenderCommandTag& Tag)
 {
 	TRACE_CPUPROFILER_EVENT_SCOPE_USE_ON_CHANNEL(Tag.GetSpecId(), Tag.GetName(), EventScope, RenderCommandsChannel, true);
 	FScopeCycleCounter Scope(Tag.GetStatId(), true);
@@ -1585,11 +1585,11 @@ void ExecuteCommand(FRHICommandList& RHICmdList, const FRenderCommandFunctionVar
 	switch (Function.GetIndex())
 	{
 	case FRenderCommandFunctionVariant::IndexOfType<TUniqueFunction<void(FRHICommandListImmediate&)>>():
-		Function.Get<TUniqueFunction<void(FRHICommandListImmediate&)>>()(RHICmdList.GetAsImmediate());
+		Function.Get<TUniqueFunction<void(FRHICommandListImmediate&)>>()(RHICmdList->GetAsImmediate());
 		break;
 
 	case FRenderCommandFunctionVariant::IndexOfType<TUniqueFunction<void(FRHICommandList&)>>():
-		Function.Get<TUniqueFunction<void(FRHICommandList&)>>()(RHICmdList);
+		Function.Get<TUniqueFunction<void(FRHICommandList&)>>()(*RHICmdList);
 		break;
 
 	case FRenderCommandFunctionVariant::IndexOfType<TUniqueFunction<void()>>():
@@ -1951,7 +1951,7 @@ void FRenderThreadCommandPipe::ExecuteCommands(UE::RenderCommandPipe::FCommandLi
 {
 	CommandList.ConsumeCommands([&RHICmdList = GetImmediateCommandList_ForRenderCommand()] (FRenderCommandFunctionVariant&& Function, const FRenderCommandTag& Tag)
 	{
-		ExecuteCommand(RHICmdList, Function, Tag);
+		ExecuteCommand(&RHICmdList, Function, Tag);
 	});
 }
 
@@ -2360,7 +2360,7 @@ void FRenderCommandPipe::ExecuteCommand(FRenderCommandFunctionVariant&& Function
 		RHICmdList->SwitchPipeline(ERHIPipeline::Graphics);
 	}
 
-	::ExecuteCommand(*RHICmdList, FunctionVariant, Tag);
+	::ExecuteCommand(RHICmdList, FunctionVariant, Tag);
 }
 
 void FRenderCommandPipe::EnqueueAndLaunch(FRenderCommandList* CommandList)
