Fixed crash due to render command pipe not being initialized.

#bs-asset-validation-passed: 1745977 v2_AAAAAGgePec=

diff --git a/Engine/Source/Runtime/RenderCore/Public/RenderingThread.h b/Engine/Source/Runtime/RenderCore/Public/RenderingThread.h
--- a/Engine/Source/Runtime/RenderCore/Public/RenderingThread.h
+++ b/Engine/Source/Runtime/RenderCore/Public/RenderingThread.h
@@ -661,8 +661,14 @@ public:
 		return Name;
 	}
 
+	FORCEINLINE bool IsValid() const
+	{
+		return Index != INDEX_NONE;
+	}
+
 	FORCEINLINE int32 GetIndex() const
 	{
+		check(IsValid());
 		return Index;
 	}
 
@@ -790,7 +796,7 @@ private:
 	FAutoConsoleVariable ConsoleVariable;
 	std::atomic_int32_t NumInFlightCommands{ 0 };
 	std::atomic_int32_t NumInFlightCommandLists{ 0 };
-	uint16 Index = uint16(-1);
+	int32 Index = INDEX_NONE;
 	bool bRecording = false;
 	bool bReplaying = false;
 	bool bEnabled = true;
@@ -1013,7 +1019,7 @@ private:
 	FORCEINLINE UE::RenderCommandPipe::FCommandList& Get(FRenderCommandPipe* Pipe)
 	{
 		LazyInit();
-		return Pipe ? CommandLists[Pipe->GetIndex()] : GetRenderThread();
+		return Pipe && Pipe->IsValid() ? CommandLists[Pipe->GetIndex()] : GetRenderThread();
 	}
 
 	FMemStackBase Allocator;
