Fixed deadlock due to FRenderCommandList not flushing commands before waiting on a render thread fence. Added a FRenderCommandList::FFlushScope which attempts to flush the currently recording command list and unbinds it temporarily before issuing the fence.

diff --git a/Engine/Source/Runtime/RenderCore/Private/RenderingThread.cpp b/Engine/Source/Runtime/RenderCore/Private/RenderingThread.cpp
--- a/Engine/Source/Runtime/RenderCore/Private/RenderingThread.cpp
+++ b/Engine/Source/Runtime/RenderCore/Private/RenderingThread.cpp
@@ -1316,6 +1316,7 @@ void FRenderCommandFence::Wait(bool bProcessGameThreadTasks) const
 {
 	if (!IsFenceComplete())
 	{
+		FRenderCommandList::FFlushScope FlushScope;
 		FlushRenderCommandFenceBundler();
 		GameThreadWaitForTask(CompletionTask, bProcessGameThreadTasks);
 		CompletionTask = {}; // release the internal memory as soon as it's not needed anymore
@@ -1333,6 +1334,7 @@ void FlushRenderingCommands()
 	}
 
 	TRACE_CPUPROFILER_EVENT_SCOPE(FlushRenderingCommands);
+	FRenderCommandList::FFlushScope FlushScope;
 	FCoreRenderDelegates::OnFlushRenderingCommandsStart.Broadcast();
 	FSuspendRenderingTickables SuspendRenderingTickables;
 
@@ -1623,9 +1625,12 @@ FRenderCommandList::FRecordScope::FRecordScope(FRenderCommandList* InCommandList
 	, StopAction(InStopAction)
 {
 #if DO_CHECK
-	CommandList->NumRecordScopeRefs++;
-	checkf(CommandList->NumRecordScopeRefs != 0, TEXT("FRecordScope is either being nested recursively on the same command list or is recording from multiple threads."));
-	check(CommandList->bRecording);
+	if (CommandList)
+	{
+		CommandList->NumRecordScopeRefs++;
+		checkf(CommandList->NumRecordScopeRefs != 0, TEXT("FRecordScope is either being nested recursively on the same command list or is recording from multiple threads."));
+		check(CommandList->bRecording);
+	}
 #endif
 
 	PreviousCommandList = SetInstanceTLS(InCommandList);
@@ -1635,18 +1640,39 @@ FRenderCommandList::FRecordScope::~FRecordScope()
 {
 	SetInstanceTLS(PreviousCommandList);
 
-#if DO_CHECK
-	CommandList->NumRecordScopeRefs--;
-#endif
+	if (CommandList)
+	{
+	#if DO_CHECK
+		CommandList->NumRecordScopeRefs--;
+	#endif
+
+		if (StopAction == EStopRecordingAction::Close)
+		{
+			CommandList->Close();
+		}
+		else if (StopAction == EStopRecordingAction::Submit)
+		{
+			CommandList->Close();
+			CommandList->Submit();
+		}
+	}
+}
 
-	if (StopAction == EStopRecordingAction::Close)
+FRenderCommandList::FFlushScope::FFlushScope()
+{
+	CommandList = SetInstanceTLS(nullptr);
+
+	if (CommandList)
 	{
-		CommandList->Close();
+		CommandList->Flush();
 	}
-	else if (StopAction == EStopRecordingAction::Submit)
+}
+
+FRenderCommandList::FFlushScope::~FFlushScope()
+{
+	if (CommandList)
 	{
-		CommandList->Close();
-		CommandList->Submit();
+		SetInstanceTLS(CommandList);
 	}
 }
 
@@ -1728,6 +1754,25 @@ FRenderCommandList::~FRenderCommandList()
 	}
 }
 
+void FRenderCommandList::Flush()
+{
+	if (bSubmitted || !bInitialized || !EnumHasAnyFlags(Flags, ERenderCommandListFlags::CloseOnSubmit))
+	{
+		return;
+	}
+
+	FRenderCommandList* FlushCommandList = FRenderCommandList::Create(ERenderCommandListFlags::CloseOnSubmit);
+	FlushCommandList->Allocator = MoveTemp(Allocator);
+	FlushCommandList->CommandLists = MoveTemp(CommandLists);
+	FlushCommandList->DispatchTaskEvent = MoveTemp(DispatchTaskEvent);
+	FlushCommandList->bInitialized = true;
+	FlushCommandList->Children = Children;
+	FlushCommandList->Submit();
+
+	bInitialized = false;
+	Children = {};
+}
+
 void FRenderCommandList::Init()
 {
 	check(!bInitialized);
diff --git a/Engine/Source/Runtime/RenderCore/Public/RenderingThread.h b/Engine/Source/Runtime/RenderCore/Public/RenderingThread.h
--- a/Engine/Source/Runtime/RenderCore/Public/RenderingThread.h
+++ b/Engine/Source/Runtime/RenderCore/Public/RenderingThread.h
@@ -858,6 +858,7 @@ public:
 		Submit
 	};
 
+	// A scope to bind a command list for recording on the current thread.
 	class FRecordScope
 	{
 	public:
@@ -872,6 +873,19 @@ public:
 		EStopRecordingAction StopAction;
 	};
 
+	// A scope to unbind and flush the contents of the currently recording command list if there is one.
+	class FFlushScope
+	{
+	public:
+		RENDERCORE_API FFlushScope();
+		RENDERCORE_API ~FFlushScope();
+
+		FFlushScope(const FFlushScope&) = delete;
+
+	private:
+		FRenderCommandList* CommandList;
+	};
+
 	// Call when the command list recording is finished.
 	RENDERCORE_API void Close();
 
@@ -992,6 +1006,8 @@ private:
 
 	RENDERCORE_API void Init();
 
+	RENDERCORE_API void Flush();
+
 	template <typename RenderCommandTag, typename FunctionType>
 	FORCEINLINE bool Enqueue(FunctionType&& Function)
 	{
