Fix for static analysis error. Explicitly nulled out command list pointer.

diff --git a/Engine/Source/Runtime/RenderCore/Private/RenderingThread.cpp b/Engine/Source/Runtime/RenderCore/Private/RenderingThread.cpp
--- a/Engine/Source/Runtime/RenderCore/Private/RenderingThread.cpp
+++ b/Engine/Source/Runtime/RenderCore/Private/RenderingThread.cpp
@@ -1768,6 +1768,7 @@ void FRenderCommandList::Flush()
 	FlushCommandList->bInitialized = true;
 	FlushCommandList->Children = Children;
 	FlushCommandList->Submit();
+	FlushCommandList = nullptr;
 
 	bInitialized = false;
 	Children = {};
