Reworked the submission API to make static analyzer happy.
- Use FRenderCommandDispatcher::Submit instead of CommandList->Submit because the static analyzer thinks that its a dangling pointer.
- Slightly clearer public API semantics since the static method takes ownership of the pointer.

diff --git a/Engine/Source/Runtime/Engine/Private/LevelTick.cpp b/Engine/Source/Runtime/Engine/Private/LevelTick.cpp
--- a/Engine/Source/Runtime/Engine/Private/LevelTick.cpp
+++ b/Engine/Source/Runtime/Engine/Private/LevelTick.cpp
@@ -1263,7 +1263,7 @@ void UWorld::SendAllEndOfFrameUpdatesInternal(EEndOfFrameUpdateLocation Location
 
 		}, EndOfFrameUpdatePrerequisiteTask, UE::Tasks::ETaskPriority::High));
 
-		TaskRenderCommandList->Submit(RootRenderCommandList);
+		FRenderCommandDispatcher::Submit(TaskRenderCommandList, RootRenderCommandList);
 
 		Context->Update_GameThread(Scene, RootRenderCommandList);
 	}
diff --git a/Engine/Source/Runtime/RenderCore/Private/RenderingThread.cpp b/Engine/Source/Runtime/RenderCore/Private/RenderingThread.cpp
--- a/Engine/Source/Runtime/RenderCore/Private/RenderingThread.cpp
+++ b/Engine/Source/Runtime/RenderCore/Private/RenderingThread.cpp
@@ -1653,7 +1653,7 @@ FRenderCommandList::FRecordScope::~FRecordScope()
 		else if (StopAction == EStopRecordingAction::Submit)
 		{
 			CommandList->Close();
-			CommandList->Submit();
+			FRenderCommandDispatcher::Submit(CommandList);
 		}
 	}
 }
@@ -1711,12 +1711,12 @@ void FRenderCommandList::FParallelForContext::Submit()
 
 		for (int32 Index = 0; Index < TaskCommandLists.Num() - 1; ++Index)
 		{
-			TaskCommandLists[Index]->Submit(RootCommandList);
+			FRenderCommandDispatcher::Submit(TaskCommandLists[Index], RootCommandList);
 		}
 
 		if (bSubmitRootCommandList)
 		{
-			RootCommandList->Submit();
+			FRenderCommandDispatcher::Submit(RootCommandList);
 		}
 
 		RootCommandList = nullptr;
@@ -1767,8 +1767,8 @@ void FRenderCommandList::Flush()
 	FlushCommandList->DispatchTaskEvent = MoveTemp(DispatchTaskEvent);
 	FlushCommandList->bInitialized = true;
 	FlushCommandList->Children = Children;
-	FlushCommandList->Submit();
-	FlushCommandList = nullptr;
+
+	FRenderCommandDispatcher::Submit(FlushCommandList);
 
 	bInitialized = false;
 	Children = {};
diff --git a/Engine/Source/Runtime/RenderCore/Public/RenderingThread.h b/Engine/Source/Runtime/RenderCore/Public/RenderingThread.h
--- a/Engine/Source/Runtime/RenderCore/Public/RenderingThread.h
+++ b/Engine/Source/Runtime/RenderCore/Public/RenderingThread.h
@@ -821,8 +821,8 @@ ENUM_CLASS_FLAGS(ERenderCommandListFlags);
 	Command lists actually contain several sub-command lists, one for each render command pipe. At submission time the sub-command lists are
 	submitted separately. It doesn't matter if commands are enqueued to a pipe, they all go into the same command list.
 
-	Command lists have two calls, Close and Submit. Call Close when recording is complete. Call Submit to patch the command list into a parent
-	command list or the global render command pipes.
+	Command lists have two operations, Close and Submit. Call Close when recording is complete. Call Submit to patch the command list into a parent
+	command list or the global render command pipes. Use FRenderCommandDispatcher::Submit to submit command lists.
 
 	Command lists support a fast-path with ERenderCommandListFlags::CloseOnSubmit. This fuses the Close / Submit operations but enables an optimization
 	to skip empty lists at the end, which is helpful when managing a large number of command lists (see FParallelForContext for a concrete use case).
@@ -889,13 +889,6 @@ public:
 	// Call when the command list recording is finished.
 	RENDERCORE_API void Close();
 
-	/**
-	 * Call to submit commands into a parent command list or render command pipes. If the parent command list is null the recording instance
-	 * is pulled from the currently bound render command list (set via FRecordScope). If both are null the commands are submitted to the
-	 * global render command pipes.
-	 */
-	RENDERCORE_API void Submit(FRenderCommandList* InParentCommandList = nullptr);
-
 	/*
 		A task context for use with parallel for that allocates a command list for each task thread.
 
@@ -970,6 +963,9 @@ public:
 private:
 	RENDERCORE_API FRenderCommandList(ERenderCommandListFlags InFlags, EPageSize InPageSize);
 
+	/** Use FRenderCommandDispatcher::Submit to submit command lists. */
+	RENDERCORE_API void Submit(FRenderCommandList* InParentCommandList = nullptr);
+
 	const UE::Tasks::FTask* TryGetDispatchTask() const
 	{
 		return DispatchTaskEvent.GetPtrOrNull();
@@ -1064,11 +1060,23 @@ private:
 	FRenderCommandList* NextSibling = nullptr;
 
 	FRenderCommandPipeBitArray PipeEnqueueFailedBits;
+
+	friend class FRenderCommandDispatcher;
 };
 
 class FRenderCommandDispatcher
 {
 public:
+	/**
+	 * Call to submit a command list into a parent command list or render command pipes. If the parent command list is null the recording instance
+	 * is pulled from the currently bound render command list (set via FRecordScope). If both are null the commands are submitted to the global render
+	 * command pipes.
+	 */
+	static void Submit(FRenderCommandList* RenderCommandList, FRenderCommandList* ParentCommandList = nullptr)
+	{
+		RenderCommandList->Submit(ParentCommandList);
+	}
+
 	template <typename RenderCommandTag>
 	static void Enqueue(TUniqueFunction<void(FRHICommandListImmediate&)>&& Function)
 	{
