Fix more stat system asserts
- Slate calls GameLoopIsStarved() when a modal dialog is showing in the editor. This was then calling FStats::AdvanceFrame, which has the RT delegate set, but does not match the call with an RHIEndFrame, causing stats to pile up in the history.
- There doesn't seem to be a reason to advance the stats frame mid-modal dialog like this, so removing the call fixes the asserts.
- Replaced with a FlushRenderingCommands(), since that was a side effect of advancing the stats frame.

Also deprecated FlushPendingDeleteRHIResources_GameThread and FlushPendingDeleteRHIResources_RenderThread.

diff --git a/Engine/Source/Runtime/Landscape/Private/LandscapeGrassWeightExporter.cpp b/Engine/Source/Runtime/Landscape/Private/LandscapeGrassWeightExporter.cpp
--- a/Engine/Source/Runtime/Landscape/Private/LandscapeGrassWeightExporter.cpp
+++ b/Engine/Source/Runtime/Landscape/Private/LandscapeGrassWeightExporter.cpp
@@ -560,7 +560,7 @@ FLandscapeGrassWeightExporter::FLandscapeGrassWeightExporter(ALandscapeProxy* In
 			[Exporter](FRHICommandListImmediate& RHICmdList)
 			{
 				Exporter->RenderLandscapeComponentToTexture_RenderThread(RHICmdList);
-				FlushPendingDeleteRHIResources_RenderThread();
+				RHICmdList.ImmediateFlush(EImmediateFlushType::FlushRHIThreadFlushResources);
 			});
 	}
 }
diff --git a/Engine/Source/Runtime/Launch/Private/LaunchEngineLoop.cpp b/Engine/Source/Runtime/Launch/Private/LaunchEngineLoop.cpp
--- a/Engine/Source/Runtime/Launch/Private/LaunchEngineLoop.cpp
+++ b/Engine/Source/Runtime/Launch/Private/LaunchEngineLoop.cpp
@@ -4488,8 +4488,7 @@ void FEngineLoop::InitTime()
 //called via FCoreDelegates::StarvedGameLoop
 void GameLoopIsStarved()
 {
-	FlushPendingDeleteRHIResources_GameThread();
-	UE::Stats::FStats::AdvanceFrame( true, UE::Stats::FStats::FOnAdvanceRenderingThreadStats::CreateStatic( &AdvanceRenderingThreadStatsGT ) );
+	FlushRenderingCommands();
 }
 
 
@@ -5539,10 +5538,10 @@ void FEngineLoop::Tick()
 
 		// handle some per-frame tasks on the rendering thread
 		ENQUEUE_RENDER_COMMAND(ResetDeferredUpdates)(
-			[](FRHICommandList& RHICmdList)
+			[](FRHICommandListImmediate& RHICmdList)
 			{
 				FDeferredUpdateResource::ResetNeedsUpdate();
-				FlushPendingDeleteRHIResources_RenderThread();
+				RHICmdList.ImmediateFlush(EImmediateFlushType::FlushRHIThreadFlushResources);
 			});
 
 		// Don't pump messages if we're running embedded as the outer application
diff --git a/Engine/Source/Runtime/RenderCore/Private/RenderingThread.cpp b/Engine/Source/Runtime/RenderCore/Private/RenderingThread.cpp
--- a/Engine/Source/Runtime/RenderCore/Private/RenderingThread.cpp
+++ b/Engine/Source/Runtime/RenderCore/Private/RenderingThread.cpp
@@ -1372,22 +1372,16 @@ void FlushRenderingCommands()
 
 void FlushPendingDeleteRHIResources_GameThread()
 {
-	if (!IsRunningRHIInSeparateThread())
+	ENQUEUE_RENDER_COMMAND(FlushPendingDeleteRHIResources)(
+		[](FRHICommandListImmediate& RHICmdList)
 	{
-		ENQUEUE_RENDER_COMMAND(FlushPendingDeleteRHIResources)(
-			[](FRHICommandList& RHICmdList)
-			{
-				FlushPendingDeleteRHIResources_RenderThread();
-			});
-	}
+		RHICmdList.ImmediateFlush(EImmediateFlushType::FlushRHIThreadFlushResources);
+	});
 }
 
 void FlushPendingDeleteRHIResources_RenderThread()
 {
-	if (!IsRunningRHIInSeparateThread())
-	{
-		FRHICommandListExecutor::GetImmediateCommandList().ImmediateFlush(EImmediateFlushType::FlushRHIThreadFlushResources);
-	}
+	FRHICommandListImmediate::Get().ImmediateFlush(EImmediateFlushType::FlushRHIThreadFlushResources);
 }
 
 FRHICommandListImmediate& GetImmediateCommandList_ForRenderCommand()
diff --git a/Engine/Source/Runtime/RenderCore/Public/RenderingThread.h b/Engine/Source/Runtime/RenderCore/Public/RenderingThread.h
--- a/Engine/Source/Runtime/RenderCore/Public/RenderingThread.h
+++ b/Engine/Source/Runtime/RenderCore/Public/RenderingThread.h
@@ -109,7 +109,9 @@ extern RENDERCORE_API void AdvanceRenderingThreadStatsGT( bool bDiscardCallstack
  */
 extern RENDERCORE_API void FlushRenderingCommands();
 
+UE_DEPRECATED(5.6, "FlushPendingDeleteRHIResources_GameThread is deprecated. Enqueue a render command that calls ImmediateFlush(EImmediateFlushType::FlushRHIThreadFlushResources) on the immediate RHI command list instead.")
 extern RENDERCORE_API void FlushPendingDeleteRHIResources_GameThread();
+UE_DEPRECATED(5.6, "FlushPendingDeleteRHIResources_RenderThread is deprecated. Call ImmediateFlush(EImmediateFlushType::FlushRHIThreadFlushResources) on the immediate RHI command list instead.")
 extern RENDERCORE_API void FlushPendingDeleteRHIResources_RenderThread();
 
 extern RENDERCORE_API void TickRenderingTickables();
