[shaders]
- fix late shadertype registration check; since shader types are now delay-loaded in a batch the previous implementation no longer works (the dynamic initializer which adds an initializer object to a shared array still happens, but nothing ever touches that array after the CommitAll call occurs)
- move and renamed the variable used to validate this condition; calling it "serialization history" is a relic of code that was deleted in 2020
- rename the equivalent variable used to validate the condition for vertex factories; the check doesn't have the same problem in this case because vertex factories don't have a similar delayed-registration mechanism

diff --git a/Engine/Source/Runtime/RenderCore/Private/Shader.cpp b/Engine/Source/Runtime/RenderCore/Private/Shader.cpp
--- a/Engine/Source/Runtime/RenderCore/Private/Shader.cpp
+++ b/Engine/Source/Runtime/RenderCore/Private/Shader.cpp
@@ -211,7 +211,6 @@ void FShaderParameterMap::UpdateHash(FSHA1& HashState) const
 	}
 }
 
-bool FShaderType::bInitializedSerializationHistory = false;
 
 static TArray<FShaderType*>& GetSortedShaderTypes(FShaderType::EShaderTypeForDynamicCast Type)
 {
@@ -288,10 +287,6 @@ FShaderType::FShaderType(
 {
 	FTypeLayoutDesc::Register(InTypeLayout);
 
-	// This will trigger if an IMPLEMENT_SHADER_TYPE was in a module not loaded before InitializeShaderTypes
-	// Shader types need to be implemented in modules that are loaded before that
-	checkf(!bInitializedSerializationHistory, TEXT("Shader type was loaded after engine init, use ELoadingPhase::PostConfigInit on your module to cause it to load earlier."));
-
 	//make sure the name is shorter than the maximum serializable length
 	check(FCString::Strlen(InName) < NAME_SIZE);
 
@@ -318,6 +313,7 @@ FShaderType::~FShaderType()
 	SortedTypes.RemoveAt(SortedIndex);
 }
 
+bool FShaderTypeRegistration::bShaderTypesInitialized = false;
 static TArray<const FShaderTypeRegistration*>* GShaderTypeRegistrationInstances = nullptr;
 TArray<const FShaderTypeRegistration*>& FShaderTypeRegistration::GetInstances()
 {
@@ -335,6 +331,12 @@ void FShaderTypeRegistration::CommitAll()
 		FShaderType& ShaderType = Instance->LazyShaderTypeAccessor(); // constructs and registers type
 	}
 	GetInstances().Empty();
+	bShaderTypesInitialized = true;
+}
+
+bool FShaderTypeRegistration::AreShaderTypesInitialized()
+{
+	return bShaderTypesInitialized;
 }
 
 TLinkedList<FShaderType*>*& FShaderType::GetTypeList()
@@ -626,8 +628,6 @@ void FShaderType::Initialize(const TMap<FString, TArray<const TCHAR*> >& ShaderF
 #endif
 	}
 #endif // WITH_EDITOR
-
-	bInitializedSerializationHistory = true;
 }
 
 int32 FShaderMapPointerTable::AddIndexedPointer(const FTypeLayoutDesc& TypeDesc, void* Ptr)
diff --git a/Engine/Source/Runtime/RenderCore/Private/VertexFactory.cpp b/Engine/Source/Runtime/RenderCore/Private/VertexFactory.cpp
--- a/Engine/Source/Runtime/RenderCore/Private/VertexFactory.cpp
+++ b/Engine/Source/Runtime/RenderCore/Private/VertexFactory.cpp
@@ -19,7 +19,7 @@
 IMPLEMENT_TYPE_LAYOUT(FVertexFactoryShaderParameters);
 
 uint32 FVertexFactoryType::NumVertexFactories = 0;
-bool FVertexFactoryType::bInitializedSerializationHistory = false;
+bool FVertexFactoryType::bVertexFactoriesInitialized = false;
 
 static TLinkedList<FVertexFactoryType*>* GVFTypeList = nullptr;
 
@@ -96,7 +96,7 @@ void FVertexFactoryType::Initialize(const TMap<FString, TArray<const TCHAR*> >&
 	}
 #endif // WITH_EDITOR
 
-	bInitializedSerializationHistory = true;
+	bVertexFactoriesInitialized = true;
 }
 
 FVertexFactoryType::FVertexFactoryType(
@@ -144,7 +144,7 @@ FVertexFactoryType::FVertexFactoryType(
 
 	// This will trigger if an IMPLEMENT_VERTEX_FACTORY_TYPE was in a module not loaded before InitializeShaderTypes
 	// Vertex factory types need to be implemented in modules that are loaded before that
-	checkf(!bInitializedSerializationHistory, TEXT("VF type was loaded after engine init, use ELoadingPhase::PostConfigInit on your module to cause it to load earlier."));
+	checkf(!bVertexFactoriesInitialized, TEXT("VF type was loaded after engine init, use ELoadingPhase::PostConfigInit on your module to cause it to load earlier."));
 
 	// Add this vertex factory type to the global list.
 	GlobalListLink.LinkHead(GetTypeList()); 
diff --git a/Engine/Source/Runtime/RenderCore/Public/Shader.h b/Engine/Source/Runtime/RenderCore/Public/Shader.h
--- a/Engine/Source/Runtime/RenderCore/Public/Shader.h
+++ b/Engine/Source/Runtime/RenderCore/Public/Shader.h
@@ -1548,9 +1548,6 @@ private:
 
 	friend void RENDERCORE_API DumpShaderStats( EShaderPlatform Platform, EShaderFrequency Frequency );
 
-	/** Tracks whether serialization history for all shader types has been initialized. */
-	static bool bInitializedSerializationHistory;
-
 #if WITH_EDITOR
 protected:
 	/**
@@ -1588,6 +1585,7 @@ public:
 	FShaderTypeRegistration(TFunctionRef<FShaderType& ()> LazyShaderTypeAccessor)
 		: LazyShaderTypeAccessor(LazyShaderTypeAccessor)
 	{
+		checkf(!AreShaderTypesInitialized(), TEXT("Shader type was loaded too late, use ELoadingPhase::PostConfigInit on your module to cause it to load earlier. This shader will not be compiled or function.  The module and shader will appear in the callstack of this assert."));
 		GetInstances().Add(this);
 	}
 
@@ -1595,8 +1593,10 @@ public:
 
 	// Actually register all the types and clear the array
 	static RENDERCORE_API void CommitAll();
+	static RENDERCORE_API bool AreShaderTypesInitialized();
 
 private:
+	static bool bShaderTypesInitialized;
 	TFunctionRef<FShaderType& ()> LazyShaderTypeAccessor;
 };
 
diff --git a/Engine/Source/Runtime/RenderCore/Public/VertexFactory.h b/Engine/Source/Runtime/RenderCore/Public/VertexFactory.h
--- a/Engine/Source/Runtime/RenderCore/Public/VertexFactory.h
+++ b/Engine/Source/Runtime/RenderCore/Public/VertexFactory.h
@@ -500,8 +500,8 @@ public:
 private:
 	static RENDERCORE_API uint32 NumVertexFactories;
 
-	/** Tracks whether serialization history for all shader types has been initialized. */
-	static bool bInitializedSerializationHistory;
+	/** Tracks whether engine vertex factory initialization has occurred */
+	static bool bVertexFactoriesInitialized;
 
 	const TCHAR* Name;
 	const TCHAR* ShaderFilename;
