[shaders] deprecating/removing some unused/unnecessary cruft related to gbuffer format and velocity depth settings

#8734

diff --git a/Engine/Config/BaseEngine.ini b/Engine/Config/BaseEngine.ini
--- a/Engine/Config/BaseEngine.ini
+++ b/Engine/Config/BaseEngine.ini
@@ -3350,7 +3350,6 @@ DefaultASTCQualityBySizeHQ=4
 
 +VersionedIntRValues=r.AllowStaticLighting
 +VersionedIntRValues=r.MaterialEditor.LWCTruncateMode
-+VersionedIntRValues=r.GBuffer
 +VersionedIntRValues=r.VelocityOutputPass
 +VersionedIntRValues=r.SelectiveBasePassOutputs
 +VersionedIntRValues=r.DBuffer
diff --git a/Engine/Source/Runtime/Engine/Private/ShaderCompiler/ShaderGenerationUtil.cpp b/Engine/Source/Runtime/Engine/Private/ShaderCompiler/ShaderGenerationUtil.cpp
--- a/Engine/Source/Runtime/Engine/Private/ShaderCompiler/ShaderGenerationUtil.cpp
+++ b/Engine/Source/Runtime/Engine/Private/ShaderCompiler/ShaderGenerationUtil.cpp
@@ -440,15 +440,6 @@ static FShaderGlobalDefines FetchShaderGlobalDefines(EShaderPlatform TargetPlatf
 
 	Ret.bSupportsDualBlending = RHISupportsDualSourceBlending(TargetPlatform);
 
-	{
-		Ret.bNeedVelocityDepth = NeedsVelocityDepth(TargetPlatform);
-	}
-
-	{
-		static IConsoleVariable* CVar = IConsoleManager::Get().FindConsoleVariable(TEXT("r.GBufferFormat"));
-		Ret.LegacyGBufferFormat = CVar->GetInt();
-	}
-
 	return Ret;
 }
 
diff --git a/Engine/Source/Runtime/RenderCore/Public/ShaderMaterial.h b/Engine/Source/Runtime/RenderCore/Public/ShaderMaterial.h
--- a/Engine/Source/Runtime/RenderCore/Public/ShaderMaterial.h
+++ b/Engine/Source/Runtime/RenderCore/Public/ShaderMaterial.h
@@ -50,8 +50,20 @@ struct FShaderGlobalDefines
 	// COMPILE_SHADERS_FOR_DEVELOPMENT = COMPILE_SHADERS_FOR_DEVELOPMENT_ALLOWED && bAllowDevelopmentShaderCompile;
 	bool COMPILE_SHADERS_FOR_DEVELOPMENT_ALLOWED;
 	bool bSupportsDualBlending;
+	UE_DEPRECATED(5.7, "LegacyGBufferFormat is no longer used in the engine")
 	int LegacyGBufferFormat;
+	UE_DEPRECATED(5.7, "bNeedVelocityDepth is no longer used in the engine")
 	bool bNeedVelocityDepth;
+
+	PRAGMA_DISABLE_DEPRECATION_WARNINGS
+	// Explicitly-defaulted ctors & assignment operators are needed temporarily due to deprecation of
+	// LegacyGBufferFormat/bNeedVelocityDepth fields. These can be removed along with the fields in 5.9
+	FShaderGlobalDefines() = default;
+	FShaderGlobalDefines(FShaderGlobalDefines&&) = default;
+	FShaderGlobalDefines(const FShaderGlobalDefines&) = default;
+	FShaderGlobalDefines& operator=(FShaderGlobalDefines&&) = default;
+	FShaderGlobalDefines& operator=(const FShaderGlobalDefines&) = default;
+	PRAGMA_ENABLE_DEPRECATION_WARNINGS
 };
 
 
