Refactored RDG to support disabling async setup queue and parallel compilation without disabling setup tasks.

diff --git a/Engine/Source/Runtime/RenderCore/Private/RenderGraphBuilder.cpp b/Engine/Source/Runtime/RenderCore/Private/RenderGraphBuilder.cpp
--- a/Engine/Source/Runtime/RenderCore/Private/RenderGraphBuilder.cpp
+++ b/Engine/Source/Runtime/RenderCore/Private/RenderGraphBuilder.cpp
@@ -610,8 +610,9 @@ FRDGBuilder::FRDGBuilder(FRHICommandListImmediate& InRHICmdList, FRDGEventName I
 #if RDG_ENABLE_PARALLEL_TASKS
 	ParallelSetup.TaskPriorityBias = GRDGParallelSetupTaskPriorityBias;
 #endif
+	AsyncSetupQueue.bEnabled = ParallelSetup.bEnabled && GRDGAsyncSetupQueue != 0;
 
-	bParallelCompileEnabled  = ::IsParallelSetupEnabled(ShaderPlatform) && EnumHasAnyFlags(InFlags, ERDGBuilderFlags::ParallelCompile);
+	bParallelCompileEnabled  = GRDGParallelCompile && ::IsParallelSetupEnabled(ShaderPlatform) && EnumHasAnyFlags(InFlags, ERDGBuilderFlags::ParallelCompile);
 
 	if (TransientResourceAllocator)
 	{
@@ -970,14 +971,14 @@ void FRDGBuilder::FlushAccessModeQueue()
 	ERHIPipeline ParameterPipelines = ERHIPipeline::None;
 
 	TArray<FRDGPass::FExternalAccessOp, FRDGArrayAllocator> Ops;
-	Ops.Reserve(ParallelSetup.bEnabled ? AccessModeQueue.Num() : 0);
+	Ops.Reserve(AsyncSetupQueue.bEnabled ? AccessModeQueue.Num() : 0);
 
 	for (FRDGViewableResource* Resource : AccessModeQueue)
 	{
 		const auto& AccessModeState = Resource->AccessModeState;
 		Resource->AccessModeState.bQueued = false;
 
-		if (ParallelSetup.bEnabled)
+		if (AsyncSetupQueue.bEnabled)
 		{
 			Ops.Emplace(Resource, AccessModeState.Mode);
 		}
@@ -1329,7 +1330,7 @@ void FRDGBuilder::Compile()
 	{
 		SCOPED_NAMED_EVENT(PassDependencies, FColor::Emerald);
 
-		if (!ParallelSetup.bEnabled)
+		if (!AsyncSetupQueue.bEnabled)
 		{
 			for (FRDGPassHandle PassHandle = ProloguePassHandle + 1; PassHandle < EpiloguePassHandle; ++PassHandle)
 			{
@@ -1337,7 +1338,7 @@ void FRDGBuilder::Compile()
 			}
 		}
 	}
-	else if (!ParallelSetup.bEnabled)
+	else if (!AsyncSetupQueue.bEnabled)
 	{
 		for (FRDGPassHandle PassHandle = ProloguePassHandle + 1; PassHandle < EpiloguePassHandle; ++PassHandle)
 		{
@@ -1717,7 +1718,7 @@ void FRDGBuilder::ProcessAsyncSetupQueue()
 
 void FRDGBuilder::FlushSetupQueue()
 {
-	if (ParallelSetup.bEnabled)
+	if (AsyncSetupQueue.bEnabled)
 	{
 		LaunchAsyncSetupQueueTask();
 	}
@@ -1776,6 +1777,8 @@ void FRDGBuilder::Execute()
 
 	FCollectResourceContext CollectResourceContext;
 
+	bCompiling = true;
+
 	if (!IsImmediateMode())
 	{
 		BeginFlushResourcesRHI();
@@ -1867,7 +1870,7 @@ void FRDGBuilder::Execute()
 		{
 			SubmitBufferUploads(RHICmdListTask, &AllocateUploadBuffersTask);
 
-		}, BufferNumElementsCallbacksTask, TaskPriority, bParallelCompileEnabled && !UploadedBuffers.IsEmpty());
+		}, BufferNumElementsCallbacksTask, TaskPriority, !UploadedBuffers.IsEmpty());
 
 		Compile();
 
@@ -2009,6 +2012,8 @@ void FRDGBuilder::Execute()
 	IF_RDG_ENABLE_DEBUG(GRDGAllowRHIAccess = ParallelExecute.IsEnabled());
 	IF_RDG_ENABLE_TRACE(Trace.OutputGraphBegin());
 
+	bCompiling = false;
+
 	ERHIPipeline OriginalPipeline = RHICmdList.GetPipeline();
 	if (!IsImmediateMode())
 	{
@@ -2422,7 +2427,7 @@ void FRDGBuilder::SetupPassResources(FRDGPass* Pass)
 			bRenderPassOnlyWrites &= EnumHasAnyFlags(AccessFlags, ERDGTextureAccessFlags::RenderTarget);
 
 			// When running in parallel this is set via MarkResourcesAsProduced. We also can't touch this as its a bitfield and not atomic.
-			if (!ParallelSetup.bEnabled)
+			if (!AsyncSetupQueue.bEnabled)
 			{
 				Texture->bProduced = true;
 			}
@@ -2472,7 +2477,7 @@ void FRDGBuilder::SetupPassResources(FRDGPass* Pass)
 			bRenderPassOnlyWrites = false;
 
 			// When running in parallel this is set via MarkResourcesAsProduced. We also can't touch this as its a bitfield and not atomic.
-			if (!ParallelSetup.bEnabled)
+			if (!AsyncSetupQueue.bEnabled)
 			{
 				Buffer->bProduced = true;
 			}
@@ -2489,7 +2494,7 @@ void FRDGBuilder::SetupPassResources(FRDGPass* Pass)
 		Pass->UniformBuffers.Emplace(UniformBuffer.GetUniformBuffer()->Handle);
 	});
 
-	if (ParallelSetup.bEnabled)
+	if (AsyncSetupQueue.bEnabled)
 	{
 		SetupPassDependencies(Pass);
 
@@ -2609,7 +2614,7 @@ FRDGPass* FRDGBuilder::SetupParameterPass(FRDGPass* Pass)
 
 	SetupPassInternals(Pass);
 
-	if (ParallelSetup.bEnabled)
+	if (AsyncSetupQueue.bEnabled)
 	{
 		MarkResourcesAsProduced(Pass);
 		AsyncSetupQueue.Push(FAsyncSetupOp::SetupPassResources(Pass));
diff --git a/Engine/Source/Runtime/RenderCore/Private/RenderGraphPrivate.cpp b/Engine/Source/Runtime/RenderCore/Private/RenderGraphPrivate.cpp
--- a/Engine/Source/Runtime/RenderCore/Private/RenderGraphPrivate.cpp
+++ b/Engine/Source/Runtime/RenderCore/Private/RenderGraphPrivate.cpp
@@ -297,6 +297,22 @@ FAutoConsoleVariableRef CVarRDGParallelDestruction(
 int32 GRDGParallelSetup = 1;
 FAutoConsoleVariableRef CVarRDGParallelSetup(
 	TEXT("r.RDG.ParallelSetup"), GRDGParallelSetup,
+	TEXT("RDG will launch setup tasks from AddSetupTask.")
+	TEXT(" 0: setup tasks are executed inline;")
+	TEXT(" 1: setup tasks are launched (default);"),
+	ECVF_RenderThreadSafe);
+
+int32 GRDGParallelCompile = 1;
+FAutoConsoleVariableRef CVarRDGParallelCompile(
+	TEXT("r.RDG.ParallelCompile"), GRDGParallelCompile,
+	TEXT("RDG will launch tasks to compile the RDG graph.")
+	TEXT(" 0: compile tasks are executed inline;")
+	TEXT(" 1: compile tasks are launched (default);"),
+	ECVF_RenderThreadSafe);
+
+int32 GRDGAsyncSetupQueue = 1;
+FAutoConsoleVariableRef CVarRDGAsyncSetupQueue(
+	TEXT("r.RDG.AsyncSetupQueue"), GRDGAsyncSetupQueue,
 	TEXT("RDG will setup passes in parallel when prompted by calls to FRDGBuilder::FlushSetupQueue.")
 	TEXT(" 0: pass setup is done synchronously in AddPass;")
 	TEXT(" 1: pass setup is done asynchronously (default);"),
diff --git a/Engine/Source/Runtime/RenderCore/Private/RenderGraphPrivate.h b/Engine/Source/Runtime/RenderCore/Private/RenderGraphPrivate.h
--- a/Engine/Source/Runtime/RenderCore/Private/RenderGraphPrivate.h
+++ b/Engine/Source/Runtime/RenderCore/Private/RenderGraphPrivate.h
@@ -95,20 +95,24 @@ extern int32 GRDGTransientIndirectArgBuffers;
 
 extern int32 GRDGParallelDestruction;
 extern int32 GRDGParallelSetup;
+extern int32 GRDGParallelCompile;
 extern int32 GRDGParallelSetupTaskPriorityBias;
 extern int32 GRDGParallelExecute;
 extern int32 GRDGParallelExecutePassMin;
 extern int32 GRDGParallelExecutePassMax;
 extern int32 GRDGParallelExecutePassTaskModeThreshold;
+extern int32 GRDGAsyncSetupQueue;
 
 #else
 
 const int32 GRDGParallelDestruction = 0;
 const int32 GRDGParallelSetup = 0;
+const int32 GRDGParallelCompile = 0;
 const int32 GRDGParallelExecute = 0;
 const int32 GRDGParallelExecutePassMin = 0;
 const int32 GRDGParallelExecutePassMax = 0;
 const int32 GRDGParallelExecutePassTaskModeThreshold = 0;
+const int32 GRDGAsyncSetupQueue = 0;
 
 #endif
 
diff --git a/Engine/Source/Runtime/RenderCore/Public/RenderGraphBuilder.h b/Engine/Source/Runtime/RenderCore/Public/RenderGraphBuilder.h
--- a/Engine/Source/Runtime/RenderCore/Public/RenderGraphBuilder.h
+++ b/Engine/Source/Runtime/RenderCore/Public/RenderGraphBuilder.h
@@ -946,6 +946,7 @@ private:
 		UE::FMutex Mutex;
 		TArray<FAsyncSetupOp, FRDGArrayAllocator> Ops;
 		UE::Tasks::FPipe Pipe{ TEXT("FRDGBuilder::AsyncSetupQueue") };
+		bool bEnabled = false;
 
 	} AsyncSetupQueue;
 
@@ -1010,6 +1011,7 @@ private:
 	/////////////////////////////////////////////////////////////////////////////
 	// Parallel Execution
 
+	bool bCompiling = false;
 	bool bParallelCompileEnabled = false;
 
 	struct FParallelExecute
diff --git a/Engine/Source/Runtime/RenderCore/Public/RenderGraphBuilder.inl b/Engine/Source/Runtime/RenderCore/Public/RenderGraphBuilder.inl
--- a/Engine/Source/Runtime/RenderCore/Public/RenderGraphBuilder.inl
+++ b/Engine/Source/Runtime/RenderCore/Public/RenderGraphBuilder.inl
@@ -427,7 +427,7 @@ inline void FRDGBuilder::QueueBufferUpload(FRDGBufferRef Buffer, FRDGBufferIniti
 inline void FRDGBuilder::QueueCommitReservedBuffer(FRDGBufferRef Buffer, uint64 CommitSizeInBytes)
 {
 	IF_RDG_ENABLE_DEBUG(UserValidation.ValidateCommitBuffer(Buffer, CommitSizeInBytes));
-	if (!ParallelSetup.bEnabled)
+	if (!AsyncSetupQueue.bEnabled)
 	{
 		Buffer->PendingCommitSize = CommitSizeInBytes;
 	}
@@ -605,7 +605,9 @@ UE::Tasks::FTask FRDGBuilder::AddSetupTask(
 
 	const UE::Tasks::EExtendedTaskPriority ExtendedTaskPriority = ParallelSetup.bEnabled ? UE::Tasks::EExtendedTaskPriority::None : UE::Tasks::EExtendedTaskPriority::Inline;
 
-	if (!bCondition || (!ParallelSetup.bEnabled && UE::RDG::IsCompleted(Prerequisites)))
+	const bool bParallelEnabled = bCompiling ? bParallelCompileEnabled : ParallelSetup.bEnabled;
+
+	if (!bCondition || (!bParallelEnabled && UE::RDG::IsCompleted(Prerequisites)))
 	{
 		OuterLambda();
 	}
@@ -676,8 +678,10 @@ UE::Tasks::FTask FRDGBuilder::AddCommandListSetupTask(
 		UE::RDG::Wait(Prerequisites);
 	}
 
+	const bool bParallelEnabled = bCompiling ? bParallelCompileEnabled : ParallelSetup.bEnabled;
+
 	// Need a separate command list with inline tasks when prerequisites are not complete yet.
-	const bool bAllocateCommandListForTask = bCondition && (ParallelSetup.bEnabled || !UE::RDG::IsCompleted(Prerequisites));
+	const bool bAllocateCommandListForTask = bCondition && (bParallelEnabled || !UE::RDG::IsCompleted(Prerequisites));
 
 	FRHICommandList* RHICmdListTask = &RHICmdList;
 
