Enable RDG parallel execution on Mobile renderer (driven with r.RDG.ParallelMobile, disabled by default)

diff --git a/Engine/Source/Runtime/RenderCore/Private/RenderGraphPrivate.cpp b/Engine/Source/Runtime/RenderCore/Private/RenderGraphPrivate.cpp
--- a/Engine/Source/Runtime/RenderCore/Private/RenderGraphPrivate.cpp
+++ b/Engine/Source/Runtime/RenderCore/Private/RenderGraphPrivate.cpp
@@ -405,6 +405,12 @@ FAutoConsoleVariableRef CVarRDGDebugParallelExecute(
 	}),
 	ECVF_RenderThreadSafe);
 
+int32 GRDGParallelMobile = 0;
+FAutoConsoleVariableRef CVarRDGParallelMobile(
+	TEXT("r.RDG.ParallelMobile"), GRDGParallelMobile,
+	TEXT("Enable RDG parallel execution on mobile renderer (disabled by default)."),
+	ECVF_RenderThreadSafe);
+
 #endif //!RDG_ENABLE_PARALLEL_TASKS
 
 // Fix for random GPU crashes on draw indirects on multiple IHVs. Force all indirect arg buffers as non transient (see UE-115982)
@@ -598,6 +604,15 @@ void ShutdownRenderGraph()
 	FRDGBuilder::WaitForAsyncDeleteTask();
 }
 
+FORCEINLINE bool IsParallelMobileEnabled()
+{
+#if RDG_ENABLE_PARALLEL_TASKS
+	return (GRDGParallelMobile > 0);
+#else
+	return false;
+#endif
+}
+
 bool IsParallelExecuteEnabled(EShaderPlatform ShaderPlatform)
 {
 	return GRDGParallelExecute > 0
@@ -605,7 +620,7 @@ bool IsParallelExecuteEnabled(EShaderPlatform ShaderPlatform)
 		&& !IsImmediateMode()
 		&& !GRDGDebugFlushGPU
 		&& !GRDGTransitionLog
-		&& !IsMobilePlatform(ShaderPlatform)
+		&& (!IsMobilePlatform(ShaderPlatform) || IsParallelMobileEnabled())
 		&& !IsOpenGLPlatform(ShaderPlatform)
 		&& !IsVulkanMobileSM5Platform(ShaderPlatform)
 		&& GRHISupportsMultithreadedShaderCreation
@@ -623,7 +638,7 @@ bool IsParallelSetupEnabled(EShaderPlatform ShaderPlatform)
 		&& !GRHICommandList.Bypass()
 		&& !IsImmediateMode()
 		&& !GRDGTransitionLog
-		&& !IsMobilePlatform(ShaderPlatform)
+		&& (!IsMobilePlatform(ShaderPlatform) || IsParallelMobileEnabled())
 		&& !IsOpenGLPlatform(ShaderPlatform)
 		&& !IsVulkanMobileSM5Platform(ShaderPlatform)
 		&& GRHISupportsMultithreadedShaderCreation
