Add SceneDepthAux to the KeyGen of MobilePlatforms to avoid desync in the shaders

diff --git a/Engine/Source/Runtime/RenderCore/Private/Shader.cpp b/Engine/Source/Runtime/RenderCore/Private/Shader.cpp
--- a/Engine/Source/Runtime/RenderCore/Private/Shader.cpp
+++ b/Engine/Source/Runtime/RenderCore/Private/Shader.cpp
@@ -2085,6 +2085,12 @@ void ShaderMapAppendKey(EShaderPlatform Platform, FShaderKeyGenerator& KeyGen)
 			KeyGen.AppendSeparator();
 			KeyGen.Append(TEXT("NoFBF"));
 		}
+
+		if (MobileRequiresSceneDepthAux(Platform))
+		{
+			KeyGen.AppendSeparator();
+			KeyGen.Append(TEXT("DAux"));
+		}
 	}
 	else
 	{
