Change the dummy initialization texturte to the new SetInitActionInitializer.

diff --git a/Engine/Source/Runtime/RenderCore/Private/GlobalRenderResources.cpp b/Engine/Source/Runtime/RenderCore/Private/GlobalRenderResources.cpp
--- a/Engine/Source/Runtime/RenderCore/Private/GlobalRenderResources.cpp
+++ b/Engine/Source/Runtime/RenderCore/Private/GlobalRenderResources.cpp
@@ -178,16 +178,19 @@ class FDummyTransitionTexture : public FTexture
 public:
 	virtual void InitRHI(FRHICommandListBase& RHICmdList) override
 	{
-		// Create the texture RHI.
-		FColorBulkData BulkData(0, 0, 0, 255);
-
 		const FRHITextureCreateDesc Desc =
 			FRHITextureCreateDesc::Create2D(TEXT("DummyTransitionTexture"), 1, 1, PF_R8G8B8A8)
 			.SetFlags(ETextureCreateFlags::ShaderResource)
-			.SetBulkData(&BulkData)
-			.SetClassName(TEXT("FDummyTransitionTexture"));
+			.SetClassName(TEXT("FDummyTransitionTexture"))
+			.SetInitActionInitializer();
+
+		FRHITextureInitializer Initializer = RHICmdList.CreateTextureInitializer(Desc);
 
-		TextureRHI = RHICreateTexture(Desc);
+		const FColor ColorValue(0, 0, 0, 255);
+		Initializer.GetTexture2DSubresource(0).WriteColor(ColorValue);
+
+		// Create the texture RHI.
+		TextureRHI = Initializer.Finalize();
 
 		// Create the sampler state RHI resource.
 		FSamplerStateInitializerRHI SamplerStateInitializer(SF_Point, AM_Wrap, AM_Wrap, AM_Wrap);
