Substrate - Fixed settings to have a common base on the cvar default. Then per platform tweaking. (BaseEngine was not working as intended)

#rnx
#tests Pc editor dx12, consoles.

diff --git a/Engine/Config/BaseEngine.ini b/Engine/Config/BaseEngine.ini
--- a/Engine/Config/BaseEngine.ini
+++ b/Engine/Config/BaseEngine.ini
@@ -3663,17 +3663,6 @@ bRequireDeviceSupportsARKit=true
 r.GPUCrashDebugging=false
 MaxSkinBones=(Default=65536,PerPlatform=(("Mobile", 256)))
 
-; Substrate settings for all platforms by default are setup for 60hz performance 
-r.Substrate.RoughDiffuse=0						; Lower lighting cost: disable rough diffuse
-r.Substrate.SheenQuality=2						; Lower lighting cost: simpler fuzz/cloth shading
-r.Substrate.ShadingQuality=2					; Lower lighting cost: simpler GGX, diffuse, non chromatic energy presevation
-r.Substrate.TileCoord8bits=1					; Memory and performance optimisation: 8 bit tile coord for lower resolution
-r.Substrate.Glints=0							; GBuffer and performance optimisation: disable GLints
-r.Substrate.ClosuresPerPixel=1					; GBuffer and performance optimisation: single closure max per pixel
-r.Substrate.OpaqueMaterialRoughRefraction=0		; Memory and performance optimisation: disable opaque rough refraction of multi layered materials
-r.Refraction.Blur=0								; Memory and performance optimisation: disable translucenct rough refraction
-r.Substrate.BlendableGBuffer=1					; GBuffer decal optimisation: use the legacy blendable gbuffer to avoid dbuffer decal memory overhead
-
 [Messaging]
 bAllowDelayedMessaging=false
 
diff --git a/Engine/Source/Runtime/RenderCore/Private/RenderUtils.cpp b/Engine/Source/Runtime/RenderCore/Private/RenderUtils.cpp
--- a/Engine/Source/Runtime/RenderCore/Private/RenderUtils.cpp
+++ b/Engine/Source/Runtime/RenderCore/Private/RenderUtils.cpp
@@ -1854,7 +1854,7 @@ static TAutoConsoleVariable<int32> CVarSubstrate(
 
 static TAutoConsoleVariable<int32> CVarSubstrateBlendableGBuffer(
 	TEXT("r.Substrate.BlendableGBuffer"),
-	0,
+	1, // Substrate 60fps by default
 	TEXT("Forces Substrate materials to fit into a blendable GBuffer. This is faster and DBuffer decals are no longer required. However, lots of visual effects and fidelity are lost. To be able to cast to blendable GBuffer, this command enforce ClosuresPerPixel=1."),
 	ECVF_ReadOnly | ECVF_RenderThreadSafe | ECVF_Preview);
 
@@ -1866,7 +1866,7 @@ static TAutoConsoleVariable<int32> CVarSubstrateBytesPerPixel(
 
 static TAutoConsoleVariable<int32> CVarSubstrateClosuresPerPixel(
 	TEXT("r.Substrate.ClosuresPerPixel"),
-	8,	// Similar to SUBSTRATE_MAX_CLOSURE_COUNT
+	1,	// Substrate 60fps by default. Must be less or equal to SUBSTRATE_MAX_CLOSURE_COUNT
 	TEXT("Substrate closure count per pixel can be constrained. That is usefull to reduce the number of byte written durting the base pass, but also to limit the number of closures in Forward. Higher value means more complex material can be represented."),
 	ECVF_ReadOnly | ECVF_RenderThreadSafe | ECVF_Preview);
 
@@ -1878,19 +1878,19 @@ static TAutoConsoleVariable<int32> CVarSubstrateBackCompatibility(
 
 static TAutoConsoleVariable<int32> CVarSubstrateOpaqueMaterialRoughRefraction(
 	TEXT("r.Substrate.OpaqueMaterialRoughRefraction"),
-	0,
+	0, // Substrate 60fps by default
 	TEXT("Enable Substrate opaque material rough refractions effect from top layers over layers below. (Experimental for now since this is not a path validated in production)"),
 	ECVF_ReadOnly | ECVF_RenderThreadSafe);
 
 static TAutoConsoleVariable<int32> CVarSubstrateSheenQuality(
 	TEXT("r.Substrate.SheenQuality"),
-	1,
+	2, // Substrate 60fps by default
 	TEXT("Define Substrate sheen quality (1: Disney-based sheen, 2: Charlie-based sheen, ). r.Substrate.ShadingQuality=2 forces SheenQuality to 2. This variable is read-only."),
 	ECVF_ReadOnly | ECVF_RenderThreadSafe | ECVF_Preview);
 
 static TAutoConsoleVariable<int32> CVarSubstrateShadingQuality(
 	TEXT("r.Substrate.ShadingQuality"),
-	1,
+	2, // Substrate 60fps by default
 	TEXT("Define Substrate shading quality (1: accurate lighting, 2: approximate lighting). This variable is read-only."),
 	ECVF_ReadOnly | ECVF_RenderThreadSafe | ECVF_Preview);
 
@@ -1915,13 +1915,13 @@ static TAutoConsoleVariable<int32> CVarMaterialRoughDiffuse(
 
 static TAutoConsoleVariable<int32> CVarSubstrateRoughDiffuse(
 	TEXT("r.Substrate.RoughDiffuse"),
-	1,
+	0, // Substrate 60fps by default
 	TEXT("Enable Substrate rough diffuse model (works only if r.Material.RoughDiffuse is enabled in the project settings). Togglable at runtime"),
 	ECVF_RenderThreadSafe);
 
 static TAutoConsoleVariable<int32> CVarSubstrateGlints(
 	TEXT("r.Substrate.Glints"),
-	1,
+	0, // Substrate 60fps by default
 	TEXT("Enable Glint support for Substrate slabs. If changed, shaders needs to be recompiled."),
 	ECVF_ReadOnly | ECVF_RenderThreadSafe | ECVF_Preview);
 
diff --git a/Engine/Source/Runtime/Renderer/Private/DistortionRendering.cpp b/Engine/Source/Runtime/Renderer/Private/DistortionRendering.cpp
--- a/Engine/Source/Runtime/Renderer/Private/DistortionRendering.cpp
+++ b/Engine/Source/Runtime/Renderer/Private/DistortionRendering.cpp
@@ -31,7 +31,7 @@ static TAutoConsoleVariable<int32> CVarDisableDistortion(
 
 static TAutoConsoleVariable<int32> CVarRefractionBlur(
 	TEXT("r.Refraction.Blur"),
-	1,
+	0, // Substrate 60fps by default
 	TEXT("Enables rough refractions, i.e. blurring of the background. Only enabled when Substrate is enabled."),
 	ECVF_Scalability | ECVF_RenderThreadSafe);
 
diff --git a/Engine/Source/Runtime/Renderer/Private/Substrate/Substrate.cpp b/Engine/Source/Runtime/Renderer/Private/Substrate/Substrate.cpp
--- a/Engine/Source/Runtime/Renderer/Private/Substrate/Substrate.cpp
+++ b/Engine/Source/Runtime/Renderer/Private/Substrate/Substrate.cpp
@@ -65,7 +65,7 @@ static TAutoConsoleVariable<int32> CVarSubstrateAllocationMode(
 
 static TAutoConsoleVariable<int32> CVarSubstrateTileCoord8Bits(
 	TEXT("r.Substrate.TileCoord8bits"),
-	0,
+	1, // Substrate 60fps by default
 	TEXT("Format of tile coord. This variable is read-only."),
 	ECVF_RenderThreadSafe);
 
