Added mobile shader platform CVar r.Mobile.SupportsSM5MaterialNodes (disabled by default). This variable is intended to be queried by the Material Translator to determine whether SM5 nodes can be utilized by the mobile renderer on a specific platform.

#rnx

diff --git a/Engine/Source/Runtime/Core/Private/HAL/ConsoleManager.cpp b/Engine/Source/Runtime/Core/Private/HAL/ConsoleManager.cpp
--- a/Engine/Source/Runtime/Core/Private/HAL/ConsoleManager.cpp
+++ b/Engine/Source/Runtime/Core/Private/HAL/ConsoleManager.cpp
@@ -4131,6 +4131,12 @@ static TAutoConsoleVariable<int32> CVarMobileDesiredResY(
 	TEXT("Desired mobile Y resolution (shortest axis) (non-zero == use for Y, calculate X to retain aspect ratio)"),
 	ECVF_Default);
 
+static TAutoConsoleVariable<int32> CVarMobileSupportsSM5MaterialNodes(
+	TEXT("r.Mobile.SupportsSM5MaterialNodes"),
+	0,
+	TEXT("Whether SM5 material nodes can be utilized or not by the mobile renderer on a specific platform (disabled by default)."),
+	ECVF_ReadOnly);
+
 static TAutoConsoleVariable<int32> CVarLWCTruncateMode(
 	TEXT("r.MaterialEditor.LWCTruncateMode"),
 	2,
diff --git a/Engine/Source/Runtime/RenderCore/Private/RenderUtils.cpp b/Engine/Source/Runtime/RenderCore/Private/RenderUtils.cpp
--- a/Engine/Source/Runtime/RenderCore/Private/RenderUtils.cpp
+++ b/Engine/Source/Runtime/RenderCore/Private/RenderUtils.cpp
@@ -633,6 +633,12 @@ RENDERCORE_API bool MobileUsesFullDepthPrepass(const FStaticShaderPlatform Platf
 	return MobileUsesShadowMaskTexture(Platform) || IsMobileAmbientOcclusionEnabled(Platform) || IsUsingDBuffers(Platform) || FReadOnlyCVARCache::MobileEarlyZPass(Platform) == 1;
 }
 
+RENDERCORE_API bool MobileSupportsSM5MaterialNodes(const FStaticShaderPlatform Platform)
+{
+	static FShaderPlatformCachedIniValue<bool> MobileSupportsSM5MaterialNodesIniValue(TEXT("r.Mobile.SupportsSM5MaterialNodes"));
+	return IsMobilePlatform(Platform) && (MobileSupportsSM5MaterialNodesIniValue.Get(Platform) != 0);
+}
+
 RENDERCORE_API bool ShouldForceFullDepthPass(const FStaticShaderPlatform Platform)
 {
 	if (IsMobilePlatform(Platform))
diff --git a/Engine/Source/Runtime/RenderCore/Public/RenderUtils.h b/Engine/Source/Runtime/RenderCore/Public/RenderUtils.h
--- a/Engine/Source/Runtime/RenderCore/Public/RenderUtils.h
+++ b/Engine/Source/Runtime/RenderCore/Public/RenderUtils.h
@@ -327,6 +327,8 @@ RENDERCORE_API bool MobileBasePassAlwaysUsesCSM(const FStaticShaderPlatform Plat
 
 RENDERCORE_API bool MobileUsesFullDepthPrepass(const FStaticShaderPlatform Platform);
 
+RENDERCORE_API bool MobileSupportsSM5MaterialNodes(const FStaticShaderPlatform Platform);
+
 RENDERCORE_API bool ShouldForceFullDepthPass(const FStaticShaderPlatform Platform);
 
 RENDERCORE_API bool SupportsGen4TAA(const FStaticShaderPlatform Platform);
