Fix missing platform-specific value for r.Substrate.OpaqueMaterialRoughRefraction

diff --git a/Engine/Source/Runtime/RenderCore/Private/RenderUtils.cpp b/Engine/Source/Runtime/RenderCore/Private/RenderUtils.cpp
--- a/Engine/Source/Runtime/RenderCore/Private/RenderUtils.cpp
+++ b/Engine/Source/Runtime/RenderCore/Private/RenderUtils.cpp
@@ -2192,7 +2192,8 @@ namespace Substrate
 
 	bool IsOpaqueRoughRefractionEnabled(EShaderPlatform InPlatform)
 	{
-		return IsSubstrateEnabled() && CVarSubstrateOpaqueMaterialRoughRefraction.GetValueOnAnyThread() > 0 && !IsSubstrateBlendableGBufferEnabled(InPlatform);
+		static FShaderPlatformCachedIniValue<int32> CVar(TEXT("r.Substrate.OpaqueMaterialRoughRefraction"));
+		return IsSubstrateEnabled() && CVar.Get(InPlatform) > 0 && !IsSubstrateBlendableGBufferEnabled(InPlatform);
 	}
 
 	bool IsAdvancedVisualizationEnabled()
