Adding LoadCachedShaderSourceFile which performs the same actions as LoadShaderSourceFile without copying the full shader source.

diff --git a/Engine/Source/Runtime/RenderCore/Private/ShaderCore.cpp b/Engine/Source/Runtime/RenderCore/Private/ShaderCore.cpp
--- a/Engine/Source/Runtime/RenderCore/Private/ShaderCore.cpp
+++ b/Engine/Source/Runtime/RenderCore/Private/ShaderCore.cpp
@@ -327,13 +327,13 @@ void FShaderCompilerDefinitions::InitializeInitialDefines(const FShaderCompilerD
 
 struct FShaderFileCacheEntry
 {
-	FString Source;
+	FShaderSharedStringPtr Source;
 	FShaderSharedAnsiStringPtr StrippedSource;				// Source with comments stripped out, and converted to ANSICHAR
 	FShaderPreprocessDependenciesShared Dependencies;		// Stripped source with include dependencies, all in one shareable struct
 	
 	bool IsEmpty() const
 	{
-		return Source.IsEmpty() && !StrippedSource.IsValid() && !Dependencies.IsValid();
+		return !Source.IsValid() && !StrippedSource.IsValid() && !Dependencies.IsValid();
 	}
 };
 
@@ -2663,7 +2663,7 @@ void ShaderConvertAndStripComments(const FString& ShaderSource, TArray<ANSICHAR>
 	OutStripped.SetNum(CurrentOut - OutStripped.GetData(), EAllowShrinking::No);
 }
 
-bool LoadShaderSourceFile(const TCHAR* InVirtualFilePath, EShaderPlatform ShaderPlatform, FString* OutFileContents, TArray<FShaderCompilerError>* OutCompileErrors, const FName* ShaderPlatformName, FShaderSharedAnsiStringPtr* OutStrippedContents) // TODO: const FString&
+bool LoadCachedShaderSourceFile(const TCHAR* InVirtualFilePath, EShaderPlatform ShaderPlatform, FShaderSharedStringPtr* OutFileContents, TArray<FShaderCompilerError>* OutCompileErrors, const FName* ShaderPlatformName, FShaderSharedAnsiStringPtr* OutStrippedContents)
 {
 #if WITH_EDITORONLY_DATA
 	// it's not expected that cooked platforms get here, but if they do, this is the final out
@@ -2685,26 +2685,33 @@ bool LoadShaderSourceFile(const TCHAR* InVirtualFilePath, EShaderPlatform Shader
 
 		bResult = 
 			GShaderFileCache.FindOrTryProduceAndApply(VirtualFilePath,
-				[&VirtualFilePath, &OutCompileErrors](FShaderFileCacheEntry& CachedFile)
+				[&VirtualFilePath, OutCompileErrors](FShaderFileCacheEntry& CachedFile)
 				{
 					FString ShaderFilePath = GetShaderSourceFilePath(VirtualFilePath, OutCompileErrors);
 
-					// verify SHA hash of shader files on load. missing entries trigger an error
-					if (!ShaderFilePath.IsEmpty() && FFileHelper::LoadFileToString(CachedFile.Source, *ShaderFilePath, FFileHelper::EHashOptions::EnableVerify | FFileHelper::EHashOptions::ErrorMissingHash))
+					if (!ShaderFilePath.IsEmpty())
 					{
-						TArray<ANSICHAR>* StrippedSource = new TArray<ANSICHAR>;
-						ShaderConvertAndStripComments(CachedFile.Source, *StrippedSource);
-						CachedFile.StrippedSource = MakeShareable(StrippedSource);
+						FShaderSharedStringPtr SourceFile = MakeShared<FString>();
+
+						// verify SHA hash of shader files on load. missing entries trigger an error
+						if (FFileHelper::LoadFileToString(*SourceFile, *ShaderFilePath, FFileHelper::EHashOptions::EnableVerify | FFileHelper::EHashOptions::ErrorMissingHash))
+						{
+							CachedFile.Source = MoveTemp(SourceFile);
 
-						return true;
+							TArray<ANSICHAR>* StrippedSource = new TArray<ANSICHAR>;
+							ShaderConvertAndStripComments(*CachedFile.Source, *StrippedSource);
+							CachedFile.StrippedSource = MakeShareable(StrippedSource);
+
+							return true;
+						}
 					}
 
 					// Create an empty entry if missing files are being cached.
 					return FMissingShaderFileCacheGuard::IsEnabled();
 				},
-				[&bFoundValidFile, &OutFileContents, &OutStrippedContents](const FShaderFileCacheEntry& CachedFile)
+				[&bFoundValidFile, OutFileContents, OutStrippedContents](const FShaderFileCacheEntry& CachedFile)
 				{
-					if (OutFileContents)
+					if (OutFileContents && CachedFile.Source.IsValid())
 					{
 						*OutFileContents = CachedFile.Source;
 					}
@@ -2734,6 +2741,25 @@ bool LoadShaderSourceFile(const TCHAR* InVirtualFilePath, EShaderPlatform Shader
 #endif // WITH_EDITORONLY_DATA
 }
 
+bool LoadShaderSourceFile(const TCHAR* InVirtualFilePath, EShaderPlatform ShaderPlatform, FString* OutFileContents, TArray<FShaderCompilerError>* OutCompileErrors, const FName* ShaderPlatformName, FShaderSharedAnsiStringPtr* OutStrippedContents) // TODO: const FString&
+{
+#if WITH_EDITORONLY_DATA
+	FShaderSharedStringPtr LocalFileContents;
+	FShaderSharedStringPtr* pFileContents = OutFileContents ? &LocalFileContents : nullptr;
+
+	const bool bResult = LoadCachedShaderSourceFile(InVirtualFilePath, ShaderPlatform, pFileContents, OutCompileErrors, ShaderPlatformName, OutStrippedContents);
+
+	if (OutFileContents)
+	{
+		*OutFileContents = *LocalFileContents;
+	}
+
+	return bResult;
+#else
+	return false;
+#endif
+}
+
 static FString FormatErrorCantFindSourceFile(const TCHAR* VirtualFilePath)
 {
 	return FString::Printf(TEXT("Couldn't find source file of virtual shader path \'%s\'"), VirtualFilePath);
@@ -2888,7 +2914,7 @@ static void AddPreprocessDependency(FShaderPreprocessDependencies& Dependencies,
 /**
  * Recursively populates IncludeFilenames with the unique include filenames found in the shader file named Filename.
  */
-static void InternalGetShaderIncludes(const TCHAR* EntryPointVirtualFilePath, const TCHAR* VirtualFilePath, const FString& FileContents, TArray<FString>& IncludeVirtualFilePaths, EShaderPlatform ShaderPlatform, uint32 DepthLimit, bool AddToIncludeFile, const FName* ShaderPlatformName, FShaderPreprocessDependencies* OutDependencies)
+static void InternalGetShaderIncludes(const TCHAR* EntryPointVirtualFilePath, const TCHAR* VirtualFilePath, FStringView FileContents, TArray<FString>& IncludeVirtualFilePaths, EShaderPlatform ShaderPlatform, uint32 DepthLimit, bool AddToIncludeFile, const FName* ShaderPlatformName, FShaderPreprocessDependencies* OutDependencies)
 {
 	//avoid an infinite loop with a 0 length string
 	if (FileContents.Len() > 0)
@@ -2899,7 +2925,7 @@ static void InternalGetShaderIncludes(const TCHAR* EntryPointVirtualFilePath, co
 		}
 
 		//find the first include directive
-		const TCHAR* IncludeBegin = FindFirstInclude(*FileContents);
+		const TCHAR* IncludeBegin = FindFirstInclude(GetData(FileContents));
 
 		uint32 SearchCount = 0;
 		const uint32 MaxSearchCount = 200;
@@ -2972,11 +2998,14 @@ static void InternalGetShaderIncludes(const TCHAR* EntryPointVirtualFilePath, co
 							}
 
 							// First element in Dependencies is root file, so initialize the StrippedSource pointer in it
-							FString IncludedFileContents;
-							LoadShaderSourceFile(*ExtractedIncludeFilename, ShaderPlatform, &IncludedFileContents, nullptr, ShaderPlatformName,
+							FShaderSharedStringPtr IncludedFileContents;
+							LoadCachedShaderSourceFile(*ExtractedIncludeFilename, ShaderPlatform, &IncludedFileContents, nullptr, ShaderPlatformName,
 								ExtractedIncludeDependencies ? &ExtractedIncludeDependencies->Dependencies[0].StrippedSource : nullptr);
 
-							InternalGetShaderIncludes(EntryPointVirtualFilePath, *ExtractedIncludeFilename, IncludedFileContents, IncludeVirtualFilePaths, ShaderPlatform, DepthLimit - 1, true, ShaderPlatformName, ExtractedIncludeDependencies.Get());
+							if (IncludedFileContents.IsValid())
+							{
+								InternalGetShaderIncludes(EntryPointVirtualFilePath, *ExtractedIncludeFilename, *IncludedFileContents, IncludeVirtualFilePaths, ShaderPlatform, DepthLimit - 1, true, ShaderPlatformName, ExtractedIncludeDependencies.Get());
+							}
 
 							if (ExtractedIncludeDependencies)
 							{
@@ -3098,10 +3127,13 @@ static void InternalGetShaderIncludes(const TCHAR* EntryPointVirtualFilePath, co
 	}
 
 	// First element in Dependencies is root file, so initialize the StrippedSource pointer in it
-	FString FileContents;
-	LoadShaderSourceFile(VirtualFilePath, ShaderPlatform, &FileContents, nullptr, ShaderPlatformName, PreprocessDependencies ? &PreprocessDependencies->Dependencies[0].StrippedSource : nullptr);
+	FShaderSharedStringPtr FileContents;
+	LoadCachedShaderSourceFile(VirtualFilePath, ShaderPlatform, &FileContents, nullptr, ShaderPlatformName, PreprocessDependencies ? &PreprocessDependencies->Dependencies[0].StrippedSource : nullptr);
 
-	InternalGetShaderIncludes(EntryPointVirtualFilePath, VirtualFilePath, FileContents, IncludeVirtualFilePaths, ShaderPlatform, DepthLimit, AddToIncludeFile, ShaderPlatformName, PreprocessDependencies.Get());
+	if (FileContents.IsValid())
+	{
+		InternalGetShaderIncludes(EntryPointVirtualFilePath, VirtualFilePath, *FileContents, IncludeVirtualFilePaths, ShaderPlatform, DepthLimit, AddToIncludeFile, ShaderPlatformName, PreprocessDependencies.Get());
+	}
 
 	if (PreprocessDependencies)
 	{
@@ -3140,8 +3172,8 @@ static bool TryUpdateSingleShaderFilehash(FSHA1& InOutHashState, const TCHAR* Vi
 	for (int32 IncludeIndex = 0; IncludeIndex < IncludeVirtualFilePaths.Num(); IncludeIndex++)
 	{
 		// Load the include file and hash it
-		FString IncludeFileContents;
-		if (!LoadShaderSourceFile(*IncludeVirtualFilePaths[IncludeIndex], ShaderPlatform, &IncludeFileContents, nullptr))
+		FShaderSharedStringPtr IncludeFileContents;
+		if (!LoadCachedShaderSourceFile(*IncludeVirtualFilePaths[IncludeIndex], ShaderPlatform, &IncludeFileContents, nullptr))
 		{
 			if (OutErrorMessage)
 			{
@@ -3149,7 +3181,7 @@ static bool TryUpdateSingleShaderFilehash(FSHA1& InOutHashState, const TCHAR* Vi
 			}
 			return false;
 		}
-		InOutHashState.UpdateWithString(*IncludeFileContents, IncludeFileContents.Len());
+		InOutHashState.UpdateWithString(GetData(*IncludeFileContents), GetNum(*IncludeFileContents));
 #if WITH_EDITOR &&  !(UE_BUILD_SHIPPING || UE_BUILD_TEST)
 		if (UE_LOG_ACTIVE(LogShaders, Verbose))
 		{
@@ -3162,8 +3194,8 @@ static bool TryUpdateSingleShaderFilehash(FSHA1& InOutHashState, const TCHAR* Vi
 	}
 
 	// Load the source file and hash it
-	FString FileContents;
-	if (!LoadShaderSourceFile(VirtualFilePath, ShaderPlatform, &FileContents, nullptr))
+	FShaderSharedStringPtr FileContents;
+	if (!LoadCachedShaderSourceFile(VirtualFilePath, ShaderPlatform, &FileContents, nullptr))
 	{
 		if (OutErrorMessage)
 		{
@@ -3171,7 +3203,7 @@ static bool TryUpdateSingleShaderFilehash(FSHA1& InOutHashState, const TCHAR* Vi
 		}
 		return false;
 	}
-	InOutHashState.UpdateWithString(*FileContents, FileContents.Len());
+	InOutHashState.UpdateWithString(GetData(*FileContents), GetNum(*FileContents));
 	if (OutErrorMessage)
 	{
 		OutErrorMessage->Reset();
diff --git a/Engine/Source/Runtime/RenderCore/Public/ShaderCompilerCore.h b/Engine/Source/Runtime/RenderCore/Public/ShaderCompilerCore.h
--- a/Engine/Source/Runtime/RenderCore/Public/ShaderCompilerCore.h
+++ b/Engine/Source/Runtime/RenderCore/Public/ShaderCompilerCore.h
@@ -26,6 +26,8 @@ class FShaderCommonCompileJob;
 class FShaderCompileJob;
 struct FShaderCompilerInput;
 class FShaderPipelineCompileJob;
+
+using FShaderSharedStringPtr = TSharedPtr<FString, ESPMode::ThreadSafe>;
 typedef TSharedPtr<TArray<ANSICHAR>, ESPMode::ThreadSafe> FShaderSharedAnsiStringPtr;
 
 // this is for the protocol, not the data, bump if FShaderCompilerInput/FShaderPreprocessOutput serialization, SerializeWorkerInput or ProcessInputFromArchive changes.
@@ -687,6 +689,8 @@ extern RENDERCORE_API void ShaderConvertAndStripComments(const FString& ShaderSo
  */
 extern RENDERCORE_API bool LoadShaderSourceFile(const TCHAR* VirtualFilePath, EShaderPlatform ShaderPlatform, FString* OutFileContents, TArray<FShaderCompilerError>* OutCompileErrors, const FName* ShaderPlatformName = nullptr, FShaderSharedAnsiStringPtr* OutStrippedContents = nullptr);
 
+extern RENDERCORE_API bool LoadCachedShaderSourceFile(const TCHAR* VirtualFilePath, EShaderPlatform ShaderPlatform, FShaderSharedStringPtr* OutFileContents, TArray<FShaderCompilerError>* OutCompileErrors, const FName* ShaderPlatformName = nullptr, FShaderSharedAnsiStringPtr* OutStrippedContents = nullptr);
+
 
 struct FShaderPreprocessDependency
 {
