Fixed issues with GPUScene assumtions when running commandlet with es31 feature level

diff --git a/Engine/Source/Runtime/RenderCore/Private/VertexFactory.cpp b/Engine/Source/Runtime/RenderCore/Private/VertexFactory.cpp
--- a/Engine/Source/Runtime/RenderCore/Private/VertexFactory.cpp
+++ b/Engine/Source/Runtime/RenderCore/Private/VertexFactory.cpp
@@ -354,7 +354,7 @@ bool FVertexFactory::AddPrimitiveIdStreamElement(EVertexInputStreamType InputStr
 		// mobile PrimitiveId stream should either disabled or same as desktop
 		check(AttributeIndex_Mobile == 0xff || AttributeIndex == AttributeIndex_Mobile);
 
-		if (GIsEditor || GMaxRHIFeatureLevel > ERHIFeatureLevel::ES3_1 || AttributeIndex_Mobile != 0xff)
+		if (GMaxRHIFeatureLevel > ERHIFeatureLevel::ES3_1 || AttributeIndex_Mobile != 0xff)
 		{
 			// UniformView path does not use PrimitiveId stream, we still need to set it to a non-negative index
 			int32 AddedStreamIndex = 0;
