Improve IO performance of FinishPopulateShaderLibrary
Can be quite slow reading in large ushaderbytecode files too granularly

#tests local lyra, Horde inc cooks

diff --git a/Engine/Source/Runtime/RenderCore/Private/ShaderCodeLibrary.cpp b/Engine/Source/Runtime/RenderCore/Private/ShaderCodeLibrary.cpp
--- a/Engine/Source/Runtime/RenderCore/Private/ShaderCodeLibrary.cpp
+++ b/Engine/Source/Runtime/RenderCore/Private/ShaderCodeLibrary.cpp
@@ -1545,6 +1545,13 @@ struct FShaderCodeStats
 	TMap<uint64, FPerTypeStats> ShaderTypeStats;
 };
 
+struct FShaderCodeReadInfo
+{
+	int64 ShaderCodeIndex;
+	int64 Offset;
+	int32 Size;
+};
+
 struct FEditorShaderCodeArchive
 {
 	FEditorShaderCodeArchive(FName InFormat, bool bInNeedsDeterministicOrder)
@@ -1977,11 +1984,10 @@ struct FEditorShaderCodeArchive
 		return ShaderMapIndex;
 	}
 
-	int32 AddShaderCode(FShaderCodeStats& CodeStats,
+	int32 PrepareAddShaderCode(FShaderCodeStats& CodeStats,
 		int32 OtherShaderMapIndex,
 		const FSerializedShaderArchive& OtherShaders,
-		int64 OtherShaderCodeOffset,
-		FArchive& Ar)
+		TArray<FShaderCodeReadInfo>& ReadInfo)
 	{
 		int32 ShaderMapIndex = 0;
 		const FSHAHash& OtherShaderMapHash = OtherShaders.ShaderMapHashes[OtherShaderMapIndex];
@@ -2007,14 +2013,13 @@ struct FEditorShaderCodeArchive
 
 					check(ShaderCode.Num() == SerializedShaders.GetNumShaders() - 1);
 
-					// Read shader code from archive and add shader to set
-					const int64 ReadSize = OtherShaderEntry.Size;
-					check(ReadSize > 0);
-					const int64 ReadOffset = OtherShaderCodeOffset + OtherShaderEntry.Offset;
-					FUniqueBuffer Code = FUniqueBuffer::Alloc(ReadSize);
-					Ar.Seek(ReadOffset);
-					Ar.Serialize(Code.GetData(), ReadSize);
-					ShaderCode.Add(Code.MoveToShared());
+					// Save information required to read in the old shader code
+					check(OtherShaderEntry.Size > 0);
+					FShaderCodeReadInfo Info{ ShaderCode.Num(), OtherShaderEntry.Offset, OtherShaderEntry.Size };
+					ReadInfo.Add(Info);
+
+					// We will add this old shader code in a separate pass
+					ShaderCode.AddZeroed(1);
 				}
 
 				TArray<uint64>& SerializedShaderTypes = SerializedShaders.ShaderTypes[ShaderIndex].Data;
@@ -2038,6 +2043,68 @@ struct FEditorShaderCodeArchive
 		return ShaderMapIndex;
 	}
 
+	void ReadShaderCodeBlock(FArchive& Ar, int64 StartOffset, const TArrayView<FShaderCodeReadInfo>& BlockReadInfo)
+	{
+		if (BlockReadInfo.Num() > 0)
+		{
+			const FShaderCodeReadInfo& First = BlockReadInfo[0];
+			const FShaderCodeReadInfo& Last = BlockReadInfo[BlockReadInfo.Num() - 1];
+			int64 ReadSize = Last.Offset - First.Offset + Last.Size;
+
+			FUniqueBuffer BlockOfShaderCode = FUniqueBuffer::Alloc(ReadSize);
+			int64 FileLocation = First.Offset + StartOffset;
+			if (FileLocation != Ar.Tell())
+			{
+				Ar.Seek(FileLocation);
+			}
+			Ar.Serialize(BlockOfShaderCode.GetData(), ReadSize);
+
+			for (auto& Block : BlockReadInfo)
+			{
+				FUniqueBuffer Code = FUniqueBuffer::Clone(static_cast<const uint8*>(BlockOfShaderCode.GetData()) + Block.Offset - First.Offset, Block.Size);
+				ShaderCode[Block.ShaderCodeIndex] = Code.MoveToShared();
+			}
+		}
+	}
+
+	void ReadShaderCode(FArchive& Ar, TArray<FShaderCodeReadInfo>& ReadInfo)
+	{
+		if (ReadInfo.Num() > 0)
+		{
+			// Sort by offset
+			ReadInfo.Sort([](FShaderCodeReadInfo const& A, FShaderCodeReadInfo const& B) { return A.Offset < B.Offset; });
+
+			// Fill in all the shader code data, reading old data in blocks
+			const int64 DesiredBlockSize = 10 * 1024 * 1024;
+			int64 PrevCookedShadersCodeStart = Ar.Tell();
+
+			// Setup starting block, then iterate the rest
+			int32 BlockStartIndex = 0;
+			int32 IndexCount = 1;
+
+			for (int64 i = 1; i < ReadInfo.Num(); ++i)
+			{
+				const FShaderCodeReadInfo& Next = ReadInfo[i];
+
+				// See how much we'd have to read to add this block, if it adds too much, read and start a new block, otherwise include in current block
+				int64 PotentialSize = Next.Offset - ReadInfo[BlockStartIndex].Offset + Next.Size;
+				if (PotentialSize > DesiredBlockSize)
+				{
+					ReadShaderCodeBlock(Ar, PrevCookedShadersCodeStart, MakeArrayView(&ReadInfo[BlockStartIndex], IndexCount));
+					BlockStartIndex = i;
+					IndexCount = 1;
+				}
+				else
+				{
+					IndexCount += 1;
+				}
+			}
+
+			// Read last block
+			ReadShaderCodeBlock(Ar, PrevCookedShadersCodeStart, MakeArrayView(&ReadInfo[BlockStartIndex], IndexCount));
+		}
+	}
+
 	bool LoadExistingShaderCodeLibrary(FString const& MetaDataDir)
 	{
 		FString IntermediateFormatPath = GetCodeArchiveFilename(MetaDataDir / TEXT("ShaderLibrarySource"), LibraryName, FormatName);
@@ -2144,11 +2211,13 @@ struct FEditorShaderCodeArchive
 					);
 				}
 
-				int64 PrevCookedShadersCodeStart = PrevCookedAr->Tell();
+				// Setup current shader library information and accumulate the reads we will need to do. Then read them.
+				TArray<FShaderCodeReadInfo> ReadInfo;
 				for (int32 PrevShaderMapIndex = 0; PrevShaderMapIndex < PrevCookedShaders.ShaderMapEntries.Num(); ++PrevShaderMapIndex)
 				{
-					AddShaderCode(CodeStats, PrevShaderMapIndex, PrevCookedShaders, PrevCookedShadersCodeStart, *PrevCookedAr);
+					PrepareAddShaderCode(CodeStats, PrevShaderMapIndex, PrevCookedShaders, ReadInfo);
 				}
+				ReadShaderCode(*PrevCookedAr, ReadInfo);
 			}
 
 			PrevCookedAr->Close();
