Enable RDG allocator debug threading checks when ASAN is enabled.

diff --git a/Engine/Source/Runtime/RenderCore/Public/RenderGraphAllocator.h b/Engine/Source/Runtime/RenderCore/Public/RenderGraphAllocator.h
--- a/Engine/Source/Runtime/RenderCore/Public/RenderGraphAllocator.h
+++ b/Engine/Source/Runtime/RenderCore/Public/RenderGraphAllocator.h
@@ -16,7 +16,7 @@
 #include "Templates/UnrealTemplate.h"
 
 #define RDG_USE_MALLOC USING_ADDRESS_SANITISER
-#define RDG_ALLOCATOR_DEBUG UE_BUILD_DEBUG
+#define RDG_ALLOCATOR_DEBUG USING_ADDRESS_SANITISER || UE_BUILD_DEBUG
 
 class FRDGAllocator;
 
