Fix incorrect logic for enabling DBuffer with Substrate on platform not using blendable GBuffer.

diff --git a/Engine/Source/Runtime/RenderCore/Private/RenderUtils.cpp b/Engine/Source/Runtime/RenderCore/Private/RenderUtils.cpp
--- a/Engine/Source/Runtime/RenderCore/Private/RenderUtils.cpp
+++ b/Engine/Source/Runtime/RenderCore/Private/RenderUtils.cpp
@@ -1478,7 +1478,10 @@ bool IsSingleLayerWaterDepthPrepassEnabled(const FStaticShaderPlatform Platform,
 
 RENDERCORE_API bool IsUsingDBuffers(const FStaticShaderPlatform Platform)
 {
-	return GDBufferPlatformMask[(int)Platform] && (!Substrate::IsSubstrateEnabled() || !Substrate::IsSubstrateBlendableGBufferEnabled(Platform));
+	// Dbuffer is enabled:
+	// * If project enabled it 
+	// * If Substrate is enabled and the platform is not set to BlendableGBuffer. When Substrate uses material buffer, the material data are not blendable and requires DBuffer.
+	return GDBufferPlatformMask[(int)Platform] || (Substrate::IsSubstrateEnabled() && !Substrate::IsSubstrateBlendableGBufferEnabled(Platform));
 }
 
 RENDERCORE_API bool AreSkinCacheShadersEnabled(EShaderPlatform Platform)
