Fixed crash in shutdown due to RDG task being active. Fixed crash in edge case where Texture->State[0] is null.

diff --git a/Engine/Source/Runtime/RenderCore/Private/RenderGraphBuilder.cpp b/Engine/Source/Runtime/RenderCore/Private/RenderGraphBuilder.cpp
--- a/Engine/Source/Runtime/RenderCore/Private/RenderGraphBuilder.cpp
+++ b/Engine/Source/Runtime/RenderCore/Private/RenderGraphBuilder.cpp
@@ -654,6 +654,7 @@ FRDGBuilder::FAsyncDeleter::~FAsyncDeleter()
 void FRDGBuilder::WaitForAsyncDeleteTask()
 {
 	FAsyncDeleter::LastTask.Wait();
+	FAsyncDeleter::LastTask = {};
 }
 
 const UE::Tasks::FTask& FRDGBuilder::GetAsyncDeleteTask()
@@ -4065,19 +4066,27 @@ void FRDGBuilder::AddLastTextureTransition(FRDGTexture* Texture)
 			// is going to be an incredibly rare case but needs to be handled correctly nonetheless.
 			if (DiscardPassPipeline == ERHIPipeline::AsyncCompute && EnumHasAnyFlags(Texture->Desc.Flags, ETextureCreateFlags::RenderTargetable | ETextureCreateFlags::DepthStencilTargetable))
 			{
-				EpilogueAccess = Texture->State[0]->Access;
-
-				for (uint32 SubresourceIndex = 1; SubresourceIndex < Texture->GetSubresourceCount(); ++SubresourceIndex)
+				if (Texture->State[0])
 				{
-					FRDGSubresourceState* SubresourceState = Texture->State[SubresourceIndex];
+					EpilogueAccess = Texture->State[0]->Access;
 
-					if (!SubresourceState || SubresourceState->Access != EpilogueAccess)
+					for (uint32 SubresourceIndex = 1; SubresourceIndex < Texture->GetSubresourceCount(); ++SubresourceIndex)
 					{
-						EpilogueAccess = GRHIMultiSubresourceDiscardIntermediateAccess;
-						SubresourceStateAfter.Access |= GRHIMultiSubresourceDiscardIntermediateAccess;
-						break;
+						FRDGSubresourceState* SubresourceState = Texture->State[SubresourceIndex];
+
+						if (!SubresourceState || SubresourceState->Access != EpilogueAccess)
+						{
+							EpilogueAccess = GRHIMultiSubresourceDiscardIntermediateAccess;
+							SubresourceStateAfter.Access |= GRHIMultiSubresourceDiscardIntermediateAccess;
+							break;
+						}
 					}
 				}
+				else
+				{
+					EpilogueAccess = GRHIMultiSubresourceDiscardIntermediateAccess;
+					SubresourceStateAfter.Access |= GRHIMultiSubresourceDiscardIntermediateAccess;
+				}
 			}
 
 			EpilogueResourceAccesses.Emplace(Texture->GetRHI(), ERHIAccess::Discard | EpilogueAccess, DiscardPassPipeline);
