Fix crash in ListShaderMaps for MemReport

Seeing ShaderMap with no Content, so just jump past those for now.

diff --git a/Engine/Source/Runtime/RenderCore/Private/ShaderMap.cpp b/Engine/Source/Runtime/RenderCore/Private/ShaderMap.cpp
--- a/Engine/Source/Runtime/RenderCore/Private/ShaderMap.cpp
+++ b/Engine/Source/Runtime/RenderCore/Private/ShaderMap.cpp
@@ -105,9 +105,12 @@ FAutoConsoleCommandWithArgsAndOutputDevice GListShaderMapsCmd(
 				}
 
 				TMap<FHashedName, TShaderRef<FShader>> Shaders;
-				ShaderMap->GetShaderList(Shaders);
 				TArray<FShaderPipelineRef> Pipelines;
-				ShaderMap->GetShaderPipelineList(Pipelines);
+				if (ShaderMap->GetContent())
+				{
+					ShaderMap->GetShaderList(Shaders);
+					ShaderMap->GetShaderPipelineList(Pipelines);
+				}
 
 				// Editor doesn't have the size baked, so grab it from shaders themselves.
 				if (CodeSize == 0)
