[Mobile][Engine] Adding in a CVar to help track exact moments when RenderTargetPool spikes occur.
- When CSV stats are used, we can get the exact tracked frame to correlate with the frames from reports.
- Normally, if enabled, we will just output the render target name and size.
- Added a cvar to control when the logging is enabled. Added compat CVars for Android and iOS for the cvar.
- Added the cvar as on by default for iOS

#rnx
#tests Tested locally with Editor. Tested using the given horde preflight.

diff --git a/Engine/Source/Runtime/RenderCore/Private/RenderTargetPool.cpp b/Engine/Source/Runtime/RenderCore/Private/RenderTargetPool.cpp
--- a/Engine/Source/Runtime/RenderCore/Private/RenderTargetPool.cpp
+++ b/Engine/Source/Runtime/RenderCore/Private/RenderTargetPool.cpp
@@ -23,6 +23,12 @@ UE_TRACE_EVENT_BEGIN(Cpu, FRenderTargetPool_CreateTexture, NoSync)
 	UE_TRACE_EVENT_FIELD(UE::Trace::WideString, Name)
 UE_TRACE_EVENT_END()
 
+static int32 GRenderTargetPoolLogCreationSizes = 0;
+static FAutoConsoleVariableRef CVarRenderTargetPoolLogCreationSizes(
+	TEXT("r.RenderTargetPool.LogCreationSizes"),
+	GRenderTargetPoolLogCreationSizes,
+	TEXT("Enable/disable warning log for render target pool creation sizes."));
+
 TRefCountPtr<IPooledRenderTarget> CreateRenderTarget(FRHITexture* Texture, const TCHAR* Name)
 {
 	check(Texture);
@@ -72,6 +78,24 @@ FPooledRenderTarget* FRenderTargetPool::CreateRenderTarget(FRHICommandListBase&
 		<< FRenderTargetPool_CreateTexture.Name(Name);
 #endif
 
+	if (GRenderTargetPoolLogCreationSizes > 0)
+	{
+#if CSV_PROFILER_STATS
+		int32 CurrentCaptureFrame = 0;
+		if (!FCsvProfiler::Get())
+		{
+			CurrentCaptureFrame = -1;
+		}
+		else if (FCsvProfiler::Get()->IsCapturing_Renderthread())
+		{
+			CurrentCaptureFrame = FCsvProfiler::Get()->GetCaptureFrameNumberRT();
+		}
+		UE_LOG(LogRenderTargetPool, Warning, TEXT("FRenderTargetPool::CreateRenderTarget{Frame %d}: %s create with expected size - %ld"), CurrentCaptureFrame, Name, RHICalcTexturePlatformSize(Desc).Size);
+#else
+		UE_LOG(LogRenderTargetPool, Warning, TEXT("FRenderTargetPool::CreateRenderTarget: %s create with expected size - %ld"), Name, RHICalcTexturePlatformSize(Desc).Size);
+#endif
+	}
+
 	const ERHIAccess AccessInitial = ERHIAccess::SRVMask;
 	FRHITextureCreateDesc CreateDesc(Desc, AccessInitial, Name);
 	const static FLazyName ClassName(TEXT("FPooledRenderTarget"));
