Added optional Z-increment parameter to FPixelShaderUtils::AddRasterizeToRectsPass, in order to control the Z of the quads (Z = quad index, i.e.
instance id * ZIncrement).
Combined with depth test enabled, this allows making the rasterization deterministic in case of overlapping quads

diff --git a/Engine/Shaders/Private/RenderGraphUtilities.usf b/Engine/Shaders/Private/RenderGraphUtilities.usf
--- a/Engine/Shaders/Private/RenderGraphUtilities.usf
+++ b/Engine/Shaders/Private/RenderGraphUtilities.usf
@@ -12,6 +12,7 @@ Buffer<uint4> RectUVBuffer;
 float DownsampleFactor;
 float2 InvViewSize;
 float2 InvTextureSize;
+float DepthIncrement;
 
 bool2 VertMax(uint VertexId)
 {
@@ -40,7 +41,9 @@ void RasterizeToRectsVS(
 	#endif
 	float2 VertexUV = select(VertMax(VertexId), RectUV.zw, RectUV.xy);
 
-	OutPosition = float4(float2(VertexCoord) * InvViewSize * float2(2.0f, -2.0f) + float2(-1.0, 1.0f), 0.0f, 1.0f);
+	float Z = InstanceId * DepthIncrement;
+
+	OutPosition = float4(float2(VertexCoord) * InvViewSize * float2(2.0f, -2.0f) + float2(-1.0, 1.0f), Z, 1.0f);
 	OutUV = VertexUV;
 	OutRectUV.x = VertMax(VertexId).x ? 1.0f : 0.0f;
 	OutRectUV.y = VertMax(VertexId).y ? 1.0f : 0.0f;
diff --git a/Engine/Source/Runtime/RenderCore/Public/PixelShaderUtils.h b/Engine/Source/Runtime/RenderCore/Public/PixelShaderUtils.h
--- a/Engine/Source/Runtime/RenderCore/Public/PixelShaderUtils.h
+++ b/Engine/Source/Runtime/RenderCore/Public/PixelShaderUtils.h
@@ -51,7 +51,7 @@ struct FPixelShaderUtils
 			SHADER_PARAMETER(FVector2f, InvViewSize)
 			SHADER_PARAMETER(FVector2f, InvTextureSize)
 			SHADER_PARAMETER(float, DownsampleFactor)
-			SHADER_PARAMETER(uint32, NumRects)
+			SHADER_PARAMETER(float, DepthIncrement)
 		END_SHADER_PARAMETER_STRUCT()
 
 		class FRectUV : SHADER_PERMUTATION_BOOL("RECT_UV");
@@ -258,7 +258,8 @@ struct FPixelShaderUtils
 		FIntPoint TextureSize = FIntPoint(1, 1),
 		FRDGBufferSRVRef RectUVBufferSRV = nullptr,
 		uint32 DownsampleFactor = 1,
-		ERDGPassFlags AdditionalPassFlags = ERDGPassFlags::None)
+		ERDGPassFlags AdditionalPassFlags = ERDGPassFlags::None,
+		float DepthIncrement = 0.0f)
 	{
 		FRasterizeToRectsVS::FPermutationDomain PermutationVector;
 		PermutationVector.Set<FRasterizeToRectsVS::FRectUV>(RectUVBufferSRV != nullptr);
@@ -269,7 +270,7 @@ struct FPixelShaderUtils
 		Parameters->VS.DownsampleFactor = 1.0f / DownsampleFactor;
 		Parameters->VS.RectCoordBuffer = RectCoordBufferSRV;
 		Parameters->VS.RectUVBuffer = RectUVBufferSRV;
-		Parameters->VS.NumRects = NumRects;
+		Parameters->VS.DepthIncrement = DepthIncrement;
 
 		ClearUnusedGraphResources(PixelShader, &Parameters->PS);
 
@@ -281,7 +282,7 @@ struct FPixelShaderUtils
 			Forward<FRDGEventName>(PassName),
 			Parameters,
 			PassFlags,
-			[Parameters, GlobalShaderMap, VertexShader, PixelShader, ViewportSize, BlendState, RasterizerState, DepthStencilState, StencilRef, bSkipRenderPass](FRDGAsyncTask, FRHICommandList& RHICmdList)
+			[Parameters, GlobalShaderMap, VertexShader, PixelShader, ViewportSize, BlendState, RasterizerState, DepthStencilState, StencilRef, bSkipRenderPass, NumRects](FRDGAsyncTask, FRHICommandList& RHICmdList)
 		{
 			if (bSkipRenderPass)
 			{
@@ -314,7 +315,7 @@ struct FPixelShaderUtils
 			SetShaderParameters(RHICmdList, PixelShader, PixelShader.GetPixelShader(), Parameters->PS);
 
 			const uint32 NumPrimitives = GRHISupportsRectTopology ? 1 : 2;
-			const uint32 NumInstances = Parameters->VS.NumRects;
+			const uint32 NumInstances = NumRects;
 			RHICmdList.DrawPrimitive(0, NumPrimitives, NumInstances);
 
 			if (bSkipRenderPass)
