Disabled RDG async setup queue while investigating a race condition.

diff --git a/Engine/Source/Runtime/RenderCore/Private/RenderGraphPrivate.cpp b/Engine/Source/Runtime/RenderCore/Private/RenderGraphPrivate.cpp
--- a/Engine/Source/Runtime/RenderCore/Private/RenderGraphPrivate.cpp
+++ b/Engine/Source/Runtime/RenderCore/Private/RenderGraphPrivate.cpp
@@ -310,7 +310,7 @@ FAutoConsoleVariableRef CVarRDGParallelCompile(
 	TEXT(" 1: compile tasks are launched (default);"),
 	ECVF_RenderThreadSafe);
 
-int32 GRDGAsyncSetupQueue = 1;
+int32 GRDGAsyncSetupQueue = 0;
 FAutoConsoleVariableRef CVarRDGAsyncSetupQueue(
 	TEXT("r.RDG.AsyncSetupQueue"), GRDGAsyncSetupQueue,
 	TEXT("RDG will setup passes in parallel when prompted by calls to FRDGBuilder::FlushSetupQueue.")
