Don't emit _RootShaderParameters when it's empty.

diff --git a/Engine/Shaders/Public/ShaderVersion.ush b/Engine/Shaders/Public/ShaderVersion.ush
--- a/Engine/Shaders/Public/ShaderVersion.ush
+++ b/Engine/Shaders/Public/ShaderVersion.ush
@@ -3,4 +3,4 @@
 // Each time the console command is executed it generates a new GUID. As this file is included
 // in Platform.ush (which should be included in any shader) it allows to invalidate the shader DDC.
 // 
-#pragma message("UESHADERMETADATA_VERSION ED27A877-556A-41DB-A67D-D0880AFFE659")
+#pragma message("UESHADERMETADATA_VERSION 8402F08A-02D1-44EF-96B0-B0C00BFC4CE1")
diff --git a/Engine/Source/Runtime/RenderCore/Private/ShaderParameterParser.cpp b/Engine/Source/Runtime/RenderCore/Private/ShaderParameterParser.cpp
--- a/Engine/Source/Runtime/RenderCore/Private/ShaderParameterParser.cpp
+++ b/Engine/Source/Runtime/RenderCore/Private/ShaderParameterParser.cpp
@@ -1121,9 +1121,7 @@ bool FShaderParameterParser::MoveShaderParametersToRootConstantBuffer(
 	// Generate the root cbuffer content.
 	if (RootParametersStructure && bNeedToMoveToRootConstantBuffer)
 	{
-		FStringBuilderBase ConstantBufferCode;
-
-		ConstantBufferCode << PlatformConfiguration.StableConstantBufferType << TEXT(" ") << FShaderParametersMetadata::kRootUniformBufferBindingName << PlatformConfiguration.GetStableConstantBufferRegisterString(ShaderFrequency) << TEXT("\n{\n");
+		FStringBuilderBase ConstantBufferMembers;
 
 		IterateShaderParameterMembers(
 			*RootParametersStructure,
@@ -1143,7 +1141,7 @@ bool FShaderParameterParser::MoveShaderParametersToRootConstantBuffer(
 				{
 					const FStringView Prefix = GetBindlessParameterPrefix(ParsedParameter->ConstantBufferParameterType);
 
-					ConstantBufferCode.Appendf(
+					ConstantBufferMembers.Appendf(
 						TEXT("uint %.*s%.*s : packoffset(c%d%s);\n"),
 						SVARG(Prefix),
 						SVARG(ParsedParameter->ParsedName),
@@ -1155,7 +1153,7 @@ bool FShaderParameterParser::MoveShaderParametersToRootConstantBuffer(
 				{
 					if (!ParsedParameter->ParsedArraySize.IsEmpty())
 					{
-						ConstantBufferCode.Appendf(
+						ConstantBufferMembers.Appendf(
 							TEXT("%.*s %s[%.*s] : packoffset(c%d%s);\n"),
 							SVARG(ParsedParameter->ParsedType),
 							ShaderBindingName,
@@ -1166,7 +1164,7 @@ bool FShaderParameterParser::MoveShaderParametersToRootConstantBuffer(
 					}
 					else
 					{
-						ConstantBufferCode.Appendf(
+						ConstantBufferMembers.Appendf(
 							TEXT("%.*s %s : packoffset(c%d%s);\n"),
 							SVARG(ParsedParameter->ParsedType),
 							ShaderBindingName,
@@ -1179,16 +1177,23 @@ bool FShaderParameterParser::MoveShaderParametersToRootConstantBuffer(
 			}
 		});
 
-		ConstantBufferCode << TEXT("}\n\n");
+		if (ConstantBufferMembers.Len() > 0)
+		{
+			FStringBuilderBase ConstantBufferCode;
+
+			ConstantBufferCode << PlatformConfiguration.StableConstantBufferType << TEXT(" ") << FShaderParametersMetadata::kRootUniformBufferBindingName << PlatformConfiguration.GetStableConstantBufferRegisterString(ShaderFrequency) << TEXT("\n{\n");
+			ConstantBufferCode << ConstantBufferMembers;
+			ConstantBufferCode << TEXT("}\n\n");
 
-		FString NewShaderCode = (
-			MakeInjectedShaderCodeBlock(TEXT("MoveShaderParametersToRootConstantBuffer"), *ConstantBufferCode) +
-			PreprocessedShaderSource);
+			FString NewShaderCode = (
+				MakeInjectedShaderCodeBlock(TEXT("MoveShaderParametersToRootConstantBuffer"), *ConstantBufferCode) +
+				PreprocessedShaderSource);
 
-		PreprocessedShaderSource = MoveTemp(NewShaderCode);
+			PreprocessedShaderSource = MoveTemp(NewShaderCode);
 
-		bMovedLoosedParametersToRootConstantBuffer = true;
-		bModifiedShader = true;
+			bMovedLoosedParametersToRootConstantBuffer = true;
+			bModifiedShader = true;
+		}
 	} // if (CompilerInput.RootParametersStructure && bNeedToMoveToRootConstantBuffer)
 
 	return bSuccess;
