Replaced FPipe with explicit task chain as a speculative fix to async setup queue race.

diff --git a/Engine/Source/Runtime/RenderCore/Private/RenderGraphBuilder.cpp b/Engine/Source/Runtime/RenderCore/Private/RenderGraphBuilder.cpp
--- a/Engine/Source/Runtime/RenderCore/Private/RenderGraphBuilder.cpp
+++ b/Engine/Source/Runtime/RenderCore/Private/RenderGraphBuilder.cpp
@@ -1666,13 +1666,15 @@ void FRDGBuilder::Compile()
 
 void FRDGBuilder::LaunchAsyncSetupQueueTask()
 {
-	if (!AsyncSetupQueue.Pipe.HasWork())
+	check(!bCompiling);
+
+	if (AsyncSetupQueue.LastTask.IsCompleted())
 	{
-		AsyncSetupQueue.Pipe.Launch(UE_SOURCE_LOCATION, [this]() mutable
+		AsyncSetupQueue.LastTask = UE::Tasks::Launch(UE_SOURCE_LOCATION, [this]() mutable
 		{
 			ProcessAsyncSetupQueue();
 
-		}, UE::Tasks::ETaskPriority::High);
+		}, AsyncSetupQueue.LastTask, UE::Tasks::ETaskPriority::High);
 	}
 }
 
@@ -1784,7 +1786,8 @@ void FRDGBuilder::Execute()
 	{
 		BeginFlushResourcesRHI();
 		WaitForParallelSetupTasks(ERDGSetupTaskWaitPoint::Compile);
-		AsyncSetupQueue.Pipe.WaitUntilEmpty();
+		AsyncSetupQueue.LastTask.Wait();
+		AsyncSetupQueue.LastTask = {};
 		ProcessAsyncSetupQueue();
 
 		const int32 ParallelCompileResourceThreshold = 32;
@@ -3165,7 +3168,7 @@ void FRDGBuilder::AllocateTransientResources(TConstArrayView<FCollectResourceOp>
 	}
 	
 	SCOPED_NAMED_EVENT_TEXT("FRDGBuilder::AllocateTransientResources", FColor::Magenta);
-	TransientResourceAllocator->SetCreateMode(ParallelSetup.bEnabled ? ERHITransientResourceCreateMode::Task : ERHITransientResourceCreateMode::Inline);
+	TransientResourceAllocator->SetCreateMode(bParallelCompileEnabled ? ERHITransientResourceCreateMode::Task : ERHITransientResourceCreateMode::Inline);
 
 	TArray<TPair<FRDGViewableResource*, FRHITransientResource*>, FRDGArrayAllocator> AllocatedResources;
 	AllocatedResources.Reserve(Ops.Num() / 2);
@@ -3914,7 +3917,7 @@ void FRDGBuilder::CreatePassBarriers()
 	{
 		CreateTransition(TaskContext, TransitionCreateQueue[Index]);
 
-	}, ParallelSetup.bEnabled ? EParallelForFlags::None : EParallelForFlags::ForceSingleThread);
+	}, bParallelCompileEnabled ? EParallelForFlags::None : EParallelForFlags::ForceSingleThread);
 
 	TransitionCreateQueue.Reset();
 }
diff --git a/Engine/Source/Runtime/RenderCore/Private/RenderGraphPrivate.cpp b/Engine/Source/Runtime/RenderCore/Private/RenderGraphPrivate.cpp
--- a/Engine/Source/Runtime/RenderCore/Private/RenderGraphPrivate.cpp
+++ b/Engine/Source/Runtime/RenderCore/Private/RenderGraphPrivate.cpp
@@ -310,7 +310,7 @@ FAutoConsoleVariableRef CVarRDGParallelCompile(
 	TEXT(" 1: compile tasks are launched (default);"),
 	ECVF_RenderThreadSafe);
 
-int32 GRDGAsyncSetupQueue = 0;
+int32 GRDGAsyncSetupQueue = 1;
 FAutoConsoleVariableRef CVarRDGAsyncSetupQueue(
 	TEXT("r.RDG.AsyncSetupQueue"), GRDGAsyncSetupQueue,
 	TEXT("RDG will setup passes in parallel when prompted by calls to FRDGBuilder::FlushSetupQueue.")
diff --git a/Engine/Source/Runtime/RenderCore/Public/RenderGraphBuilder.h b/Engine/Source/Runtime/RenderCore/Public/RenderGraphBuilder.h
--- a/Engine/Source/Runtime/RenderCore/Public/RenderGraphBuilder.h
+++ b/Engine/Source/Runtime/RenderCore/Public/RenderGraphBuilder.h
@@ -945,7 +945,7 @@ private:
 
 		UE::FMutex Mutex;
 		TArray<FAsyncSetupOp, FRDGArrayAllocator> Ops;
-		UE::Tasks::FPipe Pipe{ TEXT("FRDGBuilder::AsyncSetupQueue") };
+		UE::Tasks::FTask LastTask;
 		bool bEnabled = false;
 
 	} AsyncSetupQueue;
