[shaders] clean up secondary preprocessing/output handling
- make sure to compress and strip secondary preprocessed source prior to sending to workers
- move merging of error messages for secondary jobs to job OnComplete so the errors for each compile step can be properly remapped using the preprocess output remapping data (which is not serialized to the worker)
- move job OnComplete call to before adding the job to the cache, so merged diagnostic messages are stored in the cache (and avoid having to re-merge the messages for duplicate jobs/cache hits)
- add optional prefix to secondary output error messages, driven by IShaderFormat virtual function to return a custom name (implement for the 3 formats which currently allow secondary compilation)

diff --git a/Engine/Source/Developer/TargetPlatform/Public/Interfaces/IShaderFormat.h b/Engine/Source/Developer/TargetPlatform/Public/Interfaces/IShaderFormat.h
--- a/Engine/Source/Developer/TargetPlatform/Public/Interfaces/IShaderFormat.h
+++ b/Engine/Source/Developer/TargetPlatform/Public/Interfaces/IShaderFormat.h
@@ -168,6 +168,12 @@ public:
 		const FShaderCompilerEnvironment& Environment,
 		const FShaderPreprocessOutput& PreprocessOutput) const { return false; }
 
+	/**
+	* For shader formats which support secondary compilation, this can optionally return a name that is appended to diagnostic messages
+	* to distinguish errors in the secondary compile from the primary.
+	*/
+	virtual FStringView GetSecondaryCompileName() const { return TEXTVIEW(""); }
+
 	/**
 	 * Virtual function that can be implemented by backends which require secondary compilation to specify a unique "key" identifier
 	 * for the packing scheme (for validation in the runtime at load). see RequiresSecondaryCompile and the overloads of 
diff --git a/Engine/Source/Runtime/RenderCore/Private/ShaderCore.cpp b/Engine/Source/Runtime/RenderCore/Private/ShaderCore.cpp
--- a/Engine/Source/Runtime/RenderCore/Private/ShaderCore.cpp
+++ b/Engine/Source/Runtime/RenderCore/Private/ShaderCore.cpp
@@ -1781,22 +1781,30 @@ public:
 				UE_DEBUG_BREAK();
 			}
 		}
+		double PreprocessStartTime = FPlatformTime::Seconds();
 		Job.PreprocessOutput.bSucceeded = Backend->PreprocessShader(Job.Input, Environment, Job.PreprocessOutput);
+		Job.PreprocessOutput.PreprocessTime = FPlatformTime::Seconds() - PreprocessStartTime;
+
 		if (Job.PreprocessOutput.bSucceeded && Backend->RequiresSecondaryCompile(Job.Input, Environment, Job.PreprocessOutput))
 		{
+			PreprocessStartTime = FPlatformTime::Seconds();
 			// check that we're not compiling for a pipeline; this would be theretically possible to support but we'd have to enforce that all jobs in the pipeline
 			// have a secondary compile if any do, to ensure interpolators match across pipeline stages. currently there's no such guarantee so this is a safety measure.
 			checkf(!Job.Input.bCompilingForShaderPipeline, TEXT("Cannot combine shader pipeline interpolator removal optimization with secondary compilation at this time"));
 			Job.SecondaryPreprocessOutput = MakeUnique<FShaderPreprocessOutput>();
 			Job.SecondaryPreprocessOutput->bIsSecondary = true;
-			Job.PreprocessOutput.bSucceeded &= Backend->PreprocessShader(Job.Input, Environment, *Job.SecondaryPreprocessOutput);
+			Job.SecondaryPreprocessOutput->bSucceeded = Backend->PreprocessShader(Job.Input, Environment, *Job.SecondaryPreprocessOutput);
+			Job.SecondaryPreprocessOutput->PreprocessTime = FPlatformTime::Seconds() - PreprocessStartTime;
+
+			// Mark the main job preprocess output as failed if the secondary preprocess fails, even if the primary succeeds
+			// This means the job submission code only needs to check one flag to determine whether to actually submit the job 
+			// (no point in trying to compile if either of the two preprocess steps fail)
+			Job.PreprocessOutput.bSucceeded &= Job.SecondaryPreprocessOutput->bSucceeded;
 		}
 	}
 
 	static bool PreprocessShaderInternal(const IShaderFormat* Backend, FShaderCompileJob& Job)
 	{
-		const double StartPreprocessTime = FPlatformTime::Seconds();
-
 		if (IsValidRef(Job.Input.SharedEnvironment))
 		{
 			// only create new environment & merge if necessary, save some allocs
@@ -1811,43 +1819,68 @@ public:
 			PreprocessShaderInternal(Backend, Job, Job.Input.Environment);
 		}
 
-		if (Job.PreprocessOutput.bSucceeded)
-		{
-			if (!Job.Input.Environment.CompilerFlags.Contains(CFLAG_DisableSourceStripping))
+		Job.Output.PreprocessTime = 0;
+		bool bSucceeded = true;
+
+		auto PostPreprocessFunc = [&Job, &bSucceeded](FShaderPreprocessOutput& PreprocessOutput, FStringView DiagPrefix = FStringView())
 			{
-				// if the preprocessed job cache is enabled we strip the preprocessed code if not explicitly disabled; this removes comments, 
-				// line directives and blank lines to improve deduplication (and populates data required to remap diagnostic messages to correct
-				// filenames and line numbers)
-				Job.PreprocessOutput.StripCode(Job.Input.NeedsOriginalShaderSource());
-
-				FShaderCompilerInputHash Hash = Job.GetInputHash();
-				// Replace the placeholder debug hash value appended in StripCode with the real job input hash
-				FShaderSource::FViewType DebugHashStr = GetShaderSourceDebugHashPrefix();
-				FShaderSource::FViewType SourceView = Job.PreprocessOutput.GetSourceView();
-				int32 DebugHashLoc = SourceView.Find(DebugHashStr) + DebugHashStr.Len();
-				int32 NewlineLoc = SourceView.Find(ANSITEXTVIEW("\n"), DebugHashLoc);
-				FShaderSource::TStringBuilder<2 * sizeof(FShaderCompilerInputHash::ByteArray) + 1> HashStr;
-				HashStr << Hash;
-				// handle case where we didn't add the debug hash association (this can be disabled by a cvar for debugging purposes)
-				if (NewlineLoc - DebugHashLoc == HashStr.Len())
+				if (PreprocessOutput.bSucceeded)
 				{
-					FMemory::Memcpy(Job.PreprocessOutput.EditSource().GetData() + DebugHashLoc, HashStr.GetData(), sizeof(FShaderSource::CharType) * HashStr.Len());
+					if (!Job.Input.Environment.CompilerFlags.Contains(CFLAG_DisableSourceStripping))
+					{
+						// if the preprocessed job cache is enabled we strip the preprocessed code if not explicitly disabled; this removes comments, 
+						// line directives and blank lines to improve deduplication (and populates data required to remap diagnostic messages to correct
+						// filenames and line numbers)
+						PreprocessOutput.StripCode(Job.Input.NeedsOriginalShaderSource());
+						FShaderCompilerInputHash Hash = Job.GetInputHash();
+						// Replace the placeholder debug hash value appended in StripCode with the real job input hash
+						FShaderSource::FViewType DebugHashStr = GetShaderSourceDebugHashPrefix();
+						FShaderSource::FViewType SourceView = PreprocessOutput.GetSourceView();
+						int32 DebugHashLoc = SourceView.Find(DebugHashStr) + DebugHashStr.Len();
+						int32 NewlineLoc = SourceView.Find(ANSITEXTVIEW("\n"), DebugHashLoc);
+						FShaderSource::TStringBuilder<2 * sizeof(FShaderCompilerInputHash::ByteArray) + 1> HashStr;
+						HashStr << Hash;
+						// handle case where we didn't add the debug hash association (this can be disabled by a cvar for debugging purposes)
+						if (NewlineLoc - DebugHashLoc == HashStr.Len())
+						{
+							FMemory::Memcpy(PreprocessOutput.EditSource().GetData() + DebugHashLoc, HashStr.GetData(), sizeof(FShaderSource::CharType) * HashStr.Len());
+						}
+					}
+					// always compress the code after stripping to minimize memory footprint
+					PreprocessOutput.CompressCode();
+
+					// intentionally adding to main output's diagnostic data array and preprocess time here
+					Job.Output.PreprocessTime += PreprocessOutput.PreprocessTime;
+					Job.Output.ShaderDiagnosticDatas.Append(PreprocessOutput.GetDiagnosticDatas());
 				}
-			}
+				
+				// append a prefix to secondary job error messages to distinguish them from the primary compilation
+				if (!DiagPrefix.IsEmpty())
+				{
+					for (FShaderCompilerError& Error : PreprocessOutput.Errors)
+					{
+						Error.StrippedErrorMessage = DiagPrefix + Error.StrippedErrorMessage;
+					}
+				}
+				Job.Output.Errors.Append(PreprocessOutput.Errors);
+				bSucceeded &= PreprocessOutput.bSucceeded;
+			};
 
-			// always compress the code after stripping to minimize memory footprint
-			Job.PreprocessOutput.CompressCode();
+		PostPreprocessFunc(Job.PreprocessOutput);
+		if (Job.SecondaryPreprocessOutput.IsValid())
+		{
+			TStringBuilder<128> DiagPrefix;
+			DiagPrefix << "(" << Backend->GetSecondaryCompileName() << ") ";
+			PostPreprocessFunc(*Job.SecondaryPreprocessOutput, DiagPrefix);
 		}
 
-		Job.PreprocessOutput.ElapsedTime = FPlatformTime::Seconds() - StartPreprocessTime;
-		Job.Output.PreprocessTime = Job.PreprocessOutput.ElapsedTime;
-		Job.Output.ShaderDiagnosticDatas = Job.PreprocessOutput.GetDiagnosticDatas();
-		Job.Output.Errors.Append(Job.PreprocessOutput.Errors);
-		return Job.PreprocessOutput.bSucceeded;
+		return bSucceeded;
 	}
 
-	static void CombineOutputs(const IShaderFormat* Compiler, FShaderCompileJob& Job)
+	static void CombineCodeOutputs(const IShaderFormat* Compiler, FShaderCompileJob& Job)
 	{
+		check(Job.SecondaryOutput.IsValid());
+
 		// Pack shader code results together
 		// [int32 key][uint32 primary length][uint32 secondary length][full primary shader code][full secondary shader code]
 		TArray<uint8> CombinedSource;
@@ -1873,6 +1906,10 @@ public:
 		Job.Output.ShaderCode = {};
 		Job.Output.ShaderCode.GetWriteAccess() = MoveTemp(CombinedSource);
 		Job.Output.ShaderCode.FinalizeShaderCode(SymbolsBuffer);
+
+		// note: we intentionally keep the secondary output object around so duplicate/cache hit jobs can also dump debug info. it's a bit of memory overhead but preferable
+		// to requiring that enabling/disabling debug information mutates the DDC key (which would be necessary if we cleared it, since we output separate debug information 
+		// for the secondary jobs).
 	}
 
 	static void InvokeCompile(const IShaderFormat* Compiler, FShaderCompileJob& Job, const FString& WorkingDirectory, FString& OutExceptionMsg, FString& OutExceptionCallstack)
@@ -1881,9 +1918,13 @@ public:
 		__try
 #endif
 		{
+			// decompress if necessary; this is a no-op if source is not compressed.
+			Job.PreprocessOutput.DecompressCode();
+
 			if (Job.SecondaryOutput.IsValid())
 			{
 				check(Job.SecondaryPreprocessOutput.IsValid());
+				Job.SecondaryPreprocessOutput->DecompressCode();
 				Compiler->CompilePreprocessedShader(Job.Input, Job.PreprocessOutput, *Job.SecondaryPreprocessOutput, Job.Output, *Job.SecondaryOutput, WorkingDirectory);
 			}
 			else
@@ -1912,9 +1953,6 @@ public:
 
 		Job.Output.Errors.Append(Job.PreprocessOutput.Errors);
 		
-		// decompress if necessary; this is a no-op if source is not compressed.
-		Job.PreprocessOutput.DecompressCode();
-
 		if (Job.SecondaryPreprocessOutput.IsValid())
 		{
 			Job.SecondaryOutput = MakeUnique<FShaderCompilerOutput>();
@@ -1932,7 +1970,7 @@ public:
 			{
 				Job.SecondaryOutput->GenerateOutputHash();
 			}
-			CombineOutputs(Compiler, Job);
+			CombineCodeOutputs(Compiler, Job);
 		}
 
 		// ensure the target field is set on the job output struct as we use it for validation during serialization
@@ -4398,10 +4436,16 @@ void FShaderCompileJob::SerializeOutput(FShaderCacheSerializeContext& Ctx, int32
 
 void FShaderCompileJob::OnComplete(FShaderDebugDataContext& Ctx)
 {
-	if (Input.NeedsOriginalShaderSource())
+	const IShaderFormat* ShaderFormat = nullptr;
+	FCompressedBuffer SymbolsBuffer = Output.bSucceeded ? Output.GetFinalizedCodeResource().GetSymbolsBuffer() : FCompressedBuffer();
+	bool bHasSecondary = SecondaryOutput.IsValid();
+
+	const bool bHasSymbols = SymbolsBuffer.GetRawSize() > 0;
+	if (bHasSymbols || Input.DumpDebugInfoEnabled() || bHasSecondary)
 	{
-		// Decompress the code if needed by debug info or source extraction
-		PreprocessOutput.DecompressCode();
+		// don't bother looking up the shader format unless needed, this can have non-trivial cost due to thread safety
+		// and the possibility of available target platforms needing to be refreshed
+		ShaderFormat = GetTargetPlatformManagerRef().FindShaderFormat(Input.ShaderFormat);
 	}
 
 	// For jobs which applied source stripping, we need to remap error messages whether or not the job was actually the one that ran
@@ -4409,20 +4453,39 @@ void FShaderCompileJob::OnComplete(FShaderDebugDataContext& Ctx)
 	if (!Input.Environment.CompilerFlags.Contains(CFLAG_DisableSourceStripping))
 	{
 		PreprocessOutput.RemapErrors(Output);
+		if (bHasSecondary)
+		{
+			check(SecondaryPreprocessOutput.IsValid());
+			SecondaryPreprocessOutput->RemapErrors(*SecondaryOutput);
+		}
 	}
 
-	// always set the preprocess time on the output struct, regardless of whether this job hit the cache or actually compiled, since
-	// we always run preprocessing as a first step when the job was submitted.
-	Output.PreprocessTime = PreprocessOutput.ElapsedTime;
+	if (bHasSecondary)
+	{
+		check(ShaderFormat);
 
-	const IShaderFormat* ShaderFormat = nullptr;
-	FCompressedBuffer SymbolsBuffer = Output.bSucceeded ? Output.GetFinalizedCodeResource().GetSymbolsBuffer() : FCompressedBuffer();
-	const bool bHasSymbols = SymbolsBuffer.GetRawSize() > 0;
-	if (bHasSymbols || Input.DumpDebugInfoEnabled())
+		// append a prefix to secondary job error messages to distinguish them from the primary compilation
+		FStringView SecondaryName = ShaderFormat->GetSecondaryCompileName();
+		if (!SecondaryName.IsEmpty())
+		{
+			TStringBuilder<128> Prefix;
+			Prefix << "(" << SecondaryName << ") ";
+			for (FShaderCompilerError& Error : SecondaryOutput->Errors)
+			{
+				Error.StrippedErrorMessage.InsertAt(0, Prefix.ToString());
+			}
+		}
+		Output.Errors.Append(SecondaryOutput->Errors);
+	}
+
+	if (Input.NeedsOriginalShaderSource())
 	{
-		// don't bother looking up the shader format unless needed, this can have non-trivial cost due to thread safety
-		// and the possibility of available target platforms needing to be refreshed
-		ShaderFormat = GetTargetPlatformManagerRef().FindShaderFormat(Input.ShaderFormat);
+		// Decompress the code if needed by debug info or source extraction
+		PreprocessOutput.DecompressCode();
+		if (SecondaryPreprocessOutput.IsValid())
+		{
+			SecondaryPreprocessOutput->DecompressCode();
+		}
 	}
 
 	// output debug info for the job if either the job was the one which actually executed the compile, or debug info is requested 
diff --git a/Engine/Source/Runtime/RenderCore/Public/ShaderPreprocessTypes.h b/Engine/Source/Runtime/RenderCore/Public/ShaderPreprocessTypes.h
--- a/Engine/Source/Runtime/RenderCore/Public/ShaderPreprocessTypes.h
+++ b/Engine/Source/Runtime/RenderCore/Public/ShaderPreprocessTypes.h
@@ -225,9 +225,9 @@ public:
 		return MakeArrayView(Errors);
 	}
 
-	double GetElapsedTime() const
+	double GetPreprocessTime() const
 	{
-		return ElapsedTime;
+		return PreprocessTime;
 	}
 
 	TArray<FShaderDiagnosticData>& EditDiagnosticDatas()
@@ -284,7 +284,7 @@ private:
 	TArray<FString> PragmaDirectives;
 	FShaderDiagnosticRemapper Remapper;
 
-	double ElapsedTime = 0.0;
+	double PreprocessTime = 0.0;
 	bool bSucceeded = false;
 	bool bIsSecondary = false;
 
