[shaders] improved error message when shader library contains shaders for types which have been removed/renamed.

diff --git a/Engine/Source/Runtime/RenderCore/Private/ShaderCodeLibrary.cpp b/Engine/Source/Runtime/RenderCore/Private/ShaderCodeLibrary.cpp
--- a/Engine/Source/Runtime/RenderCore/Private/ShaderCodeLibrary.cpp
+++ b/Engine/Source/Runtime/RenderCore/Private/ShaderCodeLibrary.cpp
@@ -3835,6 +3835,8 @@ public:
 				for (const TPair<uint64, FShaderCodeStats::FPerTypeStats>& Pair : PlatformCodeStats.ShaderTypeStats)
 				{
 					const FShaderType* ShaderType = FindShaderTypeByName(FHashedName(Pair.Key));
+					checkf(ShaderType != nullptr, TEXT("Shader library contains a reference to an FShaderType with hash %llu which does not exist.\n"
+												       "If it has been removed or renamed, please bump CookIncrementalVersion to force a rebuild"), Pair.Key);
 					CountersByTypeName.Add(ShaderType, Pair.Value);
 				}
 				
