[Shaders] Remove unused local variable in AddShaderSourceFileEntry().

#rnx

diff --git a/Engine/Source/Runtime/RenderCore/Private/ShaderCore.cpp b/Engine/Source/Runtime/RenderCore/Private/ShaderCore.cpp
--- a/Engine/Source/Runtime/RenderCore/Private/ShaderCore.cpp
+++ b/Engine/Source/Runtime/RenderCore/Private/ShaderCore.cpp
@@ -2159,15 +2159,9 @@ void AddShaderSourceFileEntry(TArray<FString>& OutVirtualFilePaths, FString Virt
 	{
 		OutVirtualFilePaths.Add(VirtualFilePath);
 
-		TArray<FString> ShaderIncludes;
-
 		const uint32 DepthLimit = 100;
 		const bool bPreprocessDependencies = true;
 		InternalGetShaderIncludes(*VirtualFilePath, *VirtualFilePath, OutVirtualFilePaths, ShaderPlatform, DepthLimit, false, ShaderPlatformName, bPreprocessDependencies);
-		for( int32 IncludeIdx=0; IncludeIdx < ShaderIncludes.Num(); IncludeIdx++ )
-		{
-			OutVirtualFilePaths.AddUnique(ShaderIncludes[IncludeIdx]);
-		}
 	}
 }
 
