[shaders] fixes related to cloned jobs:
- fix setting of fields on output struct based on preprocess results. we can't set this before submitting the job as when jobs are cloned the output struct is intentionally skipped, so we do it in the oncomplete function instead.
- ensure secondary preprocess output is populated when a job is cloned (not sure if this ever happens during a cook, but good to be consistent in case)

diff --git a/Engine/Source/Runtime/Engine/Private/ShaderCompiler/ShaderCompilerJobCache.cpp b/Engine/Source/Runtime/Engine/Private/ShaderCompiler/ShaderCompilerJobCache.cpp
--- a/Engine/Source/Runtime/Engine/Private/ShaderCompiler/ShaderCompilerJobCache.cpp
+++ b/Engine/Source/Runtime/Engine/Private/ShaderCompiler/ShaderCompilerJobCache.cpp
@@ -1130,6 +1130,11 @@ static FShaderCommonCompileJob* CloneJob_Single(const FShaderCompileJob* SrcJob)
 		Job->InputHash = SrcJob->InputHash;
 		Job->bInputHashSet = true;
 	}
+	if (SrcJob->SecondaryPreprocessOutput.IsValid())
+	{
+		Job->SecondaryPreprocessOutput = MakeUnique<FShaderPreprocessOutput>();
+		*Job->SecondaryPreprocessOutput = *SrcJob->SecondaryPreprocessOutput;
+	}
 	ensure(Job->bInputHashSet == SrcJob->bInputHashSet);
 	return Job;
 }
diff --git a/Engine/Source/Runtime/RenderCore/Private/ShaderCore.cpp b/Engine/Source/Runtime/RenderCore/Private/ShaderCore.cpp
--- a/Engine/Source/Runtime/RenderCore/Private/ShaderCore.cpp
+++ b/Engine/Source/Runtime/RenderCore/Private/ShaderCore.cpp
@@ -1819,7 +1819,6 @@ public:
 			PreprocessShaderInternal(Backend, Job, Job.Input.Environment);
 		}
 
-		Job.Output.PreprocessTime = 0;
 		bool bSucceeded = true;
 
 		auto PostPreprocessFunc = [&Job, &bSucceeded](FShaderPreprocessOutput& PreprocessOutput, FStringView DiagPrefix = FStringView())
@@ -1848,10 +1847,6 @@ public:
 					}
 					// always compress the code after stripping to minimize memory footprint
 					PreprocessOutput.CompressCode();
-
-					// intentionally adding to main output's diagnostic data array and preprocess time here
-					Job.Output.PreprocessTime += PreprocessOutput.PreprocessTime;
-					Job.Output.ShaderDiagnosticDatas.Append(PreprocessOutput.GetDiagnosticDatas());
 				}
 				
 				// append a prefix to secondary job error messages to distinguish them from the primary compilation
@@ -4482,6 +4477,20 @@ void FShaderCompileJob::OnComplete(FShaderDebugDataContext& Ctx)
 		}
 	}
 
+	auto AppendPreprocessInfo = [this](const FShaderPreprocessOutput& InPreprocessOutput)
+		{
+			// intentionally adding to main output's diagnostic data array and preprocess time here; at this point 
+			// outputs are already merged and we want to store the merged results of both of these in the cache.
+			Output.PreprocessTime += InPreprocessOutput.PreprocessTime;
+			Output.ShaderDiagnosticDatas.Append(InPreprocessOutput.GetDiagnosticDatas());
+		};
+
+	AppendPreprocessInfo(PreprocessOutput);
+	if (SecondaryPreprocessOutput.IsValid())
+	{
+		AppendPreprocessInfo(*SecondaryPreprocessOutput);
+	}
+
 	// output debug info for the job if either the job was the one which actually executed the compile, or debug info is requested 
 	// for all jobs including ones which hit the job cache (or matched and was deduplicated with another in-flight job)
 	// note that depending on shaderformat implementation this may not necessarily be _all_ debug artifacts, since some require
