Silence incorrect V1051 warnings

#rnx

diff --git a/Engine/Plugins/ChaosClothAssetEditor/Source/ChaosClothAssetDataflowNodes/Private/ChaosClothAsset/SelectionNode.cpp b/Engine/Plugins/ChaosClothAssetEditor/Source/ChaosClothAssetDataflowNodes/Private/ChaosClothAsset/SelectionNode.cpp
--- a/Engine/Plugins/ChaosClothAssetEditor/Source/ChaosClothAssetDataflowNodes/Private/ChaosClothAsset/SelectionNode.cpp
+++ b/Engine/Plugins/ChaosClothAssetEditor/Source/ChaosClothAssetDataflowNodes/Private/ChaosClothAsset/SelectionNode.cpp
@@ -549,7 +549,7 @@ PRAGMA_ENABLE_DEPRECATION_WARNINGS
 					{
 						MutableThis->TransferCollectionHash = InTransferCollectionHash;
 
-						if(TransferCollectionHash)
+						if(TransferCollectionHash) //-V1051
 						{
 							TSet<int32> PrimaryFinalSelection;
 							if (Private::TransferSelectionSet<false>(TransferClothCollection, ClothCollection, InInputName, SelectionGroupName, SimTransferType, TransferSelectionThreshold, PrimaryFinalSelection))
diff --git a/Engine/Plugins/Experimental/PSDImporter/Source/PSDImporterCore/Private/Readers/LayerReader.cpp b/Engine/Plugins/Experimental/PSDImporter/Source/PSDImporterCore/Private/Readers/LayerReader.cpp
--- a/Engine/Plugins/Experimental/PSDImporter/Source/PSDImporterCore/Private/Readers/LayerReader.cpp
+++ b/Engine/Plugins/Experimental/PSDImporter/Source/PSDImporterCore/Private/Readers/LayerReader.cpp
@@ -482,7 +482,7 @@ namespace UE::PSDImporter::Private
 	TFuture<bool> FLayersReader::Read(Private::FReadContext& InContext)
 	{
 		File::FPSDDocument* Document = InContext.Document2;
-		if (psd::LayerMaskSection* LayerMasks = psd::ParseLayerMaskSection(InContext.Document, InContext.File, &InContext.Allocator))
+		if (psd::LayerMaskSection* LayerMasks = psd::ParseLayerMaskSection(InContext.Document, InContext.File, &InContext.Allocator)) //-V1051
 		{
 			File::FPSDLayerAndMaskInformation& LayersAndMasks = Document->LayerAndMaskInformation;
 			LayersAndMasks.NumLayers = LayerMasks->layerCount;
diff --git a/Engine/Plugins/Experimental/PythonScriptPlugin/Source/PythonScriptPlugin/Private/PipInstall.cpp b/Engine/Plugins/Experimental/PythonScriptPlugin/Source/PythonScriptPlugin/Private/PipInstall.cpp
--- a/Engine/Plugins/Experimental/PythonScriptPlugin/Source/PythonScriptPlugin/Private/PipInstall.cpp
+++ b/Engine/Plugins/Experimental/PythonScriptPlugin/Source/PythonScriptPlugin/Private/PipInstall.cpp
@@ -338,7 +338,7 @@ public:
 		bool res = PlatformFile.IterateDirectory(FilenameOrDirectory, SubDirVisit);
 
 		bOrphan = bOrphan && SubDirVisit.bOrphan;
-		if (SubDirVisit.bOrphan)
+		if (SubDirVisit.bOrphan) //-V1051
 		{
 			Orphans.Add(FilenameOrDirectory);
 		}
diff --git a/Engine/Plugins/FX/Niagara/Source/Niagara/Private/NiagaraDataInterfaceGrid2DCollection.cpp b/Engine/Plugins/FX/Niagara/Source/Niagara/Private/NiagaraDataInterfaceGrid2DCollection.cpp
--- a/Engine/Plugins/FX/Niagara/Source/Niagara/Private/NiagaraDataInterfaceGrid2DCollection.cpp
+++ b/Engine/Plugins/FX/Niagara/Source/Niagara/Private/NiagaraDataInterfaceGrid2DCollection.cpp
@@ -1956,7 +1956,7 @@ bool UNiagaraDataInterfaceGrid2DCollection::InitPerInstanceData(void* PerInstanc
 	InstanceData->NumAttributes = NumAttribChannelsFound;
 
 	// give an error if there are any unnamed attributes
-	if (NumAttributes > 0 && NumNamedAttribChannelsFound > 0)
+	if (NumAttributes > 0 && NumNamedAttribChannelsFound > 0) //-V1051
 	{
 		UE_LOG(LogNiagara, Warning, TEXT("%s: Unnamed attributes should not be used with named.  This is a deprecated workflow"), *Proxy->SourceDIName.ToString());
 	}
diff --git a/Engine/Plugins/Mutable/Source/CustomizableObject/Private/MuCO/UnrealMutableModelDiskStreamer.cpp b/Engine/Plugins/Mutable/Source/CustomizableObject/Private/MuCO/UnrealMutableModelDiskStreamer.cpp
--- a/Engine/Plugins/Mutable/Source/CustomizableObject/Private/MuCO/UnrealMutableModelDiskStreamer.cpp
+++ b/Engine/Plugins/Mutable/Source/CustomizableObject/Private/MuCO/UnrealMutableModelDiskStreamer.cpp
@@ -96,7 +96,7 @@ bool FUnrealMutableModelBulkReader::PrepareStreamingForObject(UCustomizableObjec
 		const TSharedPtr<FModelStreamableBulkData> ModelStreamableBulkData = CustomizableObject->GetPrivate()->GetModelStreamableBulkData();
 		NewData.ModelStreamableBulkData = TWeakPtr<FModelStreamableBulkData>(ModelStreamableBulkData);
 
-		if (!ModelStreamableBulkData || ModelStreamableBulkData->ModelStreamables.IsEmpty())
+		if (!ModelStreamableBulkData || ModelStreamableBulkData->ModelStreamables.IsEmpty()) //-V1051
 		{
 			UE_LOG(LogMutable, Warning, TEXT("Streaming: Customizable Object %s has no data to stream."), *CustomizableObject->GetName());
 
diff --git a/Engine/Plugins/Mutable/Source/CustomizableObjectEditor/Private/MuCOE/SMutableConstantsWidget.cpp b/Engine/Plugins/Mutable/Source/CustomizableObjectEditor/Private/MuCOE/SMutableConstantsWidget.cpp
--- a/Engine/Plugins/Mutable/Source/CustomizableObjectEditor/Private/MuCOE/SMutableConstantsWidget.cpp
+++ b/Engine/Plugins/Mutable/Source/CustomizableObjectEditor/Private/MuCOE/SMutableConstantsWidget.cpp
@@ -818,13 +818,13 @@ void SMutableConstantsWidget::Construct(const FArguments& InArgs,const mu::FProg
 void SMutableConstantsWidget::SetProgram(const mu::FProgram* InProgram)
 {
 	// Make sure we do not process the setting of the same program object as the one already set
-	if (this->MutableProgramPtr == InProgram)
+	if (MutableProgramPtr == InProgram)
 	{
 		return;
 	}
 	
 	// Set only once the program that is being used. No further updates should be required
-	this->MutableProgramPtr = InProgram;
+	MutableProgramPtr = InProgram;
 	
 	// Generate the backend for the lists used in this object
 	if (MutableProgramPtr != nullptr)
diff --git a/Engine/Plugins/VirtualProduction/ClonerEffector/Source/ClonerEffectorMeshBuilder/Private/CEMeshBuilder.cpp b/Engine/Plugins/VirtualProduction/ClonerEffector/Source/ClonerEffectorMeshBuilder/Private/CEMeshBuilder.cpp
--- a/Engine/Plugins/VirtualProduction/ClonerEffector/Source/ClonerEffectorMeshBuilder/Private/CEMeshBuilder.cpp
+++ b/Engine/Plugins/VirtualProduction/ClonerEffector/Source/ClonerEffectorMeshBuilder/Private/CEMeshBuilder.cpp
@@ -720,7 +720,7 @@ bool FCEMeshBuilder::AppendComponent(UNiagaraComponent* InComponent, const FTran
 					}
 
 					FNiagaraMeshRendererMeshProperties& Mesh = MeshRenderer->Meshes[ParticleMeshIndex];
-					if (!IsValid(Mesh.Mesh))
+					if (!IsValid(Mesh.Mesh)) //-V1051
 					{
 						continue;
 					}
diff --git a/Engine/Source/Editor/SVGDistanceField/Private/SvgDistanceFieldGenerate.cpp b/Engine/Source/Editor/SVGDistanceField/Private/SvgDistanceFieldGenerate.cpp
--- a/Engine/Source/Editor/SVGDistanceField/Private/SvgDistanceFieldGenerate.cpp
+++ b/Engine/Source/Editor/SVGDistanceField/Private/SvgDistanceFieldGenerate.cpp
@@ -55,7 +55,7 @@ bool SvgDistanceFieldGenerate(TArrayView64<const char> InSvgData, const FSvgDist
 		(OuterUnitSpread + InnerUnitSpread > 0.f || OuterPxSpread + InnerPxSpread > 0.f) &&
 		InConfiguration.OutputWidth > 0 &&
 		InConfiguration.OutputHeight > 0 &&
-		InConfiguration.MiterLimit >= 0.f
+		InConfiguration.MiterLimit >= 0.f //-V1051
 	)) {
 		return false;
 	}
diff --git a/Engine/Source/Editor/Sequencer/Public/KeyPropertyParams.h b/Engine/Source/Editor/Sequencer/Public/KeyPropertyParams.h
--- a/Engine/Source/Editor/Sequencer/Public/KeyPropertyParams.h
+++ b/Engine/Source/Editor/Sequencer/Public/KeyPropertyParams.h
@@ -107,7 +107,7 @@ public:
 					// Sometimes property paths have the array property twice, first with no array index,
 					// then a second so we skip over this property if that's the case
 					int32 ArrayIndex = FMath::Max(0, PropertyInfo.ArrayIndex);
-					if (PropertyInfo.ArrayIndex == INDEX_NONE && i < PropertyPath.GetNumProperties()-1)
+					if (PropertyInfo.ArrayIndex == INDEX_NONE && i < PropertyPath.GetNumProperties()-1) //-V1051
 					{
 						const FPropertyInfo& InnerPropertyInfo = PropertyPath.GetPropertyInfo(i+1);
 						const FProperty* InnerProperty = InnerPropertyInfo.Property.Get();
diff --git a/Engine/Source/Runtime/CinematicCamera/Private/CineCameraComponent.cpp b/Engine/Source/Runtime/CinematicCamera/Private/CineCameraComponent.cpp
--- a/Engine/Source/Runtime/CinematicCamera/Private/CineCameraComponent.cpp
+++ b/Engine/Source/Runtime/CinematicCamera/Private/CineCameraComponent.cpp
@@ -533,11 +533,7 @@ void UCineCameraComponent::RecalcDerivedData()
 
 	FieldOfView = GetHorizontalFieldOfViewInternal(/*bIncludeOverscan =*/ false);
 	Filmback.RecalcSensorAspectRatio();
-	AspectRatio = Filmback.SensorAspectRatio * LensSettings.SqueezeFactor;
-	if (CropSettings.AspectRatio > 0.0f)
-	{
-		AspectRatio = CropSettings.AspectRatio;
-	}
+	AspectRatio = (CropSettings.AspectRatio > 0.0f ? CropSettings.AspectRatio : Filmback.SensorAspectRatio * LensSettings.SqueezeFactor);
 
 #if WITH_EDITORONLY_DATA
 	CurrentHorizontalFOV = FieldOfView;			// informational variable only, for editor users
diff --git a/Engine/Source/Runtime/Engine/Private/Components/PrimitiveComponent.cpp b/Engine/Source/Runtime/Engine/Private/Components/PrimitiveComponent.cpp
--- a/Engine/Source/Runtime/Engine/Private/Components/PrimitiveComponent.cpp
+++ b/Engine/Source/Runtime/Engine/Private/Components/PrimitiveComponent.cpp
@@ -5370,7 +5370,7 @@ void UPrimitiveComponent::AssignSceneProxy(FPrimitiveSceneProxy* InSceneProxy)
 	check(SceneProxy == nullptr && SceneData.SceneProxy == nullptr);
 	SceneProxy = InSceneProxy;
 	SceneData.SceneProxy = InSceneProxy;
-	if (SceneProxy)
+	if (SceneData.SceneProxy)
 	{
 		SceneData.bAlwaysVisible = SceneProxy->IsAlwaysVisible();
 		SceneData.OwnerLastRenderTimePtr = FActorLastRenderTime::GetPtr(GetOwner());
diff --git a/Engine/Source/Runtime/Landscape/Private/LandscapeEditLayers.cpp b/Engine/Source/Runtime/Landscape/Private/LandscapeEditLayers.cpp
--- a/Engine/Source/Runtime/Landscape/Private/LandscapeEditLayers.cpp
+++ b/Engine/Source/Runtime/Landscape/Private/LandscapeEditLayers.cpp
@@ -9222,7 +9222,7 @@ void ALandscape::PrepareComponentDataToPackMaterialLayersCS(int32 InCurrentWeigh
 							{
 								const FLandscapeInfoLayerSettings& LayerInfo = Info->Layers[LayerInfoSettingsIndex];
 
-								if (ChannelLayerAllocation.LayerInfo == LayerInfo.LayerInfoObj)
+								if (ChannelLayerAllocation.LayerInfo == LayerInfo.LayerInfoObj) //-V1051
 								{
 									Data.SourcePaintLayerIndex[WeightmapChannelIndex] = LayerInfoSettingsIndex + 1; // due to visibility layer that is at 0
 									break;
diff --git a/Engine/Source/Runtime/RenderCore/Private/ShaderPipelineCache.cpp b/Engine/Source/Runtime/RenderCore/Private/ShaderPipelineCache.cpp
--- a/Engine/Source/Runtime/RenderCore/Private/ShaderPipelineCache.cpp
+++ b/Engine/Source/Runtime/RenderCore/Private/ShaderPipelineCache.cpp
@@ -880,10 +880,10 @@ bool FShaderPipelineCache::OpenUserPipelineFileCache(EShaderPlatform Platform)
 		FScopeLock Lock(&ShaderPipelineCache->Mutex);
 		// Just use the project for the user cache name.
 		FString UserCacheName = UE::ShaderPipeline::GetShaderPipelineBaseName(true);
-		ShaderPipelineCache->UserCacheTaskKey = UserCacheName + TEXT("_usr");
+		UserCacheTaskKey = UserCacheName + TEXT("_usr");
 		if (UE::ShaderPipeline::ShouldLoadUserCache() && !ShaderPipelineCache->GetTask(UserCacheTaskKey))
 		{
-			bFileOpen = ShaderPipelineCache->OpenPipelineFileCacheInternal(true, ShaderPipelineCache->UserCacheTaskKey, UserCacheName, Platform);
+			bFileOpen = ShaderPipelineCache->OpenPipelineFileCacheInternal(true, UserCacheTaskKey, UserCacheName, Platform);
 		}
 	}
 	return bFileOpen;
