Added option to force multipass (no FBF) renderng on Vulkan and iOS (r.Mobile.AllowFramebufferFetch)
Skip decals render-pass if there is nothing to draw

diff --git a/Engine/Source/Runtime/Core/Private/HAL/ConsoleManager.cpp b/Engine/Source/Runtime/Core/Private/HAL/ConsoleManager.cpp
--- a/Engine/Source/Runtime/Core/Private/HAL/ConsoleManager.cpp
+++ b/Engine/Source/Runtime/Core/Private/HAL/ConsoleManager.cpp
@@ -3704,6 +3704,13 @@ static TAutoConsoleVariable<int32> CVarMobileAllowFramebufferFetchOpenGL(
 		"1: Allow use of framebuffer fetch and PLS extensions (default)"),
 	ECVF_RenderThreadSafe | ECVF_ReadOnly);
 
+static TAutoConsoleVariable<int32> CVarMobileAllowFramebufferFetch(
+	TEXT("r.Mobile.AllowFramebufferFetch"),
+	1,
+	TEXT("0: Multi-pass rendering without using sub-passes in Vulkan or programmable blending on iOS. For OpenGL use r.Mobile.AllowFramebufferFetchOpenGL\n"
+		 "1: Allow use of sub-passes and programmable blending (default)"),
+	ECVF_RenderThreadSafe | ECVF_ReadOnly);
+
 static TAutoConsoleVariable<int32> CVarMobileEnableStaticAndCSMShadowReceivers(
 	TEXT("r.Mobile.EnableStaticAndCSMShadowReceivers"),
 	1,
diff --git a/Engine/Source/Runtime/RenderCore/Private/RenderUtils.cpp b/Engine/Source/Runtime/RenderCore/Private/RenderUtils.cpp
--- a/Engine/Source/Runtime/RenderCore/Private/RenderUtils.cpp
+++ b/Engine/Source/Runtime/RenderCore/Private/RenderUtils.cpp
@@ -463,11 +463,25 @@ RENDERCORE_API bool PlatformGPUSceneUsesUniformBufferView(const FStaticShaderPla
 	return IsMobilePlatform(Platform) && FDataDrivenShaderPlatformInfo::GetSupportsUniformBufferObjects(Platform);
 }
 
+RENDERCORE_API bool MobileAllowFramebufferFetch(const FStaticShaderPlatform Platform)
+{
+	if (Platform == SP_OPENGL_ES3_1_ANDROID)
+	{
+		const static IConsoleVariable* CVarAllowFramebufferFetchOpenGL = IConsoleManager::Get().FindConsoleVariable(TEXT("r.Mobile.AllowFramebufferFetchOpenGL"));
+		return CVarAllowFramebufferFetchOpenGL->GetBool();
+	}
+	else
+	{
+		static FShaderPlatformCachedIniValue<bool> MobileAllowFramebufferFetchIniValue(TEXT("r.Mobile.AllowFramebufferFetch"));
+		return MobileAllowFramebufferFetchIniValue.Get(Platform);
+	}
+}
+
 RENDERCORE_API bool MobileRequiresSceneDepthAux(const FStaticShaderPlatform Platform)
 {
 	const bool bMobileHDR = IsMobileHDR();
 
-	if (!MobileUsesFullDepthPrepass(Platform))
+	if (MobileAllowFramebufferFetch(Platform) && !MobileUsesFullDepthPrepass(Platform))
 	{
 		// SceneDepthAux is used on tiled GPUs (iOS, Android) with forward shading and on iOS only with deferred
 		if (IsMetalMobilePlatform(Platform))
@@ -483,16 +497,6 @@ RENDERCORE_API bool MobileRequiresSceneDepthAux(const FStaticShaderPlatform Plat
 	return false;
 }
 
-RENDERCORE_API bool MobileAllowFramebufferFetch(const FStaticShaderPlatform Platform)
-{
-	if (Platform == SP_OPENGL_ES3_1_ANDROID)
-	{
-		const static IConsoleVariable* CVarAllowFramebufferFetchOpenGL = IConsoleManager::Get().FindConsoleVariable(TEXT("r.Mobile.AllowFramebufferFetchOpenGL"));
-		return CVarAllowFramebufferFetchOpenGL->GetBool();
-	}
-	return true;
-}
-
 RENDERCORE_API bool SupportsTextureCubeArray(ERHIFeatureLevel::Type FeatureLevel)
 {
 	return FeatureLevel >= ERHIFeatureLevel::SM5 
diff --git a/Engine/Source/Runtime/Renderer/Private/MobileShadingRenderer.cpp b/Engine/Source/Runtime/Renderer/Private/MobileShadingRenderer.cpp
--- a/Engine/Source/Runtime/Renderer/Private/MobileShadingRenderer.cpp
+++ b/Engine/Source/Runtime/Renderer/Private/MobileShadingRenderer.cpp
@@ -2380,6 +2380,13 @@ void FMobileSceneRenderer::RenderDeferredMultiPass(FRDGBuilder& GraphBuilder, FS
 			SCOPED_GPU_MASK(GraphBuilder.RHICmdList, !View.IsInstancedStereoPass() ? View.GPUMask : (View.GPUMask | View.GetInstancedView()->GPUMask));
 			SCOPED_CONDITIONAL_DRAW_EVENTF(GraphBuilder.RHICmdList, EventView, RenderViews.Num() > 1, TEXT("View%d"), ViewContext.ViewIndex);
 
+			const EMeshPass::Type MeshDecalsPass = EMeshPass::MeshDecal_SceneColorAndGBuffer;
+			if (Scene->Decals.Num() == 0 && !HasAnyDraw(View.ParallelMeshDrawCommandPasses[MeshDecalsPass]))
+			{
+				// Skip decal render-pass if there is nothing to draw
+				continue;
+			}
+
 			View.BeginRenderView();
 
 			auto PassParameters = GraphBuilder.AllocParameters<FMobileRenderPassParameters>();
@@ -2389,7 +2396,7 @@ void FMobileSceneRenderer::RenderDeferredMultiPass(FRDGBuilder& GraphBuilder, FS
 
 			if (Scene->GPUScene.IsEnabled())
 			{
-				BuildMeshRenderingCommands(GraphBuilder, EMeshPass::MeshDecal_SceneColorAndGBuffer, View, Scene->GPUScene, InstanceCullingManager, PassParameters->InstanceCullingDrawParams);
+				BuildMeshRenderingCommands(GraphBuilder, MeshDecalsPass, View, Scene->GPUScene, InstanceCullingManager, PassParameters->InstanceCullingDrawParams);
 			}
 
 			GraphBuilder.AddPass(
@@ -2560,6 +2567,11 @@ int32 FMobileSceneRenderer::ComputeNumOcclusionQueriesToBatch() const
 // Whether we need a separate render-passes for translucency, decals etc
 bool FMobileSceneRenderer::RequiresMultiPass(int32 NumMSAASamples, EShaderPlatform ShaderPlatform)
 {
+	if (!MobileAllowFramebufferFetch(ShaderPlatform))
+	{
+		return true;
+	}
+	
 	// Vulkan uses subpasses
 	if (IsVulkanPlatform(ShaderPlatform))
 	{
