[ShaderCoreLibrary]
- Fix potential deadlock if multiple threads race to add the same resource

diff --git a/Engine/Source/Runtime/RenderCore/Private/ShaderCodeLibrary.cpp b/Engine/Source/Runtime/RenderCore/Private/ShaderCodeLibrary.cpp
--- a/Engine/Source/Runtime/RenderCore/Private/ShaderCodeLibrary.cpp
+++ b/Engine/Source/Runtime/RenderCore/Private/ShaderCodeLibrary.cpp
@@ -1112,17 +1112,17 @@ public:
 		return Resources[ShaderMapIndex];
 	}
 
-	TRefCountPtr<FShaderMapResource_SharedCode> AddOrDeleteResource(FShaderMapResource_SharedCode* Resource, FArchive* Ar)
+	TRefCountPtr<FShaderMapResource_SharedCode> AddResource(int32 ShaderMapIndex, FArchive* Ar)
 	{
-		const int32 ShaderMapIndex = Resource->ShaderMapIndex;
-		TRefCountPtr<FShaderMapResource_SharedCode> OutResource(Resource);
+		TRefCountPtr<FShaderMapResource_SharedCode> OutResource;
 		bool bPreload = false;
 		{
 			FRWScopeLock Locker(ResourceLock, SLT_Write);
 			FShaderMapResource_SharedCode* PrevResource = Resources[ShaderMapIndex];
 			if (!PrevResource)
 			{
-				Resources[ShaderMapIndex] = Resource;
+				Resources[ShaderMapIndex] = new FShaderMapResource_SharedCode(this, ShaderMapIndex);
+				OutResource = Resources[ShaderMapIndex];
 				bPreload = !GRHILazyShaderCodeLoading && GPreloadShaderMaps;
 			}
 			else
@@ -1135,7 +1135,7 @@ public:
 		{
 			SCOPED_LOADTIMER(FShaderLibraryInstance_PreloadShaderMap);
 			FGraphEventArray PreloadCompletionEvents;
-			Resource->bEntireShaderMapPreloaded = Library->PreloadShaderMap(ShaderMapIndex, PreloadCompletionEvents);
+			OutResource->bEntireShaderMapPreloaded = Library->PreloadShaderMap(ShaderMapIndex, PreloadCompletionEvents);
 			if (Ar && PreloadCompletionEvents.Num() > 0)
 			{
 				FExternalReadCallback ExternalReadCallback = [this, PreloadCompletionEvents = MoveTemp(PreloadCompletionEvents)](double ReaminingTime)
@@ -3381,8 +3381,8 @@ public:
 			TRefCountPtr<FShaderMapResource_SharedCode> Resource = LibraryInstance->GetResource(ShaderMapIndex);
 			if (!Resource)
 			{
-				SCOPED_LOADTIMER(LoadShaderResource_AddOrDeleteResource);
-				Resource = LibraryInstance->AddOrDeleteResource(new FShaderMapResource_SharedCode(LibraryInstance, ShaderMapIndex), Ar);
+				SCOPED_LOADTIMER(LoadShaderResource_AddResource);
+				Resource = LibraryInstance->AddResource(ShaderMapIndex, Ar);
 			}
 
 			return Resource;
