Added scoped debug names for render resource destruction faults

diff --git a/Engine/Source/Runtime/Engine/Private/SkeletalMesh.cpp b/Engine/Source/Runtime/Engine/Private/SkeletalMesh.cpp
--- a/Engine/Source/Runtime/Engine/Private/SkeletalMesh.cpp
+++ b/Engine/Source/Runtime/Engine/Private/SkeletalMesh.cpp
@@ -478,6 +478,12 @@ void USkeletalMesh::SetSkeletalMeshRenderData(TUniquePtr<FSkeletalMeshRenderData
 	WaitUntilAsyncPropertyReleased(ESkeletalMeshAsyncProperties::SkeletalMeshRenderData);
 	PRAGMA_DISABLE_DEPRECATION_WARNINGS
 	SkeletalMeshRenderData = MoveTemp(InSkeletalMeshRenderData);
+	
+	if (FSkeletalMeshRenderData* RenderData = SkeletalMeshRenderData.Get())
+	{
+		RenderData->OwnerName = GetFName();
+	}
+	
 	PRAGMA_ENABLE_DEPRECATION_WARNINGS
 }
 
diff --git a/Engine/Source/Runtime/Engine/Private/SkeletalMeshRenderData.cpp b/Engine/Source/Runtime/Engine/Private/SkeletalMeshRenderData.cpp
--- a/Engine/Source/Runtime/Engine/Private/SkeletalMeshRenderData.cpp
+++ b/Engine/Source/Runtime/Engine/Private/SkeletalMeshRenderData.cpp
@@ -475,6 +475,8 @@ FSkeletalMeshRenderData::FSkeletalMeshRenderData()
 
 FSkeletalMeshRenderData::~FSkeletalMeshRenderData()
 {
+	FRenderResourceNameScope FatalScope(OwnerName);
+	
 	FSkeletalMeshLODRenderData** LODRenderDataArray = LODRenderData.GetData();
 	for (int32 LODIndex = 0; LODIndex < LODRenderData.Num(); ++LODIndex)
 	{
diff --git a/Engine/Source/Runtime/Engine/Public/Rendering/SkeletalMeshRenderData.h b/Engine/Source/Runtime/Engine/Public/Rendering/SkeletalMeshRenderData.h
--- a/Engine/Source/Runtime/Engine/Public/Rendering/SkeletalMeshRenderData.h
+++ b/Engine/Source/Runtime/Engine/Public/Rendering/SkeletalMeshRenderData.h
@@ -43,6 +43,9 @@ public:
 	/** Whether ray tracing acceleration structures should be created for this mesh. Derived from owner USkinnedAsset. */
 	bool bSupportRayTracing;
 
+	/** Optional, name of the owning skeletal mesh */
+	FName OwnerName = NAME_None;
+
 #if RHI_RAYTRACING
 	RayTracing::FGeometryGroupHandle RayTracingGeometryGroupHandle = INDEX_NONE;
 #endif
diff --git a/Engine/Source/Runtime/RenderCore/Private/RenderResource.cpp b/Engine/Source/Runtime/RenderCore/Private/RenderResource.cpp
--- a/Engine/Source/Runtime/RenderCore/Private/RenderResource.cpp
+++ b/Engine/Source/Runtime/RenderCore/Private/RenderResource.cpp
@@ -18,6 +18,7 @@
 /** Whether to enable mip-level fading or not: +1.0f if enabled, -1.0f if disabled. */
 float GEnableMipLevelFading = 1.0f;
 
+static thread_local FName GRenderResourceScopeName = NAME_None;
 
 class FRenderResourceList
 {
@@ -162,6 +163,13 @@ void FRenderResource::ChangeFeatureLevel(ERHIFeatureLevel::Type NewFeatureLevel)
 	});
 }
 
+FName FRenderResource::SetScopeName(FName Name)
+{
+	FName Old = GRenderResourceScopeName;
+	GRenderResourceScopeName = Name;
+	return Old;
+}
+
 FRHICommandListBase& FRenderResource::GetImmediateCommandList()
 {
 	return FRHICommandListImmediate::Get();
@@ -246,7 +254,7 @@ FRenderResource::~FRenderResource()
 	if (IsInitialized() && !GIsCriticalError)
 	{
 		// Deleting an initialized FRenderResource will result in a crash later since it is still linked
-		UE_LOG(LogRendererCore, Fatal, TEXT("A FRenderResource was deleted without being released first! [Resource=%s, Owner=%s]"), *GetResourceName().ToString(), *GetOwnerName().ToString());
+		UE_LOG(LogRendererCore, Fatal, TEXT("A FRenderResource was deleted without being released first! [Resource=%s, Owner=%s, Scope=%s]"), *GetResourceName().ToString(), *GetOwnerName().ToString(), *GRenderResourceScopeName.ToString());
 	}
 }
 
diff --git a/Engine/Source/Runtime/RenderCore/Public/RenderResource.h b/Engine/Source/Runtime/RenderCore/Public/RenderResource.h
--- a/Engine/Source/Runtime/RenderCore/Public/RenderResource.h
+++ b/Engine/Source/Runtime/RenderCore/Public/RenderResource.h
@@ -55,6 +55,9 @@ public:
 	/** Reinitializes render resources at a new feature level. */
 	static RENDERCORE_API void ChangeFeatureLevel(ERHIFeatureLevel::Type NewFeatureLevel);
 
+	/** Set the scoped (thread local) resource name, used for error reporting */
+	static RENDERCORE_API FName SetScopeName(FName Name);
+
 	////////////////////////////////////////////////////////////////////////////////////
 
 	/** Default constructor. */
@@ -596,3 +599,20 @@ private:
 		((ResourceType*)this)->ReleaseResource();
 	}
 };
+
+/** Helper for scoped render resource names */
+struct FRenderResourceNameScope
+{
+	FRenderResourceNameScope(FName Name)
+	{
+		PreviousName = FRenderResource::SetScopeName(Name);
+	}
+
+	~FRenderResourceNameScope()
+	{
+		FRenderResource::SetScopeName(PreviousName);
+	}
+
+private:
+	FName PreviousName;
+};
