Rename r.Substrate.BlendableGBuffer into r.Substrate.GBufferFormat and add r.Substrate.PlatformBlendableGBuffer.

The former allows to configure the desired format for the project, while the latter specifies what a particular platform uses/supports as GBuffer format.

diff --git a/Engine/Config/Linux/BaseLinuxEngine.ini b/Engine/Config/Linux/BaseLinuxEngine.ini
--- a/Engine/Config/Linux/BaseLinuxEngine.ini
+++ b/Engine/Config/Linux/BaseLinuxEngine.ini
@@ -21,4 +21,4 @@ r.Substrate.Glints=1
 r.Substrate.ClosuresPerPixel=4
 r.Substrate.OpaqueMaterialRoughRefraction=0
 r.Refraction.Blur=1
-r.Substrate.BlendableGBuffer=0
+r.Substrate.PlatformGBufferFormat=1
diff --git a/Engine/Config/LinuxArm64/BaseLinuxArm64Engine.ini b/Engine/Config/LinuxArm64/BaseLinuxArm64Engine.ini
--- a/Engine/Config/LinuxArm64/BaseLinuxArm64Engine.ini
+++ b/Engine/Config/LinuxArm64/BaseLinuxArm64Engine.ini
@@ -11,4 +11,4 @@ r.Substrate.Glints=1
 r.Substrate.ClosuresPerPixel=4
 r.Substrate.OpaqueMaterialRoughRefraction=0
 r.Refraction.Blur=1
-r.Substrate.BlendableGBuffer=0
+r.Substrate.PlatformGBufferFormat=1
diff --git a/Engine/Config/Windows/BaseWindowsEngine.ini b/Engine/Config/Windows/BaseWindowsEngine.ini
--- a/Engine/Config/Windows/BaseWindowsEngine.ini
+++ b/Engine/Config/Windows/BaseWindowsEngine.ini
@@ -42,7 +42,7 @@ r.Substrate.Glints=1
 r.Substrate.ClosuresPerPixel=4
 r.Substrate.OpaqueMaterialRoughRefraction=0
 r.Refraction.Blur=1
-r.Substrate.BlendableGBuffer=0
+r.Substrate.PlatformGBufferFormat=1
 
 [OnlineSubsystem]
 LocalPlatformName=WIN
diff --git a/Engine/Source/Runtime/Engine/Classes/Engine/EngineTypes.h b/Engine/Source/Runtime/Engine/Classes/Engine/EngineTypes.h
--- a/Engine/Source/Runtime/Engine/Classes/Engine/EngineTypes.h
+++ b/Engine/Source/Runtime/Engine/Classes/Engine/EngineTypes.h
@@ -765,6 +765,22 @@ private:
 	uint16 ShadingModelField = 0;
 };
 
+/** 
+* Enumerates available Substrate storage format
+*/
+UENUM()
+namespace ESubstrateStorageFormat
+{
+	enum Type : int
+	{
+		/** Blendable GBuffer - Fast but with limited lighting complexity. If selected, all platforms will be forced to use this format. */
+		BlendableGBuffer = 0 UMETA(DisplayName = "Blendable GBuffer"),
+
+		/** Adaptive GBuffer - More heavy but enable complex lighting behavior per pixel. If selected, only compatible platforms will use this format. Uncompatible platforms will use Blendable GBuffer format. */
+		AdaptiveBuffer = 1 UMETA(DisplayName = "Adaptive GBuffer"),
+	};
+}
+
 /**
  * Specifies the Substrate runtime shading model summarized from the material graph
  * Not exposed in UI, only used in code.
diff --git a/Engine/Source/Runtime/Engine/Classes/Engine/RendererSettings.h b/Engine/Source/Runtime/Engine/Classes/Engine/RendererSettings.h
--- a/Engine/Source/Runtime/Engine/Classes/Engine/RendererSettings.h
+++ b/Engine/Source/Runtime/Engine/Classes/Engine/RendererSettings.h
@@ -1213,6 +1213,16 @@ class URendererSettings : public UDeveloperSettings
 		ConfigRestartRequired = true))
 		uint32 bEnableSubstrate : 1;
 
+	/**
+	"Enable Substrate storage format."
+	*/
+	UPROPERTY(config, EditAnywhere, Category = Substrate, meta = (
+		ConsoleVariable = "r.Substrate.ProjectGBufferFormat",
+		DisplayName = "Substrate GBuffer Format",
+		ToolTip = "Select the GBuffer format used for this project. Changing this setting requires restarting the editor.",
+		ConfigRestartRequired = true))
+	TEnumAsByte<ESubstrateStorageFormat::Type> SubstrateGBufferFormat;
+
 	/**
 	"Enable Substrate opaque material rough refractions effect from top layers over layers below."
 	*/
diff --git a/Engine/Source/Runtime/RenderCore/Private/RenderUtils.cpp b/Engine/Source/Runtime/RenderCore/Private/RenderUtils.cpp
--- a/Engine/Source/Runtime/RenderCore/Private/RenderUtils.cpp
+++ b/Engine/Source/Runtime/RenderCore/Private/RenderUtils.cpp
@@ -1906,10 +1906,20 @@ static TAutoConsoleVariable<int32> CVarSubstrate(
 	TEXT("Enable Substrate materials (Beta)."),
 	ECVF_ReadOnly | ECVF_RenderThreadSafe);
 
-static TAutoConsoleVariable<int32> CVarSubstrateBlendableGBuffer(
-	TEXT("r.Substrate.BlendableGBuffer"),
-	1, // Substrate 60fps by default
-	TEXT("Forces Substrate materials to fit into a blendable GBuffer. This is faster and DBuffer decals are no longer required. However, lots of visual effects and fidelity are lost. To be able to cast to blendable GBuffer, this command enforce ClosuresPerPixel=1."),
+static TAutoConsoleVariable<int32> CVarSubstratePlatformGBufferFormat(
+	TEXT("r.Substrate.PlatformGBufferFormat"),
+	0, // Substrate 60fps by default -> Blendable GBuffer 
+	TEXT("Define GBuffer format supported by this platforms.\n \
+		 - 0: Substrate materials are fit into a blendable GBuffer. This is faster and DBuffer decals are no longer required. However, lots of visual effects and fidelity are lost. To be able to cast to blendable GBuffer, this command enforce ClosuresPerPixel=1.\n \
+		 - 1: Substrate materials are fit into a bitstream GBuffer. This stores more data to enable complex lighting behavior per pixel. It is more costly than the blendable GBuffer format. This is only support on certain platforms. Non-supported platforms will automatically fallback onto a blendable GBuffer format."),
+	ECVF_ReadOnly | ECVF_RenderThreadSafe | ECVF_Preview);
+
+static TAutoConsoleVariable<int32> CVarSubstrateProjectGBufferFormat(
+	TEXT("r.Substrate.ProjectGBufferFormat"),
+	1,
+	TEXT("Define GBuffer format used for this project.\n \
+	- 0: Substrate materials are fit into a blendable GBuffer. This is faster and DBuffer decals are no longer required. However, lots of visual effects and fidelity are lost. To be able to cast to blendable GBuffer, this command enforce ClosuresPerPixel=1.\n \
+	- 1: Substrate materials are fit into a bitstream GBuffer. This stores more data to enable complex lighting behavior per pixel. It is more costly than the blendable GBuffer format. This is only support on certain platforms. Non-supported platforms will automatically fallback onto a blendable GBuffer format."),
 	ECVF_ReadOnly | ECVF_RenderThreadSafe | ECVF_Preview);
 
 static TAutoConsoleVariable<int32> CVarSubstrateBytesPerPixel(
@@ -2024,9 +2034,12 @@ namespace Substrate
 
 	bool IsSubstrateBlendableGBufferEnabled(EShaderPlatform InPlatform)
 	{
-		static FShaderPlatformCachedIniValue<int32> CVarSubstrateBlendableGBufferPlatform(TEXT("r.Substrate.BlendableGBuffer"));
+		static FShaderPlatformCachedIniValue<int32> CVarPlatformGBufferFormat(TEXT("r.Substrate.PlatformGBufferFormat"));
+		static FShaderPlatformCachedIniValue<int32> CVarProjectGBufferFormat(TEXT("r.Substrate.ProjectGBufferFormat"));
+		const bool bUseMaterialBuffer = CVarPlatformGBufferFormat.Get(InPlatform) > 0 && CVarProjectGBufferFormat.Get(InPlatform) > 0;
+		
 		// Mobile always uses blendable GBuffer, otherwise it is enabled on selected platforms.
-		return (IsMobilePlatform(InPlatform) && IsMobileDeferredShadingEnabled(InPlatform)) || CVarSubstrateBlendableGBufferPlatform.Get(InPlatform) > 0;
+		return (IsMobilePlatform(InPlatform) && IsMobileDeferredShadingEnabled(InPlatform)) || !bUseMaterialBuffer;
 	}
 
 	bool IsStochasticLightingEnabled(EShaderPlatform InPlatform)
