[shaders] fix potential out of bounds read when remapping diagnostic messages; could occur if there were no digits before the end of the line in the diagnostic message after the __UE_FILENAME_SENTINEL marker

diff --git a/Engine/Source/Runtime/RenderCore/Private/ShaderPreprocessTypes.cpp b/Engine/Source/Runtime/RenderCore/Private/ShaderPreprocessTypes.cpp
--- a/Engine/Source/Runtime/RenderCore/Private/ShaderPreprocessTypes.cpp
+++ b/Engine/Source/Runtime/RenderCore/Private/ShaderPreprocessTypes.cpp
@@ -120,7 +120,7 @@ void FShaderDiagnosticRemapper::Remap(FShaderCompilerError& Diagnostic) const
 		// Parse the line number from the message and record start and end indices
 		// Code assumes that the next integral number found in the message string after the filename is always line number
 		int32 LineNumberStart = FilenameIndex + GetPreprocessConstants().FilenameSentinel.Len();
-		while (!FChar::IsDigit(OriginalMessage[LineNumberStart]))
+		while (LineNumberStart < OriginalMessage.Len() && !FChar::IsDigit(OriginalMessage[LineNumberStart]))
 		{
 			++LineNumberStart;
 		}
