Remove SubstrateGBufferFormat per-platform CVar fetch, as the value is constant across all platforms.

diff --git a/Engine/Source/Runtime/RenderCore/Private/RenderUtils.cpp b/Engine/Source/Runtime/RenderCore/Private/RenderUtils.cpp
--- a/Engine/Source/Runtime/RenderCore/Private/RenderUtils.cpp
+++ b/Engine/Source/Runtime/RenderCore/Private/RenderUtils.cpp
@@ -2035,8 +2035,7 @@ namespace Substrate
 	bool IsSubstrateBlendableGBufferEnabled(EShaderPlatform InPlatform)
 	{
 		static FShaderPlatformCachedIniValue<int32> CVarPlatformGBufferFormat(TEXT("r.Substrate.PlatformGBufferFormat"));
-		static FShaderPlatformCachedIniValue<int32> CVarProjectGBufferFormat(TEXT("r.Substrate.ProjectGBufferFormat"));
-		const bool bUseMaterialBuffer = CVarPlatformGBufferFormat.Get(InPlatform) > 0 && CVarProjectGBufferFormat.Get(InPlatform) > 0;
+		const bool bUseMaterialBuffer = CVarPlatformGBufferFormat.Get(InPlatform) > 0 && CVarSubstrateProjectGBufferFormat.GetValueOnAnyThread() > 0;
 		
 		// Mobile always uses blendable GBuffer, otherwise it is enabled on selected platforms.
 		return (IsMobilePlatform(InPlatform) && IsMobileDeferredShadingEnabled(InPlatform)) || !bUseMaterialBuffer;
