Fixed crash in RDG due to external access subresource being null at [0].

diff --git a/Engine/Source/Runtime/RenderCore/Private/RenderGraphBuilder.cpp b/Engine/Source/Runtime/RenderCore/Private/RenderGraphBuilder.cpp
--- a/Engine/Source/Runtime/RenderCore/Private/RenderGraphBuilder.cpp
+++ b/Engine/Source/Runtime/RenderCore/Private/RenderGraphBuilder.cpp
@@ -4040,8 +4040,13 @@ void FRDGBuilder::AddLastTextureTransition(FRDGTexture* Texture)
 
 	if (Texture->AccessModeState.ActiveMode == FRDGViewableResource::EAccessMode::External)
 	{
+		// Find the first subresource that is valid. There must be at least one.
+		int32  FirstValidIndex = 0;
+		for (; FirstValidIndex < Texture->State.Num() && !Texture->State[FirstValidIndex]; ++FirstValidIndex) {}
+		check( FirstValidIndex < Texture->State.Num() &&  Texture->State[FirstValidIndex]);
+
 		// Assign the final state that was enqueued by the external access pass, which may include merged states.
-		EpilogueResourceAccesses.Emplace(Texture->GetRHI(), Texture->State[0]->Access, Texture->State[0]->GetPipelines());
+		EpilogueResourceAccesses.Emplace(Texture->GetRHI(), Texture->State[FirstValidIndex]->Access, Texture->State[FirstValidIndex]->GetPipelines());
 		return;
 	}
 
