Undoing changes to AddClearRenderTargetPass from CL# 44513467 (causes crash in water rendering)

diff --git a/Engine/Source/Runtime/RenderCore/Private/RenderGraphUtils.cpp b/Engine/Source/Runtime/RenderCore/Private/RenderGraphUtils.cpp
--- a/Engine/Source/Runtime/RenderCore/Private/RenderGraphUtils.cpp
+++ b/Engine/Source/Runtime/RenderCore/Private/RenderGraphUtils.cpp
@@ -675,15 +675,7 @@ void AddClearRenderTargetPass(FRDGBuilder& GraphBuilder, FRDGTextureRef Texture,
 				uint8 CurrentMipIndex = IntCastChecked<uint8>(TextureClearInfo.FirstMipIndex + MipIndex);
 
 				FRenderTargetParameters* Parameters = GraphBuilder.AllocParameters<FRenderTargetParameters>();
-				if (TextureClearInfo.NumSlices == 1 && CurrentSliceIndex == 0)
-				{
-					// Use default value for a Slice, since it can be used for a pass merging
-					Parameters->RenderTargets[0] = FRenderTargetBinding(Texture, ERenderTargetLoadAction::EClear, CurrentMipIndex);
-				}
-				else
-				{
-					Parameters->RenderTargets[0] = FRenderTargetBinding(Texture, ERenderTargetLoadAction::EClear, CurrentMipIndex, CurrentSliceIndex);
-				}
+				Parameters->RenderTargets[0] = FRenderTargetBinding(Texture, ERenderTargetLoadAction::EClear, CurrentMipIndex, CurrentSliceIndex);
 
 				GraphBuilder.AddPass(
 					RDG_EVENT_NAME("ClearRenderTarget(%s, slice %d, mip %d) %dx%d ClearAction", Texture->Name, CurrentSliceIndex, CurrentMipIndex, Texture->Desc.Extent.X, Texture->Desc.Extent.Y),
