PR #13628: Shader loading: prevent dereferencing potential nullptr when file is not found

diff --git a/Engine/Source/Runtime/RenderCore/Private/ShaderCore.cpp b/Engine/Source/Runtime/RenderCore/Private/ShaderCore.cpp
--- a/Engine/Source/Runtime/RenderCore/Private/ShaderCore.cpp
+++ b/Engine/Source/Runtime/RenderCore/Private/ShaderCore.cpp
@@ -2776,7 +2776,7 @@ bool LoadShaderSourceFile(const TCHAR* InVirtualFilePath, EShaderPlatform Shader
 
 	const bool bResult = LoadCachedShaderSourceFile(InVirtualFilePath, ShaderPlatform, pFileContents, OutCompileErrors, ShaderPlatformName, OutStrippedContents);
 
-	if (OutFileContents)
+	if (OutFileContents && LocalFileContents.IsValid())
 	{
 		*OutFileContents = *LocalFileContents;
 	}
