[Shaders] Don't return success from LoadCachedShaderSourceFile() when the requested output content was not initialized.

The function FindOrTryProduceAndApply() for the shader cache can return true even if the apply-callback did not initialize the requested output content due to a missing shader file.
This can lead a null pointer exception in TryUpdateSingleShaderFilehash().

This changes the null pointer crash to a Fatal error from this:
appError called: Assertion failed: IsValid() [File:Engine/Source/Runtime/Core/Public/Templates/SharedPointer.h] [Line: 1101]

To this:
Fatal error: [File:Engine/Source/Runtime/RenderCore/Private/ShaderCore.cpp] [Line: 2791]
Couldn't find source file of virtual shader path '/Engine/Private/HeterogeneousVolumes/HeterogeneousVolumesVisualizeCascades.usf'

#rnx

diff --git a/Engine/Source/Runtime/RenderCore/Private/ShaderCore.cpp b/Engine/Source/Runtime/RenderCore/Private/ShaderCore.cpp
--- a/Engine/Source/Runtime/RenderCore/Private/ShaderCore.cpp
+++ b/Engine/Source/Runtime/RenderCore/Private/ShaderCore.cpp
@@ -2700,7 +2700,6 @@ bool LoadCachedShaderSourceFile(const TCHAR* InVirtualFilePath, EShaderPlatform
 	}
 
 	bool bResult = false;
-	bool bFoundValidFile = false;
 
 	STAT(double ShaderFileLoadingTime = 0);
 
@@ -2710,35 +2709,36 @@ bool LoadCachedShaderSourceFile(const TCHAR* InVirtualFilePath, EShaderPlatform
 		FString VirtualFilePath(InVirtualFilePath);
 		FixupShaderFilePath(VirtualFilePath, ShaderPlatform, ShaderPlatformName);
 
-		bResult = 
-			GShaderFileCache.FindOrTryProduceAndApply(VirtualFilePath,
-				[&VirtualFilePath, OutCompileErrors](FShaderFileCacheEntry& CachedFile)
+		GShaderFileCache.FindOrTryProduceAndApply(VirtualFilePath,
+			[&VirtualFilePath, OutCompileErrors](FShaderFileCacheEntry& CachedFile)
+			{
+				FString ShaderFilePath = GetShaderSourceFilePath(VirtualFilePath, OutCompileErrors);
+
+				if (!ShaderFilePath.IsEmpty())
 				{
-					FString ShaderFilePath = GetShaderSourceFilePath(VirtualFilePath, OutCompileErrors);
+					FShaderSharedStringPtr SourceFile = MakeShared<FString>();
 
-					if (!ShaderFilePath.IsEmpty())
+					// verify SHA hash of shader files on load. missing entries trigger an error
+					if (FFileHelper::LoadFileToString(*SourceFile, *ShaderFilePath, FFileHelper::EHashOptions::EnableVerify | FFileHelper::EHashOptions::ErrorMissingHash))
 					{
-						FShaderSharedStringPtr SourceFile = MakeShared<FString>();
-
-						// verify SHA hash of shader files on load. missing entries trigger an error
-						if (FFileHelper::LoadFileToString(*SourceFile, *ShaderFilePath, FFileHelper::EHashOptions::EnableVerify | FFileHelper::EHashOptions::ErrorMissingHash))
-						{
-							CachedFile.Source = MoveTemp(SourceFile);
+						CachedFile.Source = MoveTemp(SourceFile);
 
-							TArray<ANSICHAR>* StrippedSource = new TArray<ANSICHAR>;
-							ShaderConvertAndStripComments(*CachedFile.Source, *StrippedSource);
-							CachedFile.StrippedSource = MakeShareable(StrippedSource);
+						TArray<ANSICHAR>* StrippedSource = new TArray<ANSICHAR>;
+						ShaderConvertAndStripComments(*CachedFile.Source, *StrippedSource);
+						CachedFile.StrippedSource = MakeShareable(StrippedSource);
 
-							return true;
-						}
+						return true;
 					}
+				}
 
-					// Create an empty entry if missing files are being cached.
-					return FMissingShaderFileCacheGuard::IsEnabled();
-				},
-				[&bFoundValidFile, OutFileContents, OutStrippedContents](const FShaderFileCacheEntry& CachedFile)
+				// Create an empty entry if missing files are being cached.
+				return FMissingShaderFileCacheGuard::IsEnabled();
+			},
+			[&bResult, OutFileContents, OutStrippedContents](const FShaderFileCacheEntry& CachedFile)
+			{
+				if (CachedFile.Source.IsValid())
 				{
-					if (OutFileContents && CachedFile.Source.IsValid())
+					if (OutFileContents)
 					{
 						*OutFileContents = CachedFile.Source;
 					}
@@ -2746,18 +2746,10 @@ bool LoadCachedShaderSourceFile(const TCHAR* InVirtualFilePath, EShaderPlatform
 					{
 						*OutStrippedContents = CachedFile.StrippedSource;
 					}
-					if (FMissingShaderFileCacheGuard::IsEnabled())
-					{
-						bFoundValidFile = !CachedFile.IsEmpty();
-					}
+					bResult = true;
 				}
-			);
-	}
-
-	// When caching missing files, the map find will always return true, so the result is based on whether it is an empty entry or not.
-	if (FMissingShaderFileCacheGuard::IsEnabled())
-	{
-		bResult = bFoundValidFile;
+			}
+		);
 	}
 
 	INC_FLOAT_STAT_BY(STAT_ShaderCompiling_LoadingShaderFiles,(float)ShaderFileLoadingTime);
@@ -3230,6 +3222,7 @@ static bool TryUpdateSingleShaderFilehash(FSHA1& InOutHashState, const TCHAR* Vi
 		}
 		return false;
 	}
+	checkf(FileContents.IsValid(), TEXT("LoadCachedShaderSourceFile(\"%s\") returned true but did not provide a valid shared pointer for the shader content"), VirtualFilePath);
 	InOutHashState.UpdateWithString(GetData(*FileContents), GetNum(*FileContents));
 	if (OutErrorMessage)
 	{
