Allow AddEnqueueCopyPass to be used with texture arrays.  Pass through an optional parameter to specify which slice to read from.

diff --git a/Engine/Source/Runtime/RHI/Public/RHIGPUReadback.h b/Engine/Source/Runtime/RHI/Public/RHIGPUReadback.h
--- a/Engine/Source/Runtime/RHI/Public/RHIGPUReadback.h
+++ b/Engine/Source/Runtime/RHI/Public/RHIGPUReadback.h
@@ -64,7 +64,7 @@ public:
 		unimplemented();
 	}
 
-	void EnqueueCopy(FRHICommandList& RHICmdList, FRHITexture* SourceTexture, FResolveRect Rect = FResolveRect())
+	void EnqueueCopy(FRHICommandList& RHICmdList, FRHITexture* SourceTexture, FResolveRect Rect = FResolveRect(), uint32 SourceSlice = 0)
 	{
 		FIntVector SourcePosition, Size;
 
@@ -79,7 +79,7 @@ public:
 			Size = FIntVector::ZeroValue;
 		}
 		
-		EnqueueCopy(RHICmdList, SourceTexture, SourcePosition, 0, Size);
+		EnqueueCopy(RHICmdList, SourceTexture, SourcePosition, SourceSlice, Size);
 	}
 
 	/**
diff --git a/Engine/Source/Runtime/RenderCore/Private/RenderGraphUtils.cpp b/Engine/Source/Runtime/RenderCore/Private/RenderGraphUtils.cpp
--- a/Engine/Source/Runtime/RenderCore/Private/RenderGraphUtils.cpp
+++ b/Engine/Source/Runtime/RenderCore/Private/RenderGraphUtils.cpp
@@ -809,7 +809,7 @@ BEGIN_SHADER_PARAMETER_STRUCT(FEnqueueCopyTexturePass, )
 	RDG_TEXTURE_ACCESS(Texture, ERHIAccess::CopySrc)
 END_SHADER_PARAMETER_STRUCT()
 
-void AddEnqueueCopyPass(FRDGBuilder& GraphBuilder, FRHIGPUTextureReadback* Readback, FRDGTextureRef SourceTexture, FResolveRect Rect)
+void AddEnqueueCopyPass(FRDGBuilder& GraphBuilder, FRHIGPUTextureReadback* Readback, FRDGTextureRef SourceTexture, FResolveRect Rect, uint32 SourceSlice)
 {
 	FEnqueueCopyTexturePass* PassParameters = GraphBuilder.AllocParameters<FEnqueueCopyTexturePass>();
 	PassParameters->Texture = SourceTexture;
@@ -818,9 +818,9 @@ void AddEnqueueCopyPass(FRDGBuilder& GraphBuilder, FRHIGPUTextureReadback* Readb
 		RDG_EVENT_NAME("EnqueueCopy(%s)", SourceTexture->Name),
 		PassParameters,
 		ERDGPassFlags::Readback,
-		[Readback, SourceTexture, Rect](FRDGAsyncTask, FRHICommandList& RHICmdList)
+		[Readback, SourceTexture, Rect, SourceSlice](FRDGAsyncTask, FRHICommandList& RHICmdList)
 	{
-		Readback->EnqueueCopy(RHICmdList, SourceTexture->GetRHI(), Rect);
+		Readback->EnqueueCopy(RHICmdList, SourceTexture->GetRHI(), Rect, SourceSlice);
 	});
 }
 
diff --git a/Engine/Source/Runtime/RenderCore/Public/RenderGraphUtils.h b/Engine/Source/Runtime/RenderCore/Public/RenderGraphUtils.h
--- a/Engine/Source/Runtime/RenderCore/Public/RenderGraphUtils.h
+++ b/Engine/Source/Runtime/RenderCore/Public/RenderGraphUtils.h
@@ -851,7 +851,7 @@ RENDERCORE_API void AddCopyBufferPass(FRDGBuilder& GraphBuilder, FRDGBufferRef D
 RENDERCORE_API void AddCopyBufferPass(FRDGBuilder& GraphBuilder, FRDGBufferRef DstBuffer, FRDGBufferRef SrcBuffer);
 
 /** Adds a pass to readback contents of an RDG texture. */
-RENDERCORE_API void AddEnqueueCopyPass(FRDGBuilder& GraphBuilder, FRHIGPUTextureReadback* Readback, FRDGTextureRef SourceTexture, FResolveRect Rect = FResolveRect());
+RENDERCORE_API void AddEnqueueCopyPass(FRDGBuilder& GraphBuilder, FRHIGPUTextureReadback* Readback, FRDGTextureRef SourceTexture, FResolveRect Rect = FResolveRect(), uint32 SourceSlice = 0);
 
 /** Adds a pass to readback contents of an RDG buffer. */
 RENDERCORE_API void AddEnqueueCopyPass(FRDGBuilder& GraphBuilder, FRHIGPUBufferReadback* Readback, FRDGBufferRef SourceBuffer, uint32 NumBytes);
