Reintroduce Substrate public API function and mark it as deprecated.

diff --git a/Engine/Source/Runtime/RenderCore/Private/RenderUtils.cpp b/Engine/Source/Runtime/RenderCore/Private/RenderUtils.cpp
--- a/Engine/Source/Runtime/RenderCore/Private/RenderUtils.cpp
+++ b/Engine/Source/Runtime/RenderCore/Private/RenderUtils.cpp
@@ -2182,6 +2182,11 @@ namespace Substrate
 		return CVar.Get(InPlatform) > 0;
 	}
 
+	bool IsGlintEnabled()
+	{
+		return IsGlintEnabled(GMaxRHIShaderPlatform);
+	}
+
 	bool IsGlintEnabled(EShaderPlatform InPlatform)
 	{
 		static FShaderPlatformCachedIniValue<int32> CVar(TEXT("r.Substrate.Glints"));
diff --git a/Engine/Source/Runtime/RenderCore/Public/RenderUtils.h b/Engine/Source/Runtime/RenderCore/Public/RenderUtils.h
--- a/Engine/Source/Runtime/RenderCore/Public/RenderUtils.h
+++ b/Engine/Source/Runtime/RenderCore/Public/RenderUtils.h
@@ -617,6 +617,8 @@ namespace Substrate
 	RENDERCORE_API bool IsStochasticLightingEnabled(EShaderPlatform InPlatform);
 	RENDERCORE_API bool IsRoughDiffuseEnabled();
 	RENDERCORE_API bool IsRoughDiffuseEnabled(EShaderPlatform InPlatform);
+	UE_DEPRECATED(5.7, "Plase use IsGlintEnabled(EShaderPlatform InPlatform)")
+	RENDERCORE_API bool IsGlintEnabled();
 	RENDERCORE_API bool IsGlintEnabled(EShaderPlatform InPlatform);
 	RENDERCORE_API uint32 GlintLUTIndex();
 	RENDERCORE_API float GlintLevelBias();
