Passing RHI Command Lists through to MoviePlayer and TickableObjectRenderThread functions.

diff --git a/Engine/Plugins/Experimental/PanoramicCapture/Source/PanoramicCapture/Private/SceneCapturer.cpp b/Engine/Plugins/Experimental/PanoramicCapture/Source/PanoramicCapture/Private/SceneCapturer.cpp
--- a/Engine/Plugins/Experimental/PanoramicCapture/Source/PanoramicCapture/Private/SceneCapturer.cpp
+++ b/Engine/Plugins/Experimental/PanoramicCapture/Source/PanoramicCapture/Private/SceneCapturer.cpp
@@ -21,6 +21,7 @@
 #include "GameFramework/PlayerController.h"
 #include "GenericPlatform/GenericPlatformMisc.h"
 #include "TextureResource.h"
+#include "TickableObjectRenderThread.h"
 #include "Engine/BlendableInterface.h"
 #include "ImageUtils.h"
 #include "CoreMinimal.h"
@@ -1139,9 +1140,9 @@ void USceneCapturer::Tick( float DeltaTime )
         {
             //SetStartPosition();
             ENQUEUE_RENDER_COMMAND(SceneCapturer_HeartbeatTickTickables)(
-				[](FRHICommandList& RHICmdList)
+				[](FRHICommandListImmediate& RHICmdList)
 				{
-					TickRenderingTickables();
+					TickRenderingTickables(RHICmdList);
 				});
 
             FlushRenderingCommands();
diff --git a/Engine/Plugins/Media/BinkMedia/Source/BinkMediaPlayer/Private/BinkMovieStreamer.cpp b/Engine/Plugins/Media/BinkMedia/Source/BinkMediaPlayer/Private/BinkMovieStreamer.cpp
--- a/Engine/Plugins/Media/BinkMedia/Source/BinkMediaPlayer/Private/BinkMovieStreamer.cpp
+++ b/Engine/Plugins/Media/BinkMedia/Source/BinkMediaPlayer/Private/BinkMovieStreamer.cpp
@@ -60,7 +60,7 @@ void FBinkMovieStreamer::ForceCompletion()
 	CloseMovie();
 }
 
-bool FBinkMovieStreamer::Tick(float DeltaTime) 
+bool FBinkMovieStreamer::Tick(FRHICommandListBase& RHICmdList, float DeltaTime) 
 {
 	// Loops through all movies while not valid
 	if (!bnk) 
@@ -111,8 +111,6 @@ bool FBinkMovieStreamer::Tick(float DeltaTime)
 		FVector2D destUpperLeft = GetDefault<UBinkMoviePlayerSettings>()->BinkDestinationUpperLeft; 
 		FVector2D destLowerRight = GetDefault<UBinkMoviePlayerSettings>()->BinkDestinationLowerRight;
 
-		FRHICommandListImmediate& RHICmdList = FRHICommandListImmediate::Get();
-
 		if( !CurrentTexture->IsInitialized() ) 
 		{
 			CurrentTexture->InitResource(RHICmdList);
diff --git a/Engine/Plugins/Media/BinkMedia/Source/BinkMediaPlayer/Private/BinkMovieStreamer.h b/Engine/Plugins/Media/BinkMedia/Source/BinkMediaPlayer/Private/BinkMovieStreamer.h
--- a/Engine/Plugins/Media/BinkMedia/Source/BinkMediaPlayer/Private/BinkMovieStreamer.h
+++ b/Engine/Plugins/Media/BinkMedia/Source/BinkMediaPlayer/Private/BinkMovieStreamer.h
@@ -18,7 +18,7 @@ public:
 	virtual bool Init(const TArray<FString>& MoviePaths, TEnumAsByte<EMoviePlaybackType> inPlaybackType) override;
 	//virtual void PreviousViewportInterface(const TSharedPtr<ISlateViewport>& PreviousViewportInterface) override;
 	virtual void ForceCompletion() override;
-	virtual bool Tick(float DeltaTime) override;
+	virtual bool Tick(FRHICommandListBase& RHICmdList, float DeltaTime) override;
 	virtual TSharedPtr<class ISlateViewport> GetViewportInterface() override { return MovieViewport; }
 	virtual float GetAspectRatio() const override { return (float)MovieViewport->GetSize().X / (float)MovieViewport->GetSize().Y; }
 	virtual FString GetMovieName() override;
diff --git a/Engine/Plugins/Media/MediaMovieStreamer/Source/MediaMovieStreamer/Private/MediaMovieStreamer.cpp b/Engine/Plugins/Media/MediaMovieStreamer/Source/MediaMovieStreamer/Private/MediaMovieStreamer.cpp
--- a/Engine/Plugins/Media/MediaMovieStreamer/Source/MediaMovieStreamer/Private/MediaMovieStreamer.cpp
+++ b/Engine/Plugins/Media/MediaMovieStreamer/Source/MediaMovieStreamer/Private/MediaMovieStreamer.cpp
@@ -135,7 +135,7 @@ void FMediaMovieStreamer::ForceCompletion()
 {
 }
 
-bool FMediaMovieStreamer::Tick(float DeltaTime)
+bool FMediaMovieStreamer::Tick(FRHICommandListBase& RHICmdList, float DeltaTime)
 {
 	check(IsInRenderingThread());
 
@@ -157,7 +157,7 @@ bool FMediaMovieStreamer::Tick(float DeltaTime)
 					{
 						if (!CurrentTexture->IsInitialized())
 						{
-							CurrentTexture->InitResource(FRHICommandListImmediate::Get());
+							CurrentTexture->InitResource(RHICmdList);
 						}
 
 						// Update the slate texture.
diff --git a/Engine/Plugins/Media/MediaMovieStreamer/Source/MediaMovieStreamer/Public/MediaMovieStreamer.h b/Engine/Plugins/Media/MediaMovieStreamer/Source/MediaMovieStreamer/Public/MediaMovieStreamer.h
--- a/Engine/Plugins/Media/MediaMovieStreamer/Source/MediaMovieStreamer/Public/MediaMovieStreamer.h
+++ b/Engine/Plugins/Media/MediaMovieStreamer/Source/MediaMovieStreamer/Public/MediaMovieStreamer.h
@@ -85,7 +85,7 @@ public:
 	virtual bool Init(const TArray<FString>& InMoviePaths, TEnumAsByte<EMoviePlaybackType> InPlaybackType) override;
 	virtual void PreviousViewportInterface(const TSharedPtr<ISlateViewport>& PreviousViewportInterface) override;
 	virtual void ForceCompletion() override;
-	virtual bool Tick(float DeltaTime) override;
+	virtual bool Tick(FRHICommandListBase& RHICmdList, float DeltaTime) override;
 	virtual TSharedPtr<class ISlateViewport> GetViewportInterface() override;
 	virtual float GetAspectRatio() const override;
 	virtual FString GetMovieName() override;
diff --git a/Engine/Plugins/Runtime/AndroidMoviePlayer/Source/AndroidMoviePlayer/Private/AndroidMovieStreamer.cpp b/Engine/Plugins/Runtime/AndroidMoviePlayer/Source/AndroidMoviePlayer/Private/AndroidMovieStreamer.cpp
--- a/Engine/Plugins/Runtime/AndroidMoviePlayer/Source/AndroidMoviePlayer/Private/AndroidMovieStreamer.cpp
+++ b/Engine/Plugins/Runtime/AndroidMoviePlayer/Source/AndroidMoviePlayer/Private/AndroidMovieStreamer.cpp
@@ -81,7 +81,7 @@ struct FRHICommandUpdateTextureMovieSample final : public FRHICommand<FRHIComman
 	}
 };
 
-bool FAndroidMediaPlayerStreamer::Tick(float DeltaTime)
+bool FAndroidMediaPlayerStreamer::Tick(FRHICommandListBase& RHICmdList, float DeltaTime)
 {
 	// Check the list of textures pending deletion and remove any that are no longer valid
 	for (int32 TextureIndex = 0; TextureIndex < TexturesPendingDeletion.Num();)
diff --git a/Engine/Plugins/Runtime/AndroidMoviePlayer/Source/AndroidMoviePlayer/Private/AndroidMovieStreamer.h b/Engine/Plugins/Runtime/AndroidMoviePlayer/Source/AndroidMoviePlayer/Private/AndroidMovieStreamer.h
--- a/Engine/Plugins/Runtime/AndroidMoviePlayer/Source/AndroidMoviePlayer/Private/AndroidMovieStreamer.h
+++ b/Engine/Plugins/Runtime/AndroidMoviePlayer/Source/AndroidMoviePlayer/Private/AndroidMovieStreamer.h
@@ -16,7 +16,7 @@ public:
 
 	virtual bool Init(const TArray<FString>& MoviePaths, TEnumAsByte<EMoviePlaybackType> inPlaybackType) override;
 	virtual void ForceCompletion() override;
-	virtual bool Tick(float DeltaTime) override;
+	virtual bool Tick(FRHICommandListBase& RHICmdList, float DeltaTime) override;
 	virtual TSharedPtr<class ISlateViewport> GetViewportInterface() override;
 	virtual float GetAspectRatio() const override;
 	virtual void Cleanup() override;
diff --git a/Engine/Plugins/Runtime/AppleMoviePlayer/Source/AppleMoviePlayer/Private/AppleMovieStreamer.cpp b/Engine/Plugins/Runtime/AppleMoviePlayer/Source/AppleMoviePlayer/Private/AppleMovieStreamer.cpp
--- a/Engine/Plugins/Runtime/AppleMoviePlayer/Source/AppleMoviePlayer/Private/AppleMovieStreamer.cpp
+++ b/Engine/Plugins/Runtime/AppleMoviePlayer/Source/AppleMoviePlayer/Private/AppleMovieStreamer.cpp
@@ -156,7 +156,7 @@ bool FAVPlayerMovieStreamer::Init(const TArray<FString>& MoviePaths, TEnumAsByte
 	return StartNextMovie();
 }
 
-bool FAVPlayerMovieStreamer::Tick(float DeltaTime)
+bool FAVPlayerMovieStreamer::Tick(FRHICommandListBase& RHICmdList, float DeltaTime)
 {
 	FScopeLock LockVideoTracksLoading(&VideoTracksLoadingLock);
 
@@ -170,7 +170,7 @@ bool FAVPlayerMovieStreamer::Tick(float DeltaTime)
         // Remember that we were active. Used to edge detect active/not-active transitions
         bWasActive = true;
 
-		CheckForNextFrameAndCopy();
+		CheckForNextFrameAndCopy(RHICmdList);
 
         check(AVReader != nil);
         AVAssetReaderStatus Status = [AVReader status];
@@ -457,7 +457,7 @@ bool FAVPlayerMovieStreamer::FinishLoadingTracks()
     return bLoadedAndReading;
 }
 
-bool FAVPlayerMovieStreamer::CheckForNextFrameAndCopy()
+bool FAVPlayerMovieStreamer::CheckForNextFrameAndCopy(FRHICommandListBase& RHICmdList)
 {
     check( bVideoTracksLoaded );
 
diff --git a/Engine/Plugins/Runtime/AppleMoviePlayer/Source/AppleMoviePlayer/Private/AppleMovieStreamer.h b/Engine/Plugins/Runtime/AppleMoviePlayer/Source/AppleMoviePlayer/Private/AppleMovieStreamer.h
--- a/Engine/Plugins/Runtime/AppleMoviePlayer/Source/AppleMoviePlayer/Private/AppleMovieStreamer.h
+++ b/Engine/Plugins/Runtime/AppleMoviePlayer/Source/AppleMoviePlayer/Private/AppleMovieStreamer.h
@@ -20,7 +20,7 @@ public:
 
 	virtual bool Init(const TArray<FString>& MoviePaths, TEnumAsByte<EMoviePlaybackType> inPlaybackType) override;
 	virtual void ForceCompletion() override;
-	virtual bool Tick(float DeltaTime) override;
+	virtual bool Tick(FRHICommandListBase& RHICmdList, float DeltaTime) override;
 	virtual TSharedPtr<class ISlateViewport> GetViewportInterface() override;
 	virtual float GetAspectRatio() const override;
 	virtual void Cleanup() override;
@@ -88,7 +88,7 @@ private:
 
     void TeardownPlayback();
 
-    bool CheckForNextFrameAndCopy();
+    bool CheckForNextFrameAndCopy(FRHICommandListBase& RHICmdList);
 
 	void ReleaseMovie();
 
diff --git a/Engine/Plugins/Runtime/PreLoadScreenMoviePlayer/Source/PreLoadScreenMoviePlayer/Private/PreLoadMoviePlayerScreenBase.cpp b/Engine/Plugins/Runtime/PreLoadScreenMoviePlayer/Source/PreLoadScreenMoviePlayer/Private/PreLoadMoviePlayerScreenBase.cpp
--- a/Engine/Plugins/Runtime/PreLoadScreenMoviePlayer/Source/PreLoadScreenMoviePlayer/Private/PreLoadMoviePlayerScreenBase.cpp
+++ b/Engine/Plugins/Runtime/PreLoadScreenMoviePlayer/Source/PreLoadScreenMoviePlayer/Private/PreLoadMoviePlayerScreenBase.cpp
@@ -208,14 +208,14 @@ void FPreLoadMoviePlayerScreenBase::Tick(float DeltaTime)
     }
 }
 
-void FPreLoadMoviePlayerScreenBase::RenderTick(float DeltaTime)
+void FPreLoadMoviePlayerScreenBase::RenderTick(FRHICommandListImmediate& RHICmdList, float DeltaTime)
 {
     if (MovieStreamingIsPrepared())
     {
 		MovieStreamer->TickPreEngine();
 		MovieStreamer->TickPostEngine();
 
-        const bool bMovieIsDone = MovieStreamer->Tick(DeltaTime);
+        const bool bMovieIsDone = MovieStreamer->Tick(RHICmdList, DeltaTime);
         if (bMovieIsDone)
         {
             MovieStreamingIsDone.Set(1);
diff --git a/Engine/Plugins/Runtime/PreLoadScreenMoviePlayer/Source/PreLoadScreenMoviePlayer/Public/PreLoadMoviePlayerScreenBase.h b/Engine/Plugins/Runtime/PreLoadScreenMoviePlayer/Source/PreLoadScreenMoviePlayer/Public/PreLoadMoviePlayerScreenBase.h
--- a/Engine/Plugins/Runtime/PreLoadScreenMoviePlayer/Source/PreLoadScreenMoviePlayer/Public/PreLoadMoviePlayerScreenBase.h
+++ b/Engine/Plugins/Runtime/PreLoadScreenMoviePlayer/Source/PreLoadScreenMoviePlayer/Public/PreLoadMoviePlayerScreenBase.h
@@ -19,7 +19,7 @@ public:
     virtual void OnPlay(TWeakPtr<SWindow> TargetWindow) override;
     virtual void OnStop() override;
     virtual void Tick(float DeltaTime) override;
-    virtual void RenderTick(float DeltaTime) override;
+    virtual void RenderTick(FRHICommandListImmediate& RHICmdList, float DeltaTime) override;
 
     virtual void Init() override;
 
diff --git a/Engine/Plugins/Runtime/WebMMoviePlayer/Source/WebMMoviePlayer/Private/WebMMovieStreamer.cpp b/Engine/Plugins/Runtime/WebMMoviePlayer/Source/WebMMoviePlayer/Private/WebMMovieStreamer.cpp
--- a/Engine/Plugins/Runtime/WebMMoviePlayer/Source/WebMMoviePlayer/Private/WebMMovieStreamer.cpp
+++ b/Engine/Plugins/Runtime/WebMMoviePlayer/Source/WebMMoviePlayer/Private/WebMMovieStreamer.cpp
@@ -156,7 +156,7 @@ bool FWebMMovieStreamer::IsLastMovieInPlaylist()
 	return MovieQueue.Num() == 0;
 }
 
-bool FWebMMovieStreamer::Tick(float InDeltaTime)
+bool FWebMMovieStreamer::Tick(FRHICommandListBase& RHICmdList, float InDeltaTime)
 {
 	AudioBackend->Tick(InDeltaTime);
 
diff --git a/Engine/Plugins/Runtime/WebMMoviePlayer/Source/WebMMoviePlayer/Private/WebMMovieStreamer.h b/Engine/Plugins/Runtime/WebMMoviePlayer/Source/WebMMoviePlayer/Private/WebMMovieStreamer.h
--- a/Engine/Plugins/Runtime/WebMMoviePlayer/Source/WebMMoviePlayer/Private/WebMMovieStreamer.h
+++ b/Engine/Plugins/Runtime/WebMMoviePlayer/Source/WebMMoviePlayer/Private/WebMMovieStreamer.h
@@ -30,7 +30,7 @@ public:
 	//~ IMediaOutput interface
 	virtual bool Init(const TArray<FString>& InMoviePaths, TEnumAsByte<EMoviePlaybackType> InPlaybackType) override;
 	virtual void ForceCompletion() override;
-	virtual bool Tick(float InDeltaTime) override;
+	virtual bool Tick(FRHICommandListBase& RHICmdList, float InDeltaTime) override;
 	virtual TSharedPtr<class ISlateViewport> GetViewportInterface() override;
 	virtual float GetAspectRatio() const override;
 	virtual FString GetMovieName() override;
diff --git a/Engine/Plugins/Runtime/WindowsMoviePlayer/Source/WindowsMoviePlayer/Private/WindowsMovieStreamer.cpp b/Engine/Plugins/Runtime/WindowsMoviePlayer/Source/WindowsMoviePlayer/Private/WindowsMovieStreamer.cpp
--- a/Engine/Plugins/Runtime/WindowsMoviePlayer/Source/WindowsMoviePlayer/Private/WindowsMovieStreamer.cpp
+++ b/Engine/Plugins/Runtime/WindowsMoviePlayer/Source/WindowsMoviePlayer/Private/WindowsMovieStreamer.cpp
@@ -73,7 +73,7 @@ void FMediaFoundationMovieStreamer::ForceCompletion()
 	CloseMovie();
 }
 
-void FMediaFoundationMovieStreamer::ConvertSample(FRHICommandListImmediate& RHICmdList)
+void FMediaFoundationMovieStreamer::ConvertSample(FRHICommandList& RHICmdList)
 {
 	const bool SrgbOutput = false;
 	const bool bSampleIsOutputSrgb = false;
@@ -166,10 +166,9 @@ void FMediaFoundationMovieStreamer::ConvertSample(FRHICommandListImmediate& RHIC
 	RHICmdList.Transition(FRHITransitionInfo(RenderTarget, ERHIAccess::Unknown, ERHIAccess::SRVGraphics));
 }
 
-bool FMediaFoundationMovieStreamer::Tick(float DeltaTime)
+bool FMediaFoundationMovieStreamer::Tick(FRHICommandListBase& RHICmdList, float DeltaTime)
 {
 	FSlateTexture2DRHIRef* CurrentTexture = Texture.Get();
-	FRHICommandListImmediate& RHICmdList = FRHICommandListImmediate::Get();
 
 	if (CurrentTexture && !CurrentTexture->IsInitialized())
 	{
@@ -182,13 +181,13 @@ bool FMediaFoundationMovieStreamer::Tick(float DeltaTime)
 		if (SourceFormat.SampleFormat == EMediaTextureSampleFormat::CharBGRA && SourceFormat.BufferDim == SourceFormat.OutputDim)
 		{
 			uint32 Stride;
-			uint8* DestTextureData = (uint8*)RHILockTexture2D( CurrentTexture->GetTypedResource(), 0, RLM_WriteOnly, Stride, false );
+			uint8* DestTextureData = (uint8*)RHICmdList.GetAsImmediate().LockTexture2D( CurrentTexture->GetTypedResource(), 0, RLM_WriteOnly, Stride, false );
 			FMemory::Memcpy( DestTextureData, TextureData.GetData(), TextureData.Num());
-			RHIUnlockTexture2D( CurrentTexture->GetTypedResource(), 0, false );
+			RHICmdList.GetAsImmediate().UnlockTexture2D( CurrentTexture->GetTypedResource(), 0, false );
 		}
 		else
 		{
-			ConvertSample(RHICmdList);
+			ConvertSample(RHICmdList.GetAsImmediate());
 		}
 
 		if (MovieViewport->GetViewportRenderTargetTexture() == nullptr)
diff --git a/Engine/Plugins/Runtime/WindowsMoviePlayer/Source/WindowsMoviePlayer/Private/WindowsMovieStreamer.h b/Engine/Plugins/Runtime/WindowsMoviePlayer/Source/WindowsMoviePlayer/Private/WindowsMovieStreamer.h
--- a/Engine/Plugins/Runtime/WindowsMoviePlayer/Source/WindowsMoviePlayer/Private/WindowsMovieStreamer.h
+++ b/Engine/Plugins/Runtime/WindowsMoviePlayer/Source/WindowsMoviePlayer/Private/WindowsMovieStreamer.h
@@ -26,7 +26,7 @@ public:
 	/** IMovieStreamer interface */
 	virtual bool Init(const TArray<FString>& MoviePaths, TEnumAsByte<EMoviePlaybackType> inPlaybackType) override;
 	virtual void ForceCompletion() override;
-	virtual bool Tick(float DeltaTime) override;
+	virtual bool Tick(FRHICommandListBase& RHICmdList, float DeltaTime) override;
 	virtual TSharedPtr<class ISlateViewport> GetViewportInterface() override
 	{
 		return MovieViewport;
@@ -53,7 +53,7 @@ private:
 	/** Cleans up rendering resources once movies are done playing */
 	void CleanupRenderingResources();
 	/** Use shader to generate final image */
-	void ConvertSample(FRHICommandListImmediate& RHICmdList);
+	void ConvertSample(FRHICommandList& RHICmdList);
 
 private:
 	/** A list of all the stored movie paths we have enqueued for playing */
diff --git a/Engine/Source/Runtime/Engine/Public/ResourcePool.h b/Engine/Source/Runtime/Engine/Public/ResourcePool.h
--- a/Engine/Source/Runtime/Engine/Public/ResourcePool.h
+++ b/Engine/Source/Runtime/Engine/Public/ResourcePool.h
@@ -280,7 +280,7 @@ public:
 	}
 	
 public: // From FTickableObjectRenderThread
-	virtual void Tick( float DeltaTime ) override
+	virtual void Tick(FRHICommandListImmediate& RHICmdList,  float DeltaTime ) override
 	{
 		ensure(IsInRenderingThread());
 
diff --git a/Engine/Source/Runtime/Launch/Private/LaunchEngineLoop.cpp b/Engine/Source/Runtime/Launch/Private/LaunchEngineLoop.cpp
--- a/Engine/Source/Runtime/Launch/Private/LaunchEngineLoop.cpp
+++ b/Engine/Source/Runtime/Launch/Private/LaunchEngineLoop.cpp
@@ -5508,7 +5508,7 @@ void FEngineLoop::Tick()
 	if (!GUseThreadedRendering)
 	{
 		SCOPED_NAMED_EVENT(Tick_RenderingTickables, SystemColor);
-		TickRenderingTickables();
+		TickRenderingTickables(FRHICommandListImmediate::Get());
 	}
 
 	// Ensure we aren't starting a frame while loading or playing a loading movie
diff --git a/Engine/Source/Runtime/MoviePlayer/Private/DefaultGameMoviePlayer.cpp b/Engine/Source/Runtime/MoviePlayer/Private/DefaultGameMoviePlayer.cpp
--- a/Engine/Source/Runtime/MoviePlayer/Private/DefaultGameMoviePlayer.cpp
+++ b/Engine/Source/Runtime/MoviePlayer/Private/DefaultGameMoviePlayer.cpp
@@ -537,7 +537,7 @@ void FDefaultGameMoviePlayer::WaitForMovieToFinish(bool bAllowEngineTick)
 					[InMoviePlayer, DeltaTime](FRHICommandListImmediate& RHICmdList)
 					{
 						GFrameNumberRenderThread++;				
-						InMoviePlayer->TickStreamer(DeltaTime);
+						InMoviePlayer->TickStreamer(RHICmdList, DeltaTime);
 					}
 				);
 				
@@ -554,8 +554,8 @@ void FDefaultGameMoviePlayer::WaitForMovieToFinish(bool bAllowEngineTick)
 				ENQUEUE_RENDER_COMMAND(FinishLoadingMovieFrame)(
 					[](FRHICommandListImmediate& RHICmdList)
 					{
-						GRHICommandList.GetImmediateCommandList().EndFrame();
-						GRHICommandList.GetImmediateCommandList().ImmediateFlush(EImmediateFlushType::FlushRHIThreadFlushResources);
+						RHICmdList.EndFrame();
+						RHICmdList.ImmediateFlush(EImmediateFlushType::FlushRHIThreadFlushResources);
 					}
 				);
 				FlushRenderingCommands();
@@ -698,7 +698,7 @@ void FDefaultGameMoviePlayer::SetIsSlateThreadAllowed(bool bInIsSlateThreadAllow
 	}
 }
 
-void FDefaultGameMoviePlayer::Tick( float DeltaTime )
+void FDefaultGameMoviePlayer::Tick( FRHICommandListImmediate& RHICmdList, float DeltaTime )
 {
 	check(IsInRenderingThread());
 	if (MainWindow.IsValid() && VirtualRenderWindow.IsValid() && !IsLoadingFinished() && GDynamicRHI && !GDynamicRHI->RHIIsRenderingSuspended())
@@ -709,20 +709,20 @@ void FDefaultGameMoviePlayer::Tick( float DeltaTime )
 			if(SyncMechanism->IsSlateDrawPassEnqueued())
 			{
 				GFrameNumberRenderThread++;
-				TickStreamer(DeltaTime);
+				TickStreamer(RHICmdList, DeltaTime);
 				SyncMechanism->ResetSlateDrawPassEnqueued();
-				GRHICommandList.GetImmediateCommandList().EndFrame();
-				GRHICommandList.GetImmediateCommandList().ImmediateFlush(EImmediateFlushType::FlushRHIThreadFlushResources);
+				RHICmdList.EndFrame();
+				RHICmdList.ImmediateFlush(EImmediateFlushType::FlushRHIThreadFlushResources);
 			}
 		}
 	}
 }
 
-void FDefaultGameMoviePlayer::TickStreamer(float DeltaTime)
+void FDefaultGameMoviePlayer::TickStreamer(FRHICommandListImmediate& RHICmdList, float DeltaTime)
 {	
 	if (MovieStreamingIsPrepared() && ActiveMovieStreamer.IsValid() && !IsMovieStreamingFinished())
 	{
-		const bool bMovieIsDone = ActiveMovieStreamer->Tick(DeltaTime);
+		const bool bMovieIsDone = ActiveMovieStreamer->Tick(RHICmdList, DeltaTime);
 		if (bMovieIsDone)
 		{
 			MovieStreamingIsDone.Set(1);
diff --git a/Engine/Source/Runtime/MoviePlayer/Private/DefaultGameMoviePlayer.h b/Engine/Source/Runtime/MoviePlayer/Private/DefaultGameMoviePlayer.h
--- a/Engine/Source/Runtime/MoviePlayer/Private/DefaultGameMoviePlayer.h
+++ b/Engine/Source/Runtime/MoviePlayer/Private/DefaultGameMoviePlayer.h
@@ -109,7 +109,7 @@ public:
 	virtual void SetIsSlateThreadAllowed(bool bInIsSlateThreadAllowed) override;
 
 	/** FTickableObjectRenderThread interface */
-	virtual void Tick( float DeltaTime ) override;
+	virtual void Tick( FRHICommandListImmediate& RHICmdList, float DeltaTime ) override;
 	virtual TStatId GetStatId() const override;
 	virtual bool IsTickable() const override;
 	
@@ -137,7 +137,7 @@ public:
 private:
 
 	/** Ticks the underlying MovieStreamer.  Must be done exactly once before each DrawWindows call. */
-	void TickStreamer(float DeltaTime);
+	void TickStreamer(FRHICommandListImmediate& RHICmdList, float DeltaTime);
 
 	/** True if we have both a registered movie streamer and movies to stream */
 	bool MovieStreamingIsPrepared() const;
diff --git a/Engine/Source/Runtime/MoviePlayer/Public/MoviePlayer.h b/Engine/Source/Runtime/MoviePlayer/Public/MoviePlayer.h
--- a/Engine/Source/Runtime/MoviePlayer/Public/MoviePlayer.h
+++ b/Engine/Source/Runtime/MoviePlayer/Public/MoviePlayer.h
@@ -87,7 +87,7 @@ public:
 	virtual void ForceCompletion() = 0;
 
 	/** Code run every tick for any additional per tick handling of playing the movie. Returns true if done. */
-	virtual bool Tick(float DeltaTime) = 0;
+	virtual bool Tick(FRHICommandListBase& RHICmdList, float DeltaTime) = 0;
 
 	/** Gets a viewport interface which will be used to draw the movie. */
 	virtual TSharedPtr<class ISlateViewport> GetViewportInterface() = 0;
diff --git a/Engine/Source/Runtime/PreLoadScreen/Private/PreLoadScreenManager.cpp b/Engine/Source/Runtime/PreLoadScreen/Private/PreLoadScreenManager.cpp
--- a/Engine/Source/Runtime/PreLoadScreen/Private/PreLoadScreenManager.cpp
+++ b/Engine/Source/Runtime/PreLoadScreen/Private/PreLoadScreenManager.cpp
@@ -377,18 +377,18 @@ void FPreLoadScreenManager::HandleCustomSplashScreenPlay()
 	}
 }
 
-void FPreLoadScreenManager::StaticRenderTick_RenderThread()
+void FPreLoadScreenManager::StaticRenderTick_RenderThread(FRHICommandListImmediate& RHICmdList)
 {
 	LLM_SCOPE(ELLMTag::RenderingThreadMemory);
 	check(IsInRenderingThread());
 
 	if (ensure(FPreLoadScreenManager::Get())) // The manager should clear the slate render thread before closing
 	{
-		FPreLoadScreenManager::Get()->RenderTick_RenderThread();
+		FPreLoadScreenManager::Get()->RenderTick_RenderThread(RHICmdList);
 	}
 }
 
-void FPreLoadScreenManager::RenderTick_RenderThread()
+void FPreLoadScreenManager::RenderTick_RenderThread(FRHICommandListImmediate& RHICmdList)
 {
 	//Calculate tick time
 	const double CurrentTime = FPlatformTime::Seconds();
@@ -404,9 +404,9 @@ void FPreLoadScreenManager::RenderTick_RenderThread()
 		if (MainWindow.IsValid() && VirtualRenderWindow.IsValid() && !PinnedActivePreloadScreen->IsDone())
 		{
 			GFrameNumberRenderThread++;
-			PinnedActivePreloadScreen->RenderTick(DeltaTime);
-			GRHICommandList.GetImmediateCommandList().EndFrame();
-			GRHICommandList.GetImmediateCommandList().ImmediateFlush(EImmediateFlushType::FlushRHIThreadFlushResources);
+			PinnedActivePreloadScreen->RenderTick(RHICmdList, DeltaTime);
+			RHICmdList.EndFrame();
+			RHICmdList.ImmediateFlush(EImmediateFlushType::FlushRHIThreadFlushResources);
 		}
 	}
 }
@@ -565,7 +565,7 @@ void FPreLoadScreenManager::EarlyPlayRenderFrameTick()
 						GFrameNumberRenderThread++;
 
 						bHasRenderPreLoadScreenFrame_RenderThread = true;
-						PinnedActivePreloadScreen->RenderTick(SlateDeltaTime);
+						PinnedActivePreloadScreen->RenderTick(RHICmdList, SlateDeltaTime);
 					}
 				});
 
@@ -579,8 +579,8 @@ void FPreLoadScreenManager::EarlyPlayRenderFrameTick()
 				{
 					// this is still valid because we do a FlushRenderingCommands in StopPreLoadScreen
 					bHasRenderPreLoadScreenFrame_RenderThread = false;
-					GRHICommandList.GetImmediateCommandList().EndFrame();
-					GRHICommandList.GetImmediateCommandList().ImmediateFlush(EImmediateFlushType::FlushRHIThreadFlushResources);
+					RHICmdList.EndFrame();
+					RHICmdList.ImmediateFlush(EImmediateFlushType::FlushRHIThreadFlushResources);
 				});
 		}
 	}
diff --git a/Engine/Source/Runtime/PreLoadScreen/Private/PreLoadSlateThreading.cpp b/Engine/Source/Runtime/PreLoadScreen/Private/PreLoadSlateThreading.cpp
--- a/Engine/Source/Runtime/PreLoadScreen/Private/PreLoadSlateThreading.cpp
+++ b/Engine/Source/Runtime/PreLoadScreen/Private/PreLoadSlateThreading.cpp
@@ -244,7 +244,7 @@ void FPreLoadScreenSlateSynchMechanism::RunMainLoop_SlateThread()
 						{
 							if (SyncMech->IsSlateMainLoopRunning_AnyThread())
 							{
-								FPreLoadScreenManager::StaticRenderTick_RenderThread();
+								FPreLoadScreenManager::StaticRenderTick_RenderThread(RHICmdList);
 							}
 
 							EnqueueRenderEvent->Trigger();
diff --git a/Engine/Source/Runtime/PreLoadScreen/Public/PreLoadScreen.h b/Engine/Source/Runtime/PreLoadScreen/Public/PreLoadScreen.h
--- a/Engine/Source/Runtime/PreLoadScreen/Public/PreLoadScreen.h
+++ b/Engine/Source/Runtime/PreLoadScreen/Public/PreLoadScreen.h
@@ -39,7 +39,7 @@ public:
 	 * This tick happens as part of the slate render tick during an EarlyStartupLoadScreen
 	 * @note Called on the render thread
 	 */
-	virtual void RenderTick(float DeltaTime) = 0;
+	virtual void RenderTick(FRHICommandListImmediate& RHICmdList, float DeltaTime) = 0;
 
 	/** Callback for when a PreLoadScreen starts being displayed. Provides a reference to the SWindow that will be used to display content */
 	virtual void OnPlay(TWeakPtr<SWindow> TargetWindow) = 0;
diff --git a/Engine/Source/Runtime/PreLoadScreen/Public/PreLoadScreenBase.h b/Engine/Source/Runtime/PreLoadScreen/Public/PreLoadScreenBase.h
--- a/Engine/Source/Runtime/PreLoadScreen/Public/PreLoadScreenBase.h
+++ b/Engine/Source/Runtime/PreLoadScreen/Public/PreLoadScreenBase.h
@@ -14,7 +14,7 @@ public:
     //We don't use these in the FPreLoadScreenBase, but they are useful for game-specific implementations
     virtual void Tick(float DeltaTime) override {}
 	virtual bool ShouldRender() const override { return true; }
-    virtual void RenderTick(float DeltaTime) override {};
+    virtual void RenderTick(FRHICommandListImmediate& RHICmdList, float DeltaTime) override {};
     virtual void OnStop() override {}
 
     //Store off TargetWindow
diff --git a/Engine/Source/Runtime/PreLoadScreen/Public/PreLoadScreenManager.h b/Engine/Source/Runtime/PreLoadScreen/Public/PreLoadScreenManager.h
--- a/Engine/Source/Runtime/PreLoadScreen/Public/PreLoadScreenManager.h
+++ b/Engine/Source/Runtime/PreLoadScreen/Public/PreLoadScreenManager.h
@@ -91,8 +91,8 @@ protected:
 	PRELOADSCREEN_API void PlatformSpecificGameLogicFrameTick();
 
 	//Creates a tick on the Render Thread that we run every
-	static PRELOADSCREEN_API void StaticRenderTick_RenderThread();
-	PRELOADSCREEN_API void RenderTick_RenderThread();
+	static PRELOADSCREEN_API void StaticRenderTick_RenderThread(FRHICommandListImmediate& RHICmdList);
+	PRELOADSCREEN_API void RenderTick_RenderThread(FRHICommandListImmediate& RHICmdList);
 
     /*** These functions describe how everything is handled during an non-Early PreLoadPlay. Everything is handled asynchronously in this case with a standalone renderer ***/
 	PRELOADSCREEN_API void HandleEngineLoadingPlay();
diff --git a/Engine/Source/Runtime/RenderCore/Private/RenderingThread.cpp b/Engine/Source/Runtime/RenderCore/Private/RenderingThread.cpp
--- a/Engine/Source/Runtime/RenderCore/Private/RenderingThread.cpp
+++ b/Engine/Source/Runtime/RenderCore/Private/RenderingThread.cpp
@@ -151,61 +151,6 @@ FAutoConsoleVariable CVarRenderCommandPipeMode(
  * Tick all rendering thread tickable objects
  */
 
-/** Static array of tickable objects that are ticked from rendering thread*/
-FTickableObjectRenderThread::FRenderingThreadTickableObjectsArray FTickableObjectRenderThread::RenderingThreadTickableObjects;
-FTickableObjectRenderThread::FRenderingThreadTickableObjectsArray FTickableObjectRenderThread::RenderingThreadHighFrequencyTickableObjects;
-
-void TickHighFrequencyTickables(double CurTime)
-{
-	static double LastHighFreqTime = FPlatformTime::Seconds();
-	float DeltaSecondsHighFreq = float(CurTime - LastHighFreqTime);
-
-	// tick any high frequency rendering thread tickables.
-	for (int32 ObjectIndex = 0; ObjectIndex < FTickableObjectRenderThread::RenderingThreadHighFrequencyTickableObjects.Num(); ObjectIndex++)
-	{
-		FTickableObjectRenderThread* TickableObject = FTickableObjectRenderThread::RenderingThreadHighFrequencyTickableObjects[ObjectIndex];
-		// make sure it wants to be ticked and the rendering thread isn't suspended
-		if (TickableObject->IsTickable())
-		{
-			STAT(FScopeCycleCounter(TickableObject->GetStatId());)
-				TickableObject->Tick(DeltaSecondsHighFreq);
-		}
-	}
-
-	LastHighFreqTime = CurTime;
-}
-
-void TickRenderingTickables()
-{
-	static double LastTickTime = FPlatformTime::Seconds();
-	
-
-	// calc how long has passed since last tick
-	double CurTime = FPlatformTime::Seconds();
-	float DeltaSeconds = float(CurTime - LastTickTime);
-		
-	TickHighFrequencyTickables(CurTime);
-
-	if (DeltaSeconds < (1.f/GRenderingThreadMaxIdleTickFrequency))
-	{
-		return;
-	}
-
-	// tick any rendering thread tickables
-	for (int32 ObjectIndex = 0; ObjectIndex < FTickableObjectRenderThread::RenderingThreadTickableObjects.Num(); ObjectIndex++)
-	{
-		FTickableObjectRenderThread* TickableObject = FTickableObjectRenderThread::RenderingThreadTickableObjects[ObjectIndex];
-		// make sure it wants to be ticked and the rendering thread isn't suspended
-		if (TickableObject->IsTickable())
-		{
-			STAT(FScopeCycleCounter(TickableObject->GetStatId());)
-			TickableObject->Tick(DeltaSeconds);
-		}
-	}
-	// update the last time we ticked
-	LastTickTime = CurTime;
-}
-
 /** How many cycles the renderthread used (excluding idle time). It's set once per frame in FViewport::Draw. */
 uint32 GRenderThreadTime = 0;
 /** How many cycles of wait time renderthread used. It's set once per frame in FViewport::Draw. */
@@ -532,14 +477,14 @@ public:
 			{
 				OutstandingHeartbeats.Increment();
 				ENQUEUE_RENDER_COMMAND(HeartbeatTickTickables)(
-					[](FRHICommandList& RHICmdList)
+					[](FRHICommandListImmediate& RHICmdList)
 					{
 						OutstandingHeartbeats.Decrement();
 						// make sure that rendering thread tickables get a chance to tick, even if the render thread is starving
 						// but if GSuspendRenderingTickables is != 0 a flush is happening so don't tick during it
 						if (!GSuspendRenderingTickables.Load(EMemoryOrder::Relaxed))
 						{
-							TickRenderingTickables();
+							TickRenderingTickables(RHICmdList);
 						}
 					});
 			}
diff --git a/Engine/Source/Runtime/RenderCore/Private/ShaderPipelineCache.cpp b/Engine/Source/Runtime/RenderCore/Private/ShaderPipelineCache.cpp
--- a/Engine/Source/Runtime/RenderCore/Private/ShaderPipelineCache.cpp
+++ b/Engine/Source/Runtime/RenderCore/Private/ShaderPipelineCache.cpp
@@ -1785,7 +1785,7 @@ bool FShaderPipelineCache::IsTickable() const
 		);
 }
 
-void FShaderPipelineCache::Tick(float DeltaTime)
+void FShaderPipelineCache::Tick(FRHICommandListImmediate& RHICmdList, float DeltaTime)
 {
 	LLM_SCOPE(ELLMTag::Shaders);
 	FScopeLock Lock(&Mutex);
@@ -1864,6 +1864,12 @@ void FShaderPipelineCache::Tick(float DeltaTime)
 	}
 }
 
+// DEPRECATED
+void FShaderPipelineCache::Tick(float DeltaTime)
+{
+	Tick(FRHICommandListImmediate::Get(), DeltaTime);
+}
+
 void FShaderPipelineCacheTask::Tick(float DeltaTime)
 {
 	if (PrecompileStartTime > 0.0)
diff --git a/Engine/Source/Runtime/RenderCore/Private/TickableObjectRenderThread.cpp b/Engine/Source/Runtime/RenderCore/Private/TickableObjectRenderThread.cpp
new file mode 100644
--- /dev/null
+++ b/Engine/Source/Runtime/RenderCore/Private/TickableObjectRenderThread.cpp
@@ -0,0 +1,68 @@
+// Copyright Epic Games, Inc. All Rights Reserved.
+
+#include "TickableObjectRenderThread.h"
+#include "RHICommandList.h"
+
+/** Static array of tickable objects that are ticked from rendering thread*/
+FTickableObjectRenderThread::FRenderingThreadTickableObjectsArray FTickableObjectRenderThread::RenderingThreadTickableObjects;
+FTickableObjectRenderThread::FRenderingThreadTickableObjectsArray FTickableObjectRenderThread::RenderingThreadHighFrequencyTickableObjects;
+
+void TickHighFrequencyTickables(FRHICommandListImmediate& RHICmdList, double CurTime)
+{
+	static double LastHighFreqTime = FPlatformTime::Seconds();
+	float DeltaSecondsHighFreq = float(CurTime - LastHighFreqTime);
+
+	// tick any high frequency rendering thread tickables.
+	for (int32 ObjectIndex = 0; ObjectIndex < FTickableObjectRenderThread::RenderingThreadHighFrequencyTickableObjects.Num(); ObjectIndex++)
+	{
+		FTickableObjectRenderThread* TickableObject = FTickableObjectRenderThread::RenderingThreadHighFrequencyTickableObjects[ObjectIndex];
+		// make sure it wants to be ticked and the rendering thread isn't suspended
+		if (TickableObject->IsTickable())
+		{
+			STAT(FScopeCycleCounter(TickableObject->GetStatId());)
+			TickableObject->Tick(RHICmdList, DeltaSecondsHighFreq);
+		}
+	}
+
+	LastHighFreqTime = CurTime;
+}
+
+extern float GRenderingThreadMaxIdleTickFrequency;
+
+void TickRenderingTickables(FRHICommandListImmediate& RHICmdList)
+{
+	static double LastTickTime = FPlatformTime::Seconds();
+
+	// calc how long has passed since last tick
+	double CurTime = FPlatformTime::Seconds();
+	float DeltaSeconds = float(CurTime - LastTickTime);
+
+	TickHighFrequencyTickables(RHICmdList, CurTime);
+
+	if (DeltaSeconds < (1.f/GRenderingThreadMaxIdleTickFrequency))
+	{
+		return;
+	}
+
+	// tick any rendering thread tickables
+	for (int32 ObjectIndex = 0; ObjectIndex < FTickableObjectRenderThread::RenderingThreadTickableObjects.Num(); ObjectIndex++)
+	{
+		FTickableObjectRenderThread* TickableObject = FTickableObjectRenderThread::RenderingThreadTickableObjects[ObjectIndex];
+		// make sure it wants to be ticked and the rendering thread isn't suspended
+		if (TickableObject->IsTickable())
+		{
+			STAT(FScopeCycleCounter(TickableObject->GetStatId());)
+			TickableObject->Tick(RHICmdList, DeltaSeconds);
+		}
+	}
+
+	// update the last time we ticked
+	LastTickTime = CurTime;
+}
+
+// DEPRECATED
+void TickRenderingTickables()
+{
+	TickRenderingTickables(FRHICommandListImmediate::Get());
+}
+
diff --git a/Engine/Source/Runtime/RenderCore/Public/RenderingThread.h b/Engine/Source/Runtime/RenderCore/Public/RenderingThread.h
--- a/Engine/Source/Runtime/RenderCore/Public/RenderingThread.h
+++ b/Engine/Source/Runtime/RenderCore/Public/RenderingThread.h
@@ -114,8 +114,6 @@ extern RENDERCORE_API void FlushPendingDeleteRHIResources_GameThread();
 UE_DEPRECATED(5.6, "FlushPendingDeleteRHIResources_RenderThread is deprecated. Call ImmediateFlush(EImmediateFlushType::FlushRHIThreadFlushResources) on the immediate RHI command list instead.")
 extern RENDERCORE_API void FlushPendingDeleteRHIResources_RenderThread();
 
-extern RENDERCORE_API void TickRenderingTickables();
-
 extern RENDERCORE_API void StartRenderCommandFenceBundler();
 extern RENDERCORE_API void StopRenderCommandFenceBundler();
 
diff --git a/Engine/Source/Runtime/RenderCore/Public/ShaderPipelineCache.h b/Engine/Source/Runtime/RenderCore/Public/ShaderPipelineCache.h
--- a/Engine/Source/Runtime/RenderCore/Public/ShaderPipelineCache.h
+++ b/Engine/Source/Runtime/RenderCore/Public/ShaderPipelineCache.h
@@ -138,6 +138,9 @@ public:
 
 	RENDERCORE_API bool IsTickable() const;
 
+	RENDERCORE_API void Tick( FRHICommandListImmediate& RHICmdList, float DeltaTime );
+
+	UE_DEPRECATED(5.7, "Tick needs a command list")
 	RENDERCORE_API void Tick( float DeltaTime );
 
 	RENDERCORE_API bool NeedsRenderingResumedForRenderingThreadTick() const;
diff --git a/Engine/Source/Runtime/RenderCore/Public/TickableObjectRenderThread.h b/Engine/Source/Runtime/RenderCore/Public/TickableObjectRenderThread.h
--- a/Engine/Source/Runtime/RenderCore/Public/TickableObjectRenderThread.h
+++ b/Engine/Source/Runtime/RenderCore/Public/TickableObjectRenderThread.h
@@ -1,14 +1,12 @@
 // Copyright Epic Games, Inc. All Rights Reserved.
 
-/*=============================================================================
-	TickableObjectRenderThread.h: Rendering thread tickable object definition.
-=============================================================================*/
-
 #pragma once
 
 #include "CoreMinimal.h"
 #include "Stats/Stats.h"
 
+class FRHICommandListImmediate;
+
 /**
  * This class provides common registration for render thread tickable objects. It is an
  * abstract base class requiring you to implement the Tick() method.
@@ -43,9 +41,8 @@ public:
 	 *
 	 * @param bRegisterImmediately true if the object should be registered immediately.
 	 */
-	FTickableObjectRenderThread(bool bRegisterImmediately=true, bool bInHighFrequency=false) :
-		bRegistered(false),
-		bHighFrequency(bInHighFrequency)
+	FTickableObjectRenderThread(bool bRegisterImmediately=true, bool bInHighFrequency=false)
+		: bHighFrequency(bInHighFrequency)
 	{
 		if(bRegisterImmediately)
 		{
@@ -105,7 +102,7 @@ public:
 	 *
 	 * @param DeltaTime	Game time passed since the last call.
 	 */
-	virtual void Tick( float DeltaTime ) = 0;
+	virtual void Tick(FRHICommandListImmediate& RHICmdList, float DeltaTime) = 0;
 
 	/** return the stat id to use for this tickable **/
 	virtual TStatId GetStatId() const = 0;
@@ -132,7 +129,11 @@ public:
 	}
 
 private:
-	bool bRegistered;
-	bool bHighFrequency;
+	bool bRegistered = false;
+	const bool bHighFrequency;
 };
 
+extern RENDERCORE_API void TickRenderingTickables(FRHICommandListImmediate& RHICmdList);
+
+UE_DEPRECATED(5.7, "TickRenderingTickables needs a command list")
+extern RENDERCORE_API void TickRenderingTickables();
