ShaderCompileCore UE_SHADER_CACHE_VERSION -> .cpp file

It was declared static const meaning every TU including the header ended up
with both a copy and the global ctor call for it.

diff --git a/Engine/Source/Runtime/RenderCore/Private/ShaderCompilerCore.cpp b/Engine/Source/Runtime/RenderCore/Private/ShaderCompilerCore.cpp
new file mode 100644
--- /dev/null
+++ b/Engine/Source/Runtime/RenderCore/Private/ShaderCompilerCore.cpp
@@ -0,0 +1,5 @@
+// Copyright Epic Games, Inc. All Rights Reserved.
+
+#include "ShaderCompilerCore.h"
+
+const FGuid UE_SHADER_CACHE_VERSION{"B04528D6-1A2C-44C8-A742-E7D497D3562C"};
diff --git a/Engine/Source/Runtime/RenderCore/Public/ShaderCompilerCore.h b/Engine/Source/Runtime/RenderCore/Public/ShaderCompilerCore.h
--- a/Engine/Source/Runtime/RenderCore/Public/ShaderCompilerCore.h
+++ b/Engine/Source/Runtime/RenderCore/Public/ShaderCompilerCore.h
@@ -39,9 +39,10 @@ inline const int32 ShaderCompileWorkerSingleJobHeader = 'S';
 // this is for the protocol, not the data.
 inline const int32 ShaderCompileWorkerPipelineJobHeader = 'P';
 
-// Modify this to invalidate _just_ the cache/DDC entries for individual shaders (will not cause shadermaps to rebuild if they are not otherwise out-of-date).
+// Modify this in ShaderCompilerCore.cpp to invalidate _just_ the cache/DDC entries for individual shaders (will not cause
+// shadermaps to rebuild if they are not otherwise out-of-date).
 // This should be bumped for changes to the FShaderCompilerOutput data structure (in addition to ShaderCompileWorkerOutputVersion)
-inline static const FGuid UE_SHADER_CACHE_VERSION = FGuid("B04528D6-1A2C-44C8-A742-E7D497D3562C");
+extern RENDERCORE_API const FGuid UE_SHADER_CACHE_VERSION;
 
 namespace UE::ShaderCompiler
 {
