Fix deadlock when running with r.VSync=True, r.GTSyncType=2 and r.OneFrameThreadLag=0
- The FFrameEndSync render command fences weren't being flushed after they were issued, meaning the game thread can wait on an event that will never be signalled.

diff --git a/Engine/Source/Runtime/RenderCore/Private/RenderingThread.cpp b/Engine/Source/Runtime/RenderCore/Private/RenderingThread.cpp
--- a/Engine/Source/Runtime/RenderCore/Private/RenderingThread.cpp
+++ b/Engine/Source/Runtime/RenderCore/Private/RenderingThread.cpp
@@ -1040,6 +1040,8 @@ void FRenderCommandFence::BeginFence(ESyncDepth SyncDepth)
 				// Subtracting 1 gives us the index of the frame that has *just* been presented.
 				RHITriggerTaskEventOnFlip(GRHIPresentCounter - 1, SwapchainEvent);
 			});
+
+			RHICmdList.ImmediateFlush(EImmediateFlushType::DispatchToRHIThread);
 		}
 		else if (SyncDepth == ESyncDepth::RHIThread)
 		{
