Cleaning up FTickableObjectRenderThread

diff --git a/Engine/Source/Runtime/RenderCore/Private/RenderingThread.cpp b/Engine/Source/Runtime/RenderCore/Private/RenderingThread.cpp
--- a/Engine/Source/Runtime/RenderCore/Private/RenderingThread.cpp
+++ b/Engine/Source/Runtime/RenderCore/Private/RenderingThread.cpp
@@ -147,10 +147,6 @@ FAutoConsoleVariable CVarRenderCommandPipeMode(
 		GRenderCommandPipeMode = GetValidatedRenderCommandPipeMode(Variable->GetInt());
 	}));
 
-/**
- * Tick all rendering thread tickable objects
- */
-
 /** How many cycles the renderthread used (excluding idle time). It's set once per frame in FViewport::Draw. */
 uint32 GRenderThreadTime = 0;
 /** How many cycles of wait time renderthread used. It's set once per frame in FViewport::Draw. */
diff --git a/Engine/Source/Runtime/RenderCore/Public/TickableObjectRenderThread.h b/Engine/Source/Runtime/RenderCore/Public/TickableObjectRenderThread.h
--- a/Engine/Source/Runtime/RenderCore/Public/TickableObjectRenderThread.h
+++ b/Engine/Source/Runtime/RenderCore/Public/TickableObjectRenderThread.h
@@ -75,9 +75,8 @@ public:
 
 	/**
 	 * Registers the object for ticking.
-	 * @param bIsRenderingThreadObject true if this object is owned by the rendering thread.
 	 */
-	void Register(bool bIsRenderingThreadObject = false)
+	void Register()
 	{
 		// make sure that only the rendering thread is attempting to add items to the RenderingThreadTickableObjects list
 		checkf(IsInRenderingThread(), TEXT("Game thread attempted to register an object in the RenderingThreadTickableObjects array."));
@@ -95,10 +94,15 @@ public:
 		bRegistered = true;
 	}
 
+	UE_DEPRECATED(5.7, "bIsRenderingThreadObject argument is no longer needed")
+	void Register(bool bIsRenderingThreadObject)
+	{
+		Register();
+	}
+
 	/**
 	 * Pure virtual that must be overloaded by the inheriting class. It will
-	 * be called from within LevelTick.cpp after ticking all actors or from
-	 * the rendering thread (depending on bIsRenderingThreadObject)
+	 * be called from the rendering thread.
 	 *
 	 * @param DeltaTime	Game time passed since the last call.
 	 */
