Adding the FRDGBuilder::SubmitBufferUploads breadcrumb to improve profiling

diff --git a/Engine/Source/Runtime/RenderCore/Private/RenderGraphBuilder.cpp b/Engine/Source/Runtime/RenderCore/Private/RenderGraphBuilder.cpp
--- a/Engine/Source/Runtime/RenderCore/Private/RenderGraphBuilder.cpp
+++ b/Engine/Source/Runtime/RenderCore/Private/RenderGraphBuilder.cpp
@@ -1870,7 +1870,7 @@ void FRDGBuilder::Execute()
 
 		UE::Tasks::FTaskEvent AllocateUploadBuffersTask{ UE_SOURCE_LOCATION };
 
-		UE::Tasks::FTask SubmitBufferUploadsTask = AddCommandListSetupTask([this, AllocateUploadBuffersTask] (FRHICommandListBase& RHICmdListTask) mutable
+		UE::Tasks::FTask SubmitBufferUploadsTask = AddCommandListSetupTask([this, AllocateUploadBuffersTask] (FRHICommandList& RHICmdListTask) mutable
 		{
 			SubmitBufferUploads(RHICmdListTask, &AllocateUploadBuffersTask);
 
@@ -2666,9 +2666,11 @@ void FRDGBuilder::CompilePassOps(FRDGPass* Pass)
 
 ///////////////////////////////////////////////////////////////////////////////////////////////////
 
-void FRDGBuilder::SubmitBufferUploads(FRHICommandListBase& RHICmdListUpload, UE::Tasks::FTaskEvent* AllocateUploadBuffersTask)
+void FRDGBuilder::SubmitBufferUploads(FRHICommandList& RHICmdListUpload, UE::Tasks::FTaskEvent* AllocateUploadBuffersTask)
 {
 	SCOPED_NAMED_EVENT_TEXT("FRDGBuilder::SubmitBufferUploads", FColor::Magenta);
+	RHI_BREADCRUMB_EVENT(RHICmdListUpload, "FRDGBuilder::SubmitBufferUploads");
+
 	FRHICommandListScopedFence ScopedFence(RHICmdListUpload);
 
 	{
diff --git a/Engine/Source/Runtime/RenderCore/Public/RenderGraphBuilder.h b/Engine/Source/Runtime/RenderCore/Public/RenderGraphBuilder.h
--- a/Engine/Source/Runtime/RenderCore/Public/RenderGraphBuilder.h
+++ b/Engine/Source/Runtime/RenderCore/Public/RenderGraphBuilder.h
@@ -1091,7 +1091,7 @@ private:
 
 	TArray<FUploadedBuffer, FRDGArrayAllocator> UploadedBuffers;
 
-	void SubmitBufferUploads(FRHICommandListBase& InRHICmdList, UE::Tasks::FTaskEvent* AllocateUploadBuffersTask = nullptr);
+	void SubmitBufferUploads(FRHICommandList& InRHICmdList, UE::Tasks::FTaskEvent* AllocateUploadBuffersTask = nullptr);
 
 	/////////////////////////////////////////////////////////////////////////////
 	// External Access Queue
