Retry (pre)loading shader code if the previous IoRequest failed.

- We can issue a request when an IoStore container from a DLC plugin is mounted. If the DLC has a shader group with the same FIoChunkId as the main game, an attempt to load base game shaders can result in IoStore using that container instead of base game's one. Then if the DLC is unmounted, loading of the base game mat will fail.
- We choose to address that by a retrial of the load (potentially with a hitch) rather than changing the IoStore request resolution logic (i.e. which container to choose), as the latter has more ramifications.
- If we ever need to limit the number of retries, another bit from the NumRefs can be taken to signify that a retry was attempted to avoid failing repeatedly. However as of now it is not clear if unlimited retries is a real problem, as all such scenarios will likely result in a crash anyway.

diff --git a/Engine/Source/Runtime/RenderCore/Private/ShaderCodeArchive.cpp b/Engine/Source/Runtime/RenderCore/Private/ShaderCodeArchive.cpp
--- a/Engine/Source/Runtime/RenderCore/Private/ShaderCodeArchive.cpp
+++ b/Engine/Source/Runtime/RenderCore/Private/ShaderCodeArchive.cpp
@@ -2077,23 +2077,26 @@ bool FIoStoreShaderCodeArchive::PreloadShaderGroup(int32 ShaderGroupIndex, FGrap
 	PreloadEntry.DebugInfo.Append(AppendInfo);
 #endif
 
-	if (PreloadEntry.IoRequest.Status() == FIoStatus::Invalid)
-	{
-		SetupPreloadEntryForLoading(ShaderGroupIndex, PreloadEntry);
-
-		// only global shaders are going to hit this path, all other shaders will be preloaded with the package
-		if (UNLIKELY(AttachShaderReadRequestFuncPtr == nullptr))
+	auto KickOffIoRequest = [this, AttachShaderReadRequestFuncPtr, &PreloadEntry, ShaderGroupIndex]()
 		{
-			FIoBatch IoBatch = IoDispatcher.NewBatch();
-			PreloadEntry.IoRequest = IoBatch.Read(Header.ShaderGroupIoHashes[ShaderGroupIndex], FIoReadOptions(), GetShaderCodeArchivePriority());
-			IoBatch.IssueAndDispatchSubsequents(PreloadEntry.PreloadEvent);
-		}
-		else
-		{
-			PreloadEntry.IoRequest = (*AttachShaderReadRequestFuncPtr)(Header.ShaderGroupIoHashes[ShaderGroupIndex], PreloadEntry.PreloadEvent);
-		}
+			SetupPreloadEntryForLoading(ShaderGroupIndex, PreloadEntry);
 
+			// only global shaders and retries of asset shaders are going to hit this path, regular asset shaders will be preloaded with the package
+			if (UNLIKELY(AttachShaderReadRequestFuncPtr == nullptr))
+			{
+				FIoBatch IoBatch = IoDispatcher.NewBatch();
+				PreloadEntry.IoRequest = IoBatch.Read(Header.ShaderGroupIoHashes[ShaderGroupIndex], FIoReadOptions(), GetShaderCodeArchivePriority());
+				IoBatch.IssueAndDispatchSubsequents(PreloadEntry.PreloadEvent);
+			}
+			else
+			{
+				PreloadEntry.IoRequest = (*AttachShaderReadRequestFuncPtr)(Header.ShaderGroupIoHashes[ShaderGroupIndex], PreloadEntry.PreloadEvent);
+			}
+		};
 
+	if (PreloadEntry.IoRequest.Status() == FIoStatus::Invalid)
+	{
+		KickOffIoRequest();
 
 		uint32 ShaderGroupSize = Header.ShaderGroupEntries[ShaderGroupIndex].CompressedSize + sizeof(FShaderGroupPreloadEntry);
 
@@ -2103,6 +2106,15 @@ bool FIoStoreShaderCodeArchive::PreloadShaderGroup(int32 ShaderGroupIndex, FGrap
 		CsvStatPreloadedShaderMB->Add((float)ShaderGroupSize / (1024.0f * 1024.0f));
 #endif
 	}
+	else if (PreloadEntry.PreloadEvent.IsValid() && PreloadEntry.PreloadEvent->IsCompleted() && PreloadEntry.IoRequest.GetResult() == nullptr)
+	{
+		// failed request, re-kick in-place, without changing stats
+		PreloadEntry.IoRequest.Cancel();
+		PreloadEntry.IoRequest = FIoRequest();
+		PreloadEntry.PreloadEvent.SafeRelease();
+
+		KickOffIoRequest();
+	}
 
 	if (AttachShaderReadRequestFuncPtr == nullptr && PreloadEntry.PreloadEvent && !PreloadEntry.PreloadEvent->IsComplete())
 	{
