Merging 44833691 to Main (//UE5/Main)
Update HDRGetMetaData for better compatibility with "fixed HDR setting" (PLATFORM_USES_FIXED_HDR_SETTING)

diff --git a/Engine/Source/Runtime/RenderCore/Private/RenderCore.cpp b/Engine/Source/Runtime/RenderCore/Private/RenderCore.cpp
--- a/Engine/Source/Runtime/RenderCore/Private/RenderCore.cpp
+++ b/Engine/Source/Runtime/RenderCore/Private/RenderCore.cpp
@@ -585,7 +585,12 @@ RENDERCORE_API void HDRGetMetaData(EDisplayOutputFormat& OutDisplayOutputFormat,
 	OutDisplayOutputFormat = HDRMetaData.DisplayOutputFormat;
 	OutDisplayColorGamut = HDRMetaData.DisplayColorGamut;
 	OutbHDRSupported = HDRMetaData.bHDRSupported;
+#if PLATFORM_USES_FIXED_HDR_SETTING
+	// IsHDREnabled() is controlled by code, depending on system and display HDR availability
+	if (!IsHDRAllowed())
+#else
 	if (!IsHDREnabled())
+#endif
 	{
 		return;
 	}
