[shaders] fix oversight in serializing shared FShaderCompilerEnvironment structs to SCW input files. these were unnecessarily serializing the full struct rather than just the compilation dependencies.

diff --git a/Engine/Source/Programs/ShaderCompileWorker/Private/ShaderCompileWorker.cpp b/Engine/Source/Programs/ShaderCompileWorker/Private/ShaderCompileWorker.cpp
--- a/Engine/Source/Programs/ShaderCompileWorker/Private/ShaderCompileWorker.cpp
+++ b/Engine/Source/Programs/ShaderCompileWorker/Private/ShaderCompileWorker.cpp
@@ -604,7 +604,7 @@ private:
 
 			for (int32 EnvironmentIndex = 0; EnvironmentIndex < NumSharedEnvironments; EnvironmentIndex++)
 			{
-				InputFile << SharedEnvironments[EnvironmentIndex];
+				SharedEnvironments[EnvironmentIndex].SerializeCompilationDependencies(InputFile);
 			}
 		}
 
diff --git a/Engine/Source/Runtime/RenderCore/Private/ShaderCompileWorkerUtil.cpp b/Engine/Source/Runtime/RenderCore/Private/ShaderCompileWorkerUtil.cpp
--- a/Engine/Source/Runtime/RenderCore/Private/ShaderCompileWorkerUtil.cpp
+++ b/Engine/Source/Runtime/RenderCore/Private/ShaderCompileWorkerUtil.cpp
@@ -447,7 +447,7 @@ bool FShaderCompileWorkerUtil::WriteTasks(const TArray<FShaderCommonCompileJobPt
 
 		for (int32 EnvironmentIndex = 0; EnvironmentIndex < SharedEnvironments.Num(); EnvironmentIndex++)
 		{
-			TransferFile << *SharedEnvironments[EnvironmentIndex];
+			SharedEnvironments[EnvironmentIndex]->SerializeCompilationDependencies(TransferFile);
 		}
 	}
 
