[HWRT] Add additional logging to RayTracingGeometry

#rnx

diff --git a/Engine/Source/Runtime/RenderCore/Private/RayTracingGeometry.cpp b/Engine/Source/Runtime/RenderCore/Private/RayTracingGeometry.cpp
--- a/Engine/Source/Runtime/RenderCore/Private/RayTracingGeometry.cpp
+++ b/Engine/Source/Runtime/RenderCore/Private/RayTracingGeometry.cpp
@@ -38,7 +38,7 @@ bool FRayTracingGeometry::HasValidInitializer() const
 		if (!Segment.VertexBuffer)
 		{
 			bAllSegmentsAreValid = false;
-			check(Initializer.OfflineData == nullptr);
+			checkf(Initializer.OfflineData == nullptr, TEXT("RayTracingGeometry (%s) should not have OfflineData"), *Initializer.DebugName.ToString());
 			break;
 		}
 	}
@@ -60,10 +60,11 @@ void FRayTracingGeometry::InitRHIForStreaming(FRHIRayTracingGeometry* Intermedia
 {
 	checkf(IntermediateGeometry,
 		TEXT("IntermediateGeometry should be valid when streaming-in ray tracing geometry.\n")
-		TEXT("This will result in FRayTracingGeometry not being correctly initialized.\n")
+		TEXT("This will result in FRayTracingGeometry (%s) not being correctly initialized.\n"),
+		*Initializer.DebugName.ToString()
 	);
 
-	checkf(RayTracingGeometryRHI, TEXT("RayTracingGeometryRHI must be valid when InitRHIForStreaming is called.\n"));
+	checkf(RayTracingGeometryRHI, TEXT("RayTracingGeometryRHI (%s) must be valid when InitRHIForStreaming is called.\n"), *Initializer.DebugName.ToString());
 
 	Initializer.Type = ERayTracingGeometryInitializerType::Rendering;
 
@@ -79,7 +80,7 @@ void FRayTracingGeometry::ReleaseRHIForStreaming(FRHIResourceReplaceBatcher& Bat
 {
 	RemoveBuildRequest();
 
-	checkf(RayTracingGeometryRHI, TEXT("RayTracingGeometryRHI must be valid when ReleaseRHIForStreaming is called.\n"));
+	checkf(RayTracingGeometryRHI, TEXT("RayTracingGeometryRHI (%s) must be valid when ReleaseRHIForStreaming is called.\n"), *Initializer.DebugName.ToString());
 
 	EnumRemoveFlags(GeometryState, EGeometryStateFlags::StreamedIn);
 	EnumRemoveFlags(GeometryState, EGeometryStateFlags::Valid);
@@ -102,13 +103,14 @@ void FRayTracingGeometry::RequestBuildIfNeeded(FRHICommandListBase& RHICmdList,
 
 void FRayTracingGeometry::MakeResident(FRHICommandList& RHICmdList)
 {
-	check(EnumHasAllFlags(GeometryState, EGeometryStateFlags::Evicted) && RayTracingGeometryRHI == nullptr);
+	checkf(EnumHasAllFlags(GeometryState, EGeometryStateFlags::Evicted) && RayTracingGeometryRHI == nullptr, 
+		TEXT("Evicted FRayTracingGeometry (%s) should have evicted flag set and no RHI object."), *Initializer.DebugName.ToString());
 	checkf(!EnumHasAllFlags(GeometryState, EGeometryStateFlags::StreamedIn),
-		TEXT("Evicted FRayTracingGeometry shouldn't have StreamedIn flag set."));
+		TEXT("Evicted FRayTracingGeometry (%s) shouldn't have StreamedIn flag set."), *Initializer.DebugName.ToString());
 
 	if (!ensureMsgf(DynamicGeometrySharedBufferGenerationID == NonSharedVertexBuffers,
 		TEXT("Cannot call MakeResident(...) on FRayTracingGeometry using shared vertex buffers.\n")
-		TEXT("Dynamic geometry should be rebuilt instead.")))
+		TEXT("Dynamic geometry (%s) should be rebuilt instead."), *Initializer.DebugName.ToString()))
 	{
 		// if geometry is using shared buffers those buffers might not be valid at this point
 		// instead of being made resident here, dynamic geometries need to be manually updated as necessary
@@ -122,8 +124,10 @@ void FRayTracingGeometry::MakeResident(FRHICommandList& RHICmdList)
 
 void FRayTracingGeometry::Evict()
 {
-	check(!EnumHasAllFlags(GeometryState, EGeometryStateFlags::Evicted) && RayTracingGeometryRHI != nullptr);
-	checkf(!EnumHasAllFlags(GeometryState, EGeometryStateFlags::StreamedIn), TEXT("Ray tracing geometry must be streamed out before it can be evicted."));
+	checkf(!EnumHasAllFlags(GeometryState, EGeometryStateFlags::Evicted) && RayTracingGeometryRHI != nullptr,
+		TEXT("RayTracingGeometry (%s) must not have been evicted already."), *Initializer.DebugName.ToString());
+	checkf(!EnumHasAllFlags(GeometryState, EGeometryStateFlags::StreamedIn), 
+		TEXT("RayTracingGeometry (%s) must be streamed out before it can be evicted."), *Initializer.DebugName.ToString());
 
 	RemoveBuildRequest();
 	RayTracingGeometryRHI.SafeRelease();
@@ -212,8 +216,8 @@ bool FRayTracingGeometry::IsValid() const
 
 	if (bIsValid)
 	{
-		check(Initializer.TotalPrimitiveCount > 0);
-		check(RayTracingGeometryRHI != nullptr || EnumHasAllFlags(GeometryState, EGeometryStateFlags::Evicted));
+		checkf(Initializer.TotalPrimitiveCount > 0, TEXT("Valid RayTracingGeometry (%s) must have non-zero primitive count."), *Initializer.DebugName.ToString());
+		checkf(RayTracingGeometryRHI != nullptr || EnumHasAllFlags(GeometryState, EGeometryStateFlags::Evicted), TEXT("Valid RayTracingGeometry (%s) must have valid RHI object or evicted flag set"), *Initializer.DebugName.ToString());
 	}
 
 	return bIsValid;
@@ -228,7 +232,7 @@ bool FRayTracingGeometry::IsEvicted() const
 
 	if (bIsEvicted)
 	{
-		check(RayTracingGeometryRHI == nullptr);
+		checkf(RayTracingGeometryRHI == nullptr, TEXT("Evicted RayTracingGeometry (%s) cannot have valid RHI object."), *Initializer.DebugName.ToString());
 	}
 
 	return bIsEvicted;
@@ -265,7 +269,7 @@ void FRayTracingGeometry::RemoveBuildRequest()
 
 void FRayTracingGeometry::InitResource(FRHICommandListBase& RHICmdList)
 {
-	ensureMsgf(IsRayTracingAllowed(), TEXT("FRayTracingGeometry should only be initialized when Ray Tracing is allowed."));
+	ensureMsgf(IsRayTracingAllowed(), TEXT("FRayTracingGeometry (%s) should only be initialized when Ray Tracing is allowed."), *Initializer.DebugName.ToString());
 
 	FRenderResource::InitResource(RHICmdList);
 
@@ -277,7 +281,7 @@ void FRayTracingGeometry::InitResource(FRHICommandListBase& RHICmdList)
 
 void FRayTracingGeometry::ReleaseResource()
 {
-	ensureMsgf(IsRayTracingAllowed() || !IsInitialized(), TEXT("FRayTracingGeometry should only be initialized when Ray Tracing is allowed."));
+	ensureMsgf(IsRayTracingAllowed() || !IsInitialized(), TEXT("FRayTracingGeometry (%s) should only be initialized when Ray Tracing is allowed."), *Initializer.DebugName.ToString());
 
 	if (RayTracingGeometryHandle != INDEX_NONE)
 	{
@@ -298,7 +302,7 @@ bool FRayTracingGeometry::HasPendingBuildRequest() const
 
 	if (bHasPendingBuildRequest)
 	{
-		ensure(IsValid() && !IsEvicted());
+		ensureMsgf(IsValid() && !IsEvicted(), TEXT("RayTracingGeometry (%s) with pending build request must be valid and not evicted."), *Initializer.DebugName.ToString());
 	}
 
 	return bHasPendingBuildRequest;
@@ -306,7 +310,7 @@ bool FRayTracingGeometry::HasPendingBuildRequest() const
 
 void FRayTracingGeometry::BoostBuildPriority(float InBoostValue) const
 {
-	check(HasPendingBuildRequest());
+	checkf(HasPendingBuildRequest(), TEXT("RayTracingGeometry (%s) must have pending build request"), *Initializer.DebugName.ToString());
 	GRayTracingGeometryManager->BoostPriority(RayTracingBuildRequestIndex, InBoostValue);
 }
 
