Use forward declaration for enum type to avoid Engine include in RenderViewportFeedback.

diff --git a/Engine/Source/Runtime/RenderCore/Internal/RenderViewportFeedback.h b/Engine/Source/Runtime/RenderCore/Internal/RenderViewportFeedback.h
--- a/Engine/Source/Runtime/RenderCore/Internal/RenderViewportFeedback.h
+++ b/Engine/Source/Runtime/RenderCore/Internal/RenderViewportFeedback.h
@@ -7,7 +7,8 @@
 #include "Async/Mutex.h"
 #include "Templates/SharedPointer.h"
 #include "Templates/UniquePtr.h"
-#include "Engine/Scene.h"
+
+enum EBloomMethod : int; // from Engine/Scene.h
 
 /** 
   * Feedback from the renderer to the viewport, gathered while rendering. 
diff --git a/Engine/Source/Runtime/RenderCore/RenderCore.Build.cs b/Engine/Source/Runtime/RenderCore/RenderCore.Build.cs
--- a/Engine/Source/Runtime/RenderCore/RenderCore.Build.cs
+++ b/Engine/Source/Runtime/RenderCore/RenderCore.Build.cs
@@ -38,7 +38,7 @@ public class RenderCore : ModuleRules
 
 		PrivateDependencyModuleNames.AddRange(new string[] { "Core", "Projects", "ApplicationCore", "TraceLog", "CookOnTheFly" });
 		
-		PublicIncludePathModuleNames.AddRange(new string[] { "Engine", "RHI" });
+		PublicIncludePathModuleNames.AddRange(new string[] { "RHI" });
 
 		if (Target.bBuildEditor == true)
 		{
