[shaders] add additional diagnostic info when pipeline jobs fail to dump debug info (and downgrade what was a check to a warning; this shouldn't fail the cook)

diff --git a/Engine/Source/Developer/ShaderCompilerCommon/Private/ShaderCompilerCommon.cpp b/Engine/Source/Developer/ShaderCompilerCommon/Private/ShaderCompilerCommon.cpp
--- a/Engine/Source/Developer/ShaderCompilerCommon/Private/ShaderCompilerCommon.cpp
+++ b/Engine/Source/Developer/ShaderCompilerCommon/Private/ShaderCompilerCommon.cpp
@@ -2051,8 +2051,7 @@ namespace UE::ShaderCompilerCommon
 		FString DebugSourcePath = Options.GetDebugShaderPath(Input, Suffix);
 		FFileHelper::SaveStringToFile(Contents, *DebugSourcePath);
 
-		Ctx.DebugSourceFiles.Add(MoveTemp(DebugSourcePath));
-
+		Ctx.DebugSourceFiles.Add(Input.Target.GetFrequency(), MoveTemp(DebugSourcePath));
 	}
 
 	void DumpDebugShaderData(const FShaderCompilerInput& Input, const FString& PreprocessedSource, const FDebugShaderDataOptions& Options)
diff --git a/Engine/Source/Runtime/RenderCore/Private/ShaderCompileWorkerUtil.cpp b/Engine/Source/Runtime/RenderCore/Private/ShaderCompileWorkerUtil.cpp
--- a/Engine/Source/Runtime/RenderCore/Private/ShaderCompileWorkerUtil.cpp
+++ b/Engine/Source/Runtime/RenderCore/Private/ShaderCompileWorkerUtil.cpp
@@ -635,8 +635,9 @@ static void WriteShaderCompileWorkerDebugCommandLine(FShaderCommonCompileJob& Jo
 	CmdLine << " -DebugSourceFiles=";
 	TArray<FString> RelativeSourcePaths;
 	RelativeSourcePaths.Reserve(Ctx.DebugSourceFiles.Num());
-	for (FString SourceFile : Ctx.DebugSourceFiles) // note: intentional copy of path here since MakePathRelativeTo modifies in-place
+	for (const TPair<EShaderFrequency, FString>& SourceFilePair : Ctx.DebugSourceFiles) 
 	{
+		FString SourceFile = SourceFilePair.Value; // note: intentional copy of path here since MakePathRelativeTo modifies in-place
 		// as above this may refer to multiple source files for different stages of a pipeline job
 		// so make all paths relative to the working directory for this specific job.
 		FPaths::MakePathRelativeTo(SourceFile, JobArgsPath.ToString());
diff --git a/Engine/Source/Runtime/RenderCore/Private/ShaderCore.cpp b/Engine/Source/Runtime/RenderCore/Private/ShaderCore.cpp
--- a/Engine/Source/Runtime/RenderCore/Private/ShaderCore.cpp
+++ b/Engine/Source/Runtime/RenderCore/Private/ShaderCore.cpp
@@ -4745,8 +4745,29 @@ void FShaderPipelineCompileJob::OnComplete(FShaderDebugDataContext& Ctx)
 	bool bShouldDump = !Ctx.DebugSourceFiles.IsEmpty() && (CVarDumpDebugInfoForCacheHits.GetValueOnAnyThread() || !JobStatusPtr->WasCompilationSkipped());
 	if (bShouldDump)
 	{
-		check(Ctx.DebugSourceFiles.Num() == StageJobs.Num());
-		FShaderCompileWorkerUtil::DumpDebugCompileInput(*this, Ctx);
+		if (Ctx.DebugSourceFiles.Num() == StageJobs.Num())
+		{
+			FShaderCompileWorkerUtil::DumpDebugCompileInput(*this, Ctx);
+		}
+		else
+		{
+			TArray<const FShaderCompileJob*> MissingStages;
+			for (int32 Index = 0, Num = StageJobs.Num(); Index < Num; ++Index)
+			{
+				FShaderCompileJob* StageJob = StageJobs[Index];
+				if (!Ctx.DebugSourceFiles.Contains(StageJob->Input.Target.GetFrequency()))
+				{
+					MissingStages.Add(StageJob);
+				}
+			}
+			FString MissingStagesStr = FString::JoinBy(MissingStages, TEXT("\n"), [](const FShaderCompileJob* StageJob)
+				{
+					bool bDebugInfoEnabled = StageJob->Input.DumpDebugInfoEnabled();
+					return FString::Printf(TEXT("%s (debug info %s)"), *StageJob->Input.DebugGroupName, (bDebugInfoEnabled ? TEXT("enabled") : TEXT("disabled")));
+				}
+			);
+			UE_LOG(LogShaders, Warning, TEXT("Cannot dump debug compile input for pipeline job; the following stage job(s) failed to dump source files: \n%s"), *MissingStagesStr);
+		}
 	}
 }
 
diff --git a/Engine/Source/Runtime/RenderCore/Public/ShaderCompilerCore.h b/Engine/Source/Runtime/RenderCore/Public/ShaderCompilerCore.h
--- a/Engine/Source/Runtime/RenderCore/Public/ShaderCompilerCore.h
+++ b/Engine/Source/Runtime/RenderCore/Public/ShaderCompilerCore.h
@@ -197,7 +197,7 @@ using FShaderCompilerInputHash = FBlake3Hash;
 struct FShaderDebugDataContext
 {
 	bool bIsPipeline = false;
-	TArray<FString> DebugSourceFiles;
+	TMap<EShaderFrequency, FString> DebugSourceFiles;
 };
 
 /** Struct that gathers all readonly inputs needed for the compilation of a single shader. */
