Fixed assert due to race condition with refcounting. The assert should check the state inside of the branch for ref 0.

diff --git a/Engine/Source/Runtime/RenderCore/Private/RenderGraphResourcePool.cpp b/Engine/Source/Runtime/RenderCore/Private/RenderGraphResourcePool.cpp
--- a/Engine/Source/Runtime/RenderCore/Private/RenderGraphResourcePool.cpp
+++ b/Engine/Source/Runtime/RenderCore/Private/RenderGraphResourcePool.cpp
@@ -261,9 +261,10 @@ uint32 FRDGTransientRenderTarget::AddRef() const
 uint32 FRDGTransientRenderTarget::Release()
 {
 	const int32 Refs = FPlatformAtomics::InterlockedDecrement(&RefCount);
-	check(Refs >= 0 && LifetimeState == ERDGTransientResourceLifetimeState::Allocated);
+	check(Refs >= 0);
 	if (Refs == 0)
 	{
+		check(LifetimeState == ERDGTransientResourceLifetimeState::Allocated);
 		if (GRDGTransientResourceAllocator.IsValid())
 		{
 			GRDGTransientResourceAllocator.AddPendingDeallocation(this);
