Removing duplicate reference to TargetPlatform.

diff --git a/Engine/Source/Runtime/RenderCore/RenderCore.Build.cs b/Engine/Source/Runtime/RenderCore/RenderCore.Build.cs
--- a/Engine/Source/Runtime/RenderCore/RenderCore.Build.cs
+++ b/Engine/Source/Runtime/RenderCore/RenderCore.Build.cs
@@ -49,9 +49,6 @@ public class RenderCore : ModuleRules
 			PrivateIncludePathModuleNames.Add("DerivedDataCache");
 		}
 
-		// Added in Dev-VT, still needed?
-		PrivateIncludePathModuleNames.AddRange(new string[] { "TargetPlatform" });
-
 		if (Target.IsInPlatformGroup(UnrealPlatformGroup.Windows)
 			|| Target.IsInPlatformGroup(UnrealPlatformGroup.Android)
 			|| Target.IsInPlatformGroup(UnrealPlatformGroup.Linux))
