Change r.Substrate.PlatformGBufferFormat into bSupportsAdaptiveGBuffer, a DataDrivenPlatform variable.

This fixes crashes when previewing platforms not supporting adaptive GBuffer in editor while the editor run on a platform supporting adaptive GBuffer. This is was due to preview settings stomping CVars when enabled, causing mismatch configuration when other view (e.g., thumbnail) where running at the same time as the preview.

diff --git a/Engine/Config/Linux/BaseLinuxEngine.ini b/Engine/Config/Linux/BaseLinuxEngine.ini
--- a/Engine/Config/Linux/BaseLinuxEngine.ini
+++ b/Engine/Config/Linux/BaseLinuxEngine.ini
@@ -21,4 +21,3 @@ r.Substrate.Glints=1
 r.Substrate.ClosuresPerPixel=4
 r.Substrate.OpaqueMaterialRoughRefraction=0
 r.Refraction.Blur=1
-r.Substrate.PlatformGBufferFormat=1
diff --git a/Engine/Config/LinuxArm64/BaseLinuxArm64Engine.ini b/Engine/Config/LinuxArm64/BaseLinuxArm64Engine.ini
--- a/Engine/Config/LinuxArm64/BaseLinuxArm64Engine.ini
+++ b/Engine/Config/LinuxArm64/BaseLinuxArm64Engine.ini
@@ -11,4 +11,3 @@ r.Substrate.Glints=1
 r.Substrate.ClosuresPerPixel=4
 r.Substrate.OpaqueMaterialRoughRefraction=0
 r.Refraction.Blur=1
-r.Substrate.PlatformGBufferFormat=1
diff --git a/Engine/Config/Mac/DataDrivenPlatformInfo.ini b/Engine/Config/Mac/DataDrivenPlatformInfo.ini
--- a/Engine/Config/Mac/DataDrivenPlatformInfo.ini
+++ b/Engine/Config/Mac/DataDrivenPlatformInfo.ini
@@ -105,6 +105,7 @@ bSupportsBindless=true
 bSupportsMeshShadersTier0=false
 bSupportsMeshShadersTier1=false
 bSupportsMeshShadersWithClipDistance=false
+bSupportsAdaptiveGBuffer=true
 MaxMeshShaderThreadGroupSize=128
 bSupportsRealTypes=RuntimeGuaranteed
 FriendlyName=LOCTEXT("FriendlyShaderPlatformName_MacSM6", "Mac Metal SM6")
diff --git a/Engine/Config/VulkanPC/DataDrivenPlatformInfo.ini b/Engine/Config/VulkanPC/DataDrivenPlatformInfo.ini
--- a/Engine/Config/VulkanPC/DataDrivenPlatformInfo.ini
+++ b/Engine/Config/VulkanPC/DataDrivenPlatformInfo.ini
@@ -245,6 +245,8 @@ bSupportsVertexShaderLayer=true
 
 bSupportsAnisotropicMaterials=true
 
+bSupportsAdaptiveGBuffer=true
+
 ; disable until DXC -> SPRIV ROV support is implemented
 bSupportsROV=false
 bSupportsOIT=false
diff --git a/Engine/Config/Windows/BaseWindowsEngine.ini b/Engine/Config/Windows/BaseWindowsEngine.ini
--- a/Engine/Config/Windows/BaseWindowsEngine.ini
+++ b/Engine/Config/Windows/BaseWindowsEngine.ini
@@ -42,7 +42,6 @@ r.Substrate.Glints=1
 r.Substrate.ClosuresPerPixel=4
 r.Substrate.OpaqueMaterialRoughRefraction=0
 r.Refraction.Blur=1
-r.Substrate.PlatformGBufferFormat=1
 
 [OnlineSubsystem]
 LocalPlatformName=WIN
diff --git a/Engine/Config/Windows/DataDrivenPlatformInfo.ini b/Engine/Config/Windows/DataDrivenPlatformInfo.ini
--- a/Engine/Config/Windows/DataDrivenPlatformInfo.ini
+++ b/Engine/Config/Windows/DataDrivenPlatformInfo.ini
@@ -129,6 +129,7 @@ SupportsMultiViewport=RuntimeGuaranteed
 bSupportsDualSourceBlending=true
 bSupportsNNEShaders=true
 bSupportsIndependentSamplers=true
+bSupportsAdaptiveGBuffer=true
 MaxSamplers=32
 SupportsBarycentricsSemantic=RuntimeDependent
 bSupportsWorkGraphs=true
diff --git a/Engine/Source/Runtime/Core/Private/UObject/FortniteMainBranchObjectVersion.cpp b/Engine/Source/Runtime/Core/Private/UObject/FortniteMainBranchObjectVersion.cpp
--- a/Engine/Source/Runtime/Core/Private/UObject/FortniteMainBranchObjectVersion.cpp
+++ b/Engine/Source/Runtime/Core/Private/UObject/FortniteMainBranchObjectVersion.cpp
@@ -13,7 +13,7 @@ TMap<FGuid, FGuid> FFortniteMainBranchObjectVersion::GetSystemGuids()
 	SystemGuids.Add(DevGuids.MATERIALSHADERMAP_DERIVEDDATA_VER, FGuid("F051877974B44D409435F321BDB39DBB"));
 	SystemGuids.Add(DevGuids.NIAGARASHADERMAP_DERIVEDDATA_VER, FGuid("6360A977062842A29ED30E3A7ACB0E64"));
 	SystemGuids.Add(DevGuids.Niagara_LatestScriptCompileVersion, FGuid("783FF3700BB641C799B2D43552FBF91E"));
-	SystemGuids.Add(DevGuids.MaterialTranslationDDCVersion, FGuid("400CBC2A3FF64455800F843453D85DE3"));
+	SystemGuids.Add(DevGuids.MaterialTranslationDDCVersion, FGuid("7801FD563B40414CB305CCFFB349F2CB"));
 
 	// These GUIDs are in AppendRestrictedFortniteMainStreamObjectVersionGuids as they should rarely be changed.
 	//SystemGuids.Add(DevGuids.NANITE_DERIVEDDATA_VER, FGuid("ACF98184978742A592537DA223A1E6BE"));
diff --git a/Engine/Source/Runtime/RHI/Private/DataDrivenShaderPlatformInfo.cpp b/Engine/Source/Runtime/RHI/Private/DataDrivenShaderPlatformInfo.cpp
--- a/Engine/Source/Runtime/RHI/Private/DataDrivenShaderPlatformInfo.cpp
+++ b/Engine/Source/Runtime/RHI/Private/DataDrivenShaderPlatformInfo.cpp
@@ -320,6 +320,7 @@ void FGenericDataDrivenShaderPlatformInfo::ParseDataDrivenShaderInfo(const FConf
 	GET_SECTION_BOOL_HELPER(bSupportsWorkGraphs);
 	GET_SECTION_BOOL_HELPER(bSupportsWorkGraphsTier1_1);
 	GET_SECTION_BOOL_HELPER(bSupportsDLSSShaders);
+	GET_SECTION_BOOL_HELPER(bSupportsAdaptiveGBuffer);
 #undef GET_SECTION_BOOL_HELPER
 #undef GET_SECTION_INT_HELPER
 #undef GET_SECTION_SUPPORT_HELPER
diff --git a/Engine/Source/Runtime/RHI/Public/DataDrivenShaderPlatformInfo.h b/Engine/Source/Runtime/RHI/Public/DataDrivenShaderPlatformInfo.h
--- a/Engine/Source/Runtime/RHI/Public/DataDrivenShaderPlatformInfo.h
+++ b/Engine/Source/Runtime/RHI/Public/DataDrivenShaderPlatformInfo.h
@@ -132,6 +132,7 @@ class FGenericDataDrivenShaderPlatformInfo
 	uint32 bSupportsWorkGraphs : 1;
 	uint32 bSupportsWorkGraphsTier1_1 : 1;
 	uint32 bSupportsDLSSShaders : 1;
+	uint32 bSupportsAdaptiveGBuffer : 1;
 
 	// NOTE: When adding fields, you must also add to ParseDataDrivenShaderInfo!
 	uint32 bContainsValidPlatformInfo : 1;
@@ -869,6 +870,12 @@ PRAGMA_ENABLE_DEPRECATION_WARNINGS
 		return Infos[Platform].bSupportsDLSSShaders;
 	}
 
+	static inline const bool GetSupportsAdaptiveGBuffer(const FStaticShaderPlatform Platform)
+	{
+		check(IsValid(Platform));
+		return Infos[Platform].bSupportsAdaptiveGBuffer;
+	}
+	
 	static inline const bool IsValid(const FStaticShaderPlatform Platform)
 	{
 		return Infos[Platform].bContainsValidPlatformInfo;
diff --git a/Engine/Source/Runtime/RenderCore/Private/RenderUtils.cpp b/Engine/Source/Runtime/RenderCore/Private/RenderUtils.cpp
--- a/Engine/Source/Runtime/RenderCore/Private/RenderUtils.cpp
+++ b/Engine/Source/Runtime/RenderCore/Private/RenderUtils.cpp
@@ -1915,14 +1915,6 @@ static TAutoConsoleVariable<int32> CVarSubstrate(
 	TEXT("Enable Substrate materials (Beta)."),
 	ECVF_ReadOnly | ECVF_RenderThreadSafe);
 
-static TAutoConsoleVariable<int32> CVarSubstratePlatformGBufferFormat(
-	TEXT("r.Substrate.PlatformGBufferFormat"),
-	0, // Substrate 60fps by default -> Blendable GBuffer 
-	TEXT("Define GBuffer format supported by this platforms.\n \
-		 - 0: Substrate materials are fit into a blendable GBuffer. This is faster and DBuffer decals are no longer required. However, lots of visual effects and fidelity are lost. To be able to cast to blendable GBuffer, this command enforce ClosuresPerPixel=1.\n \
-		 - 1: Substrate materials are fit into a bitstream GBuffer. This stores more data to enable complex lighting behavior per pixel. It is more costly than the blendable GBuffer format. This is only support on certain platforms. Non-supported platforms will automatically fallback onto a blendable GBuffer format."),
-	ECVF_ReadOnly | ECVF_RenderThreadSafe | ECVF_Preview);
-
 static TAutoConsoleVariable<int32> CVarSubstrateProjectGBufferFormat(
 	TEXT("r.Substrate.ProjectGBufferFormat"),
 	1,
@@ -2043,11 +2035,11 @@ namespace Substrate
 
 	bool IsSubstrateBlendableGBufferEnabled(EShaderPlatform InPlatform)
 	{
-		static FShaderPlatformCachedIniValue<int32> CVarPlatformGBufferFormat(TEXT("r.Substrate.PlatformGBufferFormat"));
-		const bool bUseMaterialBuffer = CVarPlatformGBufferFormat.Get(InPlatform) > 0 && CVarSubstrateProjectGBufferFormat.GetValueOnAnyThread() > 0;
+		const bool bPlatformSupportsAdaptiveGBuffer = FDataDrivenShaderPlatformInfo::GetSupportsAdaptiveGBuffer(InPlatform);
+		const bool bProjectRequestsAdaptiveGBuffer = CVarSubstrateProjectGBufferFormat.GetValueOnAnyThread() > 0;
 		
 		// Mobile always uses blendable GBuffer, otherwise it is enabled on selected platforms.
-		return (IsMobilePlatform(InPlatform) && IsMobileDeferredShadingEnabled(InPlatform)) || !bUseMaterialBuffer;
+		return (IsMobilePlatform(InPlatform) && IsMobileDeferredShadingEnabled(InPlatform)) || !(bPlatformSupportsAdaptiveGBuffer && bProjectRequestsAdaptiveGBuffer);
 	}
 
 	bool IsStochasticLightingEnabled(EShaderPlatform InPlatform)
