[Shaders] Fix pipeling skipping work if number of shaders is less than number of tasks

#tests Preflight

diff --git a/Engine/Source/Runtime/RenderCore/Private/PipelineCacheUtilities.cpp b/Engine/Source/Runtime/RenderCore/Private/PipelineCacheUtilities.cpp
--- a/Engine/Source/Runtime/RenderCore/Private/PipelineCacheUtilities.cpp
+++ b/Engine/Source/Runtime/RenderCore/Private/PipelineCacheUtilities.cpp
@@ -701,7 +701,7 @@ bool UE::PipelineCacheUtilities::LoadStablePipelineCacheFile(const FString& File
 
 	FGraphEventArray RemappingStableKeysTasks;
 	int32 NumRemappingTasks = FPlatformMisc::NumberOfWorkerThreadsToSpawn();	// match the number of working threads, but essentially can be any number
-	int32 NumKeysPerTask = SavedStableKeys.Num() / NumRemappingTasks;
+	int32 NumKeysPerTask = FMath::Max(SavedStableKeys.Num() / NumRemappingTasks, 1);
 
 	for (int32 FirstKey = 0; FirstKey < SavedStableKeys.Num();)
 	{
