Change r.Mobile.AllowPerPixelShadingModels and r.Mobile.ShadingModelsMask to be able to set them on a per platform basis

#rnx

diff --git a/Engine/Source/Runtime/RenderCore/Private/RenderUtils.cpp b/Engine/Source/Runtime/RenderCore/Private/RenderUtils.cpp
--- a/Engine/Source/Runtime/RenderCore/Private/RenderUtils.cpp
+++ b/Engine/Source/Runtime/RenderCore/Private/RenderUtils.cpp
@@ -535,8 +535,8 @@ RENDERCORE_API bool AllowPerPixelShadingModels(const FStaticShaderPlatform Platf
 {
 	if (IsMobilePlatform(Platform))
 	{
-		static const auto CVar = IConsoleManager::Get().FindTConsoleVariableDataInt(TEXT("r.Mobile.AllowPerPixelShadingModels"));
-		return CVar->GetValueOnAnyThread() != 0;
+		static FShaderPlatformCachedIniValue<int32> CVarAllowPerPixelShadingModels(TEXT("r.Mobile.AllowPerPixelShadingModels"));
+		return CVarAllowPerPixelShadingModels.Get(Platform) != 0;
 	}
 	return true;
 }
@@ -545,8 +545,8 @@ RENDERCORE_API uint32 GetPlatformShadingModelsMask(const FStaticShaderPlatform P
 {
 	if (IsMobilePlatform(Platform))
 	{
-		static const auto CVar = IConsoleManager::Get().FindTConsoleVariableDataInt(TEXT("r.Mobile.ShadingModelsMask"));
-		return CVar->GetValueOnAnyThread();
+		static FShaderPlatformCachedIniValue<int32> CVarMobileShadingModelsMask(TEXT("r.Mobile.ShadingModelsMask"));
+		return static_cast<uint32>(CVarMobileShadingModelsMask.Get(Platform));
 	}
 	return 0xFFFFFFFF;
 }
