Minor fix for custom ShaderTypeStats customized name
Needs to lookup cvar by ShaderFormat

diff --git a/Engine/Source/Developer/ShaderCompilerCommon/Private/ShaderSymbolExport.cpp b/Engine/Source/Developer/ShaderCompilerCommon/Private/ShaderSymbolExport.cpp
--- a/Engine/Source/Developer/ShaderCompilerCommon/Private/ShaderSymbolExport.cpp
+++ b/Engine/Source/Developer/ShaderCompilerCommon/Private/ShaderSymbolExport.cpp
@@ -87,7 +87,7 @@ void FShaderSymbolExport::Initialize()
 
 		// Setup ExportFileName to the default or the cvar override
 		ExportFileName = TEXT("ShaderSymbols");
-		GetShaderFileNameOverride(ExportFileName, TEXT("r.Shaders.SymbolFileNameOverride"), ShaderFormat);
+		GetShaderFileNameOverride(ExportFileName, TEXT("r.Shaders.SymbolFileNameOverride"), ShaderFormat, ShaderFormat);
 
 		IPlatformFile& PlatformFile = FPlatformFileManager::Get().GetPlatformFile();
 		bExportShaderSymbols = PlatformFile.CreateDirectoryTree(*ExportPath);
diff --git a/Engine/Source/Runtime/RenderCore/Private/ShaderCodeLibrary.cpp b/Engine/Source/Runtime/RenderCore/Private/ShaderCodeLibrary.cpp
--- a/Engine/Source/Runtime/RenderCore/Private/ShaderCodeLibrary.cpp
+++ b/Engine/Source/Runtime/RenderCore/Private/ShaderCodeLibrary.cpp
@@ -3878,8 +3878,9 @@ public:
 				}
 
 				// Add a copy of the type stats to the build metadata as well, potentially overriding name by cvar
+				const FName ShaderFormat = FDataDrivenShaderPlatformInfo::GetShaderFormat((EShaderPlatform)PlatformId);
 				FString StatsFileName = FString::Printf(TEXT("ShaderTypeStats-%s"), *PlatformName);
-				GetShaderFileNameOverride(StatsFileName, TEXT("r.Shaders.ShaderTypeStatsFileNameOverride"), PlatformFName);
+				GetShaderFileNameOverride(StatsFileName, TEXT("r.Shaders.ShaderTypeStatsFileNameOverride"), ShaderFormat, PlatformFName);
 				FString MetaPath = FString::Printf(TEXT("%s/%s.csv"), *MetaDataDir, *StatsFileName);
 
 				if (IFileManager::Get().Copy(*MetaPath, *Filename) != COPY_OK)
diff --git a/Engine/Source/Runtime/RenderCore/Private/ShaderCore.cpp b/Engine/Source/Runtime/RenderCore/Private/ShaderCore.cpp
--- a/Engine/Source/Runtime/RenderCore/Private/ShaderCore.cpp
+++ b/Engine/Source/Runtime/RenderCore/Private/ShaderCore.cpp
@@ -957,7 +957,7 @@ bool GetShaderSymbolPathOverride(FString& OutPathOverride, FName ShaderFormat)
 	return false;
 }
 
-bool GetShaderFileNameOverride(FString& OutFileNameOverride, const TCHAR* Cvar, FName ShaderFormat)
+bool GetShaderFileNameOverride(FString& OutFileNameOverride, const TCHAR* Cvar, FName ShaderFormat, FName PlatformName)
 {
 	FString OutFileName;
 	const FShaderSymbolSettingHelper SymbolFileNameOverride(Cvar);
@@ -966,7 +966,7 @@ bool GetShaderFileNameOverride(FString& OutFileNameOverride, const TCHAR* Cvar,
 		if (!OutFileName.IsEmpty())
 		{
 			// Allow the user to specify the location of the per-platform string.
-			OutFileName = OutFileName.Replace(TEXT("{Platform}"), *ShaderFormat.ToString(), ESearchCase::IgnoreCase);
+			OutFileName = OutFileName.Replace(TEXT("{Platform}"), *PlatformName.ToString(), ESearchCase::IgnoreCase);
 
 			// Allow adding name of the DLC if any to differentiate cooks, and one or more optional formatting dashes
 			FString DLCName;
diff --git a/Engine/Source/Runtime/RenderCore/Public/ShaderCompilerCore.h b/Engine/Source/Runtime/RenderCore/Public/ShaderCompilerCore.h
--- a/Engine/Source/Runtime/RenderCore/Public/ShaderCompilerCore.h
+++ b/Engine/Source/Runtime/RenderCore/Public/ShaderCompilerCore.h
@@ -75,8 +75,9 @@ extern RENDERCORE_API bool ShouldWriteShaderSymbols(FName ShaderFormat);
 extern RENDERCORE_API bool GetShaderSymbolPathOverride(FString& OutPathOverride, FName ShaderFormat);
 
 /** Returns true if the shader file name is overridden by the Cvar and OutFileNameOverride contains the override filename.
+    ShaderFormat used for format specific cvar, PlatformName used for {Platform} string substitution, if requested.
     Returns false otherwise and leaves OutFileNameOverride unchanged */
-extern RENDERCORE_API bool GetShaderFileNameOverride(FString& OutFileNameOverride, const TCHAR* Cvar, FName ShaderFormat);
+extern RENDERCORE_API bool GetShaderFileNameOverride(FString& OutFileNameOverride, const TCHAR* Cvar, FName ShaderFormat, FName PlatformName);
 
 /** Returns true if (external) shader symbols should be specific to each shader rather than be de-duplicated. */
 extern RENDERCORE_API bool ShouldAllowUniqueShaderSymbols(FName ShaderFormat);
