Add Substrate's rough difffuse to shader key in order to invalidate all shaders when rough diffuse is enabled.

#test EngineTest editor dx12

diff --git a/Engine/Source/Runtime/RenderCore/Private/Shader.cpp b/Engine/Source/Runtime/RenderCore/Private/Shader.cpp
--- a/Engine/Source/Runtime/RenderCore/Private/Shader.cpp
+++ b/Engine/Source/Runtime/RenderCore/Private/Shader.cpp
@@ -2529,7 +2529,9 @@ void ShaderMapAppendKey(EShaderPlatform Platform, FShaderKeyGenerator& KeyGen)
 
 	{
 		static const auto CVar = IConsoleManager::Get().FindTConsoleVariableDataInt(TEXT("r.Material.RoughDiffuse"));
-		if (CVar && CVar->GetValueOnAnyThread() > 0)
+		const bool bSubstrateRough = Substrate::IsSubstrateEnabled() && Substrate::IsRoughDiffuseEnabled(Platform);
+		const bool bMaterialtRough = CVar && CVar->GetValueOnAnyThread() > 0;
+		if (bSubstrateRough || bMaterialtRough)
 		{
 			KeyGen.AppendSeparator();
 			KeyGen.Append(TEXT("MATRDIFF"));
