Fix crash on shutdown when a FShaderMapResource_SharedCode was still referenced:
FShaderLibrariesCollection::GatherShaderMapResourceReferences is trying to access NamedLibrariesStack while it's being destroyed.
The flow is FShaderCodeLibrary::Shutdown -> ~FShaderLibraryInstance through implicit destruction FShaderLibrariesCollection::NamedLibrariesStack -> FNamedShaderLibrary::Components
in ~FShaderLibraryInstance, we call BeginReleaseResource which ends up calling FShaderMapResource_SharedCode::ReleaseRHI directly (engine shutdown).
This is where we end up calling DiagnoseLeakingShaderMapResourceRefs to look for leaks through GatherShaderMapResourceReferences.
This CL doesn't prevent access to FShaderLibrariesCollection being destructed, but at least NamedLibrariesStack is in a good state

diff --git a/Engine/Source/Runtime/RenderCore/Private/ShaderCodeLibrary.cpp b/Engine/Source/Runtime/RenderCore/Private/ShaderCodeLibrary.cpp
--- a/Engine/Source/Runtime/RenderCore/Private/ShaderCodeLibrary.cpp
+++ b/Engine/Source/Runtime/RenderCore/Private/ShaderCodeLibrary.cpp
@@ -3002,6 +3002,21 @@ public:
 			IConsoleManager::Get().UnregisterConsoleObject(DumpLibraryContentsCmd);
 		}
 #endif
+
+		TArray<FString> LibraryNames;
+		{
+			FRWScopeLock ReadLock(NamedLibrariesMutex, SLT_ReadOnly);
+			LibraryNames.Reserve(NamedLibrariesStack.Num());
+			for (TTuple<FString, TUniquePtr<UE::ShaderLibrary::Private::FNamedShaderLibrary>>& NamedLibraryPair : NamedLibrariesStack)
+			{
+				LibraryNames.Add(NamedLibraryPair.Key);
+			}
+		}
+
+		for (const FString& LibraryName : LibraryNames)
+		{
+			CloseLibrary(LibraryName);
+		}
 	}
 
 	void GatherShaderMapResourceReferences(const FShaderMapResource_SharedCode* InResource, TArray<FString>& OutLibrariesWithRefToResource)
