[shaders] change task file compression to use oodle instead of lz4; faster and avoids tsan issues in lz4 compressor

diff --git a/Engine/Source/Runtime/RenderCore/Private/ShaderCompileWorkerUtil.cpp b/Engine/Source/Runtime/RenderCore/Private/ShaderCompileWorkerUtil.cpp
--- a/Engine/Source/Runtime/RenderCore/Private/ShaderCompileWorkerUtil.cpp
+++ b/Engine/Source/Runtime/RenderCore/Private/ShaderCompileWorkerUtil.cpp
@@ -582,7 +582,7 @@ bool FShaderCompileWorkerUtil::WriteTasks(const TArray<FShaderCommonCompileJobPt
 	{
 		TransferFile.Close();
 
-		FName CompressionFormatToUse = NAME_LZ4;
+		FName CompressionFormatToUse = NAME_Oodle;
 
 		FString FormatName = CompressionFormatToUse.ToString();
 		InTransferFile << FormatName;
