Add Substrate project's closure count and override behavior.

* When Blendable GBuffer is selected for the project, no change possible
* When Adaptive GBuffer is selected for the project, you can specify a project-wide closure count. It will only affect platforms supporting Adaptive GBuffer. In addition you can specify the behavior:
  - Either the project's closure count is used as min between the project's closure count and the platform's closure count
  - Or the project's closure count forced on all platforms supporting Adaptive GBuffer

diff --git a/Engine/Source/Runtime/Engine/Classes/Engine/EngineTypes.h b/Engine/Source/Runtime/Engine/Classes/Engine/EngineTypes.h
--- a/Engine/Source/Runtime/Engine/Classes/Engine/EngineTypes.h
+++ b/Engine/Source/Runtime/Engine/Classes/Engine/EngineTypes.h
@@ -787,6 +787,22 @@ namespace ESubstrateStorageFormat
 	};
 }
 
+/** 
+ * Enumerates available Substrate closure configuration
+ */
+UENUM()
+namespace ESubstrateClosureConfig
+{
+	enum Type : int
+	{
+		/** Use the mininum between the platform's closure count and the project's closure count, on platforms supporting Adaptive GBuffer */
+		UseMinProjectAndPlatforClosuresPerPixel = 0 UMETA(DisplayName = "Min of project's and platfom's closure count"),
+
+		/** Force the project's closure count, on all platforms supporting Adaptive GBuffer */
+		ForceProjectClosuresPerPixel = 1 UMETA(DisplayName = "Force Closure Count"),
+	};
+}
+
 /**
  * Specifies the Substrate runtime shading model summarized from the material graph
  * Not exposed in UI, only used in code.
diff --git a/Engine/Source/Runtime/Engine/Classes/Engine/RendererSettings.h b/Engine/Source/Runtime/Engine/Classes/Engine/RendererSettings.h
--- a/Engine/Source/Runtime/Engine/Classes/Engine/RendererSettings.h
+++ b/Engine/Source/Runtime/Engine/Classes/Engine/RendererSettings.h
@@ -1232,6 +1232,29 @@ class URendererSettings : public UDeveloperSettings
 		ConfigRestartRequired = true))
 	TEnumAsByte<ESubstrateStorageFormat::Type> SubstrateGBufferFormat;
 
+	/**
+	"Substrate maximum number of closure evaluated per pixel for this project."
+	*/
+	UPROPERTY(config, EditAnywhere, Category = Substrate, meta = (
+		ConsoleVariable = "r.Substrate.ProjectClosuresPerPixel", 
+		DisplayName = "Substrate Closure Per Pixel (Project)",
+		ToolTip = "Define the maximum number of closures evaluated per pixel for this project. Materials with a number of closures beyond this limit will be simplified until it fits the requirements. This only applies when Adaptive GBuffer is enabled. On platforms not supporting Adaptive GBuffer, the material will be simplified into a single closure",
+		ClampMin = "1", ClampMax = "8", 
+		EditCondition = "SubstrateGBufferFormat == 1",
+		ConfigRestartRequired = true))
+	uint32 SubstrateProjectClosuresPerPixel;
+
+	/**
+	"Substrate maximum number of closure evaluated per pixel override behavior."
+	*/
+	UPROPERTY(config, EditAnywhere, Category = Substrate, meta = (
+		ConsoleVariable = "r.Substrate.ProjectClosuresPerPixelOverride",
+		DisplayName = "Substrate Closure-Per-Pixel Override",
+		ToolTip = "Select how the maximum number of closures evaluated per pixel is defined for each platforms.\n - Use Platform Closure Count: Use the mininum between the platform's closure count and the project's closure count.\n - Force Closure Count : Force the project's closure count on platforms supporting Adaptive GBuffer.",
+		EditCondition = "SubstrateGBufferFormat == 1",
+		ConfigRestartRequired = true))
+	TEnumAsByte<ESubstrateClosureConfig::Type> SubstrateClosurePerPixelOverride;
+
 	/**
 	"Enable Substrate opaque material rough refractions effect from top layers over layers below."
 	*/
@@ -1542,7 +1565,7 @@ private:
 #endif // WITH_EDITOR
 
 	void SanatizeReflectionCaptureResolution();
-
+	void SanatizeSubstrateSettings();
 	void UpdateWorkingColorSpaceAndChromaticities();
 };
 
diff --git a/Engine/Source/Runtime/Engine/Private/RendererSettings.cpp b/Engine/Source/Runtime/Engine/Private/RendererSettings.cpp
--- a/Engine/Source/Runtime/Engine/Private/RendererSettings.cpp
+++ b/Engine/Source/Runtime/Engine/Private/RendererSettings.cpp
@@ -121,6 +121,8 @@ void URendererSettings::PostInitProperties()
 
 	UpdateWorkingColorSpaceAndChromaticities();
 
+	SanatizeSubstrateSettings();
+
 #if WITH_EDITOR
 	CheckForMissingShaderModels();
 
@@ -440,6 +442,13 @@ void URendererSettings::PostEditChangeProperty(FPropertyChangedEvent& PropertyCh
 		}
 
 		FixAntiAliasingOnShadingPathChange(PropertyChangedEvent);
+
+		// Update Substrate closure count if project's GBuffer format change
+		if (PropertyChangedEvent.Property->GetFName() == GET_MEMBER_NAME_CHECKED(URendererSettings, SubstrateGBufferFormat) ||
+			PropertyChangedEvent.Property->GetFName() == GET_MEMBER_NAME_CHECKED(URendererSettings, SubstrateProjectClosuresPerPixel))
+		{
+			SanatizeSubstrateSettings();
+		}
 	}
 }
 
@@ -731,6 +740,22 @@ void URendererSettings::SanatizeReflectionCaptureResolution()
 #endif // WITH_EDITOR
 }
 
+void URendererSettings::SanatizeSubstrateSettings()
+{
+	// When blendable GBuffer is used, force project closure count to 1 for clarity
+	if (SubstrateGBufferFormat == ESubstrateStorageFormat::Type::BlendableGBuffer)
+	{
+		SubstrateProjectClosuresPerPixel = 1;
+	}
+
+	// Ensure project settings and the attached CVar are consistent
+	// Don't use ExportValuesToConsoleVariables() as it filters out ReadOnly CVars
+	if (IConsoleVariable* CVar = IConsoleManager::Get().FindConsoleVariable(TEXT("r.Substrate.ProjectClosuresPerPixel")))
+	{
+		CVar->Set((int32)SubstrateProjectClosuresPerPixel, ECVF_SetByProjectSetting);
+	}
+}
+
 void URendererSettings::UpdateWorkingColorSpaceAndChromaticities()
 {
 	using namespace UE::Color;
diff --git a/Engine/Source/Runtime/RenderCore/Private/RenderUtils.cpp b/Engine/Source/Runtime/RenderCore/Private/RenderUtils.cpp
--- a/Engine/Source/Runtime/RenderCore/Private/RenderUtils.cpp
+++ b/Engine/Source/Runtime/RenderCore/Private/RenderUtils.cpp
@@ -1948,10 +1948,22 @@ static TAutoConsoleVariable<int32> CVarSubstrateBytesPerPixel(
 
 static TAutoConsoleVariable<int32> CVarSubstrateClosuresPerPixel(
 	TEXT("r.Substrate.ClosuresPerPixel"),
-	1,	// Substrate 60fps by default. Must be less or equal to SUBSTRATE_MAX_CLOSURE_COUNT
+	1,	// Substrate 60fps by default. Must be less or equal to SUBSTRATE_MAX_CLOSURE_COUNT. This value is set for each platform
 	TEXT("Substrate closure count per pixel can be constrained. That is usefull to reduce the number of byte written durting the base pass, but also to limit the number of closures in Forward. Higher value means more complex material can be represented."),
 	ECVF_ReadOnly | ECVF_RenderThreadSafe | ECVF_Preview);
 
+static TAutoConsoleVariable<int32> CVarSubstrateProjectClosuresPerPixel(
+	TEXT("r.Substrate.ProjectClosuresPerPixel"),
+	4, // 4 closures by default in order to support advanced use case
+	TEXT("Define the maximum number of closures evaluated per pixel for this project. Materials with a number of closures beyond this limit will be simplified until it fits the requirements. This only applies when Adaptive GBuffer is enabled. On platforms not supporting Adaptive GBuffer, the material will be simplified into a single closure."),
+	ECVF_ReadOnly | ECVF_RenderThreadSafe | ECVF_Preview);
+
+static TAutoConsoleVariable<int32> CVarSubstrateProjectClosuresPerPixelOverride(
+	TEXT("r.Substrate.ProjectClosuresPerPixelOverride"),
+	0,	// By default use platform's max closure count (i.e., ESubstrateClosureConfig::Type::UsePlatformMaxClosuresPerPixel)
+	TEXT("Select how the maximum number of closures evaluated per pixel is defined for each platforms.\n - Use Platform Closure Count: Use the mininum between the platform's closure count and the project's closure count.\n - Force Closure Count : Force the project's closure count on platforms supporting Adaptive GBuffer."),
+	ECVF_ReadOnly | ECVF_RenderThreadSafe | ECVF_Preview);
+
 static TAutoConsoleVariable<int32> CVarSubstrateBackCompatibility(
 	TEXT("r.SubstrateBackCompatibility"),
 	0,
@@ -2131,6 +2143,19 @@ namespace Substrate
 		OutClosurePerPixel = FMath::Min(OutClosurePerPixel, ConfigClosurePerPixel);
 		#endif
 
+		// Take min between the platform's settings and the project's settings
+		const int32 ProjectClosurePerPixel = FMath::Clamp(CVarSubstrateProjectClosuresPerPixel.GetValueOnAnyThread(), 1, SUBSTRATE_MAX_CLOSURE_COUNT);
+		OutClosurePerPixel = FMath::Min(ProjectClosurePerPixel, OutClosurePerPixel);
+
+		// Force max closure count if enabled
+		// - ESubstrateClosureConfig::Type::UseMinProjectAndPlatforClosuresPerPixel = 0
+		// - ESubstrateClosureConfig::Type::ForceProjectClosuresPerPixel = 1
+		const uint32 bForceProjectClosureCount = CVarSubstrateProjectClosuresPerPixelOverride.GetValueOnAnyThread() > 0;
+		if (bForceProjectClosureCount)
+		{
+			OutClosurePerPixel = ProjectClosurePerPixel;
+		}
+
 		return uint32(FMath::Max(1, OutClosurePerPixel));
 	}
 
