Reduce memory used for shaders under compilation by shrinking remapper state when built.
#rnx

diff --git a/Engine/Source/Runtime/RenderCore/Private/ShaderPreprocessTypes.cpp b/Engine/Source/Runtime/RenderCore/Private/ShaderPreprocessTypes.cpp
--- a/Engine/Source/Runtime/RenderCore/Private/ShaderPreprocessTypes.cpp
+++ b/Engine/Source/Runtime/RenderCore/Private/ShaderPreprocessTypes.cpp
@@ -190,6 +190,12 @@ void FShaderDiagnosticRemapper::AddStrippedLine(int32 StrippedLineNum, int32 Off
 	StrippedLineOffsets.Add(Offset);
 }
 
+void FShaderDiagnosticRemapper::Shrink()
+{
+	Blocks.Shrink();
+	StrippedLineOffsets.Shrink();
+}
+
 FShaderDiagnosticRemapper::FRemapData FShaderDiagnosticRemapper::GetRemapData(int32 StrippedLineNum) const
 {
 	if (StrippedLineNum >= 1 && StrippedLineNum <= StrippedLineOffsets.Num())
@@ -368,6 +374,7 @@ void FShaderPreprocessOutput::StripCode(bool bCopyOriginalPreprocessdSource)
 	check(OutStripped <= OutStrippedData + PreprocessedSourceStripped.Len());
 	// ShrinkToLen null terminates for us by virtue of adding zero'd SIMD padding
 	PreprocessedSourceStripped.ShrinkToLen((int32)(OutStripped - OutStrippedData));
+	Remapper.Shrink();
 
 	if (bCopyOriginalPreprocessdSource)
 	{
diff --git a/Engine/Source/Runtime/RenderCore/Public/ShaderPreprocessTypes.h b/Engine/Source/Runtime/RenderCore/Public/ShaderPreprocessTypes.h
--- a/Engine/Source/Runtime/RenderCore/Public/ShaderPreprocessTypes.h
+++ b/Engine/Source/Runtime/RenderCore/Public/ShaderPreprocessTypes.h
@@ -32,6 +32,7 @@ private:
 	}
 	void AddSourceBlock(int32 OriginalLineNum, int32 StrippedLineNum, FString&& OriginalPath);
 	void AddStrippedLine(int32 StrippedLineNum, int32 Offset);
+	void Shrink();
 
 	struct FRemapData
 	{
