[Render] Test code to debug a shader collision problem.

diff --git a/Engine/Source/Runtime/RenderCore/Private/ShaderCodeArchive.cpp b/Engine/Source/Runtime/RenderCore/Private/ShaderCodeArchive.cpp
--- a/Engine/Source/Runtime/RenderCore/Private/ShaderCodeArchive.cpp
+++ b/Engine/Source/Runtime/RenderCore/Private/ShaderCodeArchive.cpp
@@ -1684,7 +1684,20 @@ void FIoStoreShaderCodeArchive::CreateIoStoreShaderCodeArchiveHeader(const FName
 		static_assert(sizeof(uint8) == sizeof(UTF8CHAR), "Unexpected UTF-8 char size.");
 
 		GroupEntry.UncompressedSize = CurrentGroupSize;
-		OutHeader.ShaderGroupIoHashes.Add(FIoStoreShaderCodeArchive::GetShaderCodeChunkId(GroupHasher.Finalize()));
+		const FIoChunkId ChunkId = FIoStoreShaderCodeArchive::GetShaderCodeChunkId(GroupHasher.Finalize());
+		OutHeader.ShaderGroupIoHashes.Add(ChunkId);
+		
+		// @TEST remove once shader collision issue is solved.
+		if (LexToString(ChunkId) == TEXT("8b8392e3d3baea4e9b986509"))
+		{
+			for (int32 ShaderIdxIdx = 0, NumIdxIdx = ShaderIndicesInGroup.Num(); ShaderIdxIdx < NumIdxIdx; ++ShaderIdxIdx)
+			{
+				int32 ShaderIndex = ShaderIndicesInGroup[ShaderIdxIdx];
+				UE_LOG(LogShaderLibrary, Display, TEXT("8b839 %d index %d hash %s"), ShaderIdxIdx, ShaderIndex, *LexToString(ShaderHashes[ShaderIndex].Hash));
+				UE_LOG(LogShaderLibrary, Display, TEXT("8b839 %d index %d uncompressed size %d"), ShaderIdxIdx, ShaderIndex, ShaderEntries[ShaderIndex].UncompressedSize);
+			}
+			UE_LOG(LogShaderLibrary, Display, TEXT("8b839 format %s"), *Utf8FormatName.Get());
+		}
 
 		Stats_TotalUncompressedMemory += CurrentGroupSize;
 		Stats_MinGroupSize = FMath::Min(Stats_MinGroupSize, CurrentGroupSize);
