[Render] Quick fix for previous CL until I can compile again.

diff --git a/Engine/Source/Runtime/RenderCore/Private/ShaderCodeArchive.cpp b/Engine/Source/Runtime/RenderCore/Private/ShaderCodeArchive.cpp
--- a/Engine/Source/Runtime/RenderCore/Private/ShaderCodeArchive.cpp
+++ b/Engine/Source/Runtime/RenderCore/Private/ShaderCodeArchive.cpp
@@ -1696,7 +1696,7 @@ void FIoStoreShaderCodeArchive::CreateIoStoreShaderCodeArchiveHeader(const FName
 				UE_LOG(LogShaderLibrary, Display, TEXT("8b839 %d index %d hash %s"), ShaderIdxIdx, ShaderIndex, *LexToString(ShaderHashes[ShaderIndex].Hash));
 				UE_LOG(LogShaderLibrary, Display, TEXT("8b839 %d index %d uncompressed size %d"), ShaderIdxIdx, ShaderIndex, ShaderEntries[ShaderIndex].UncompressedSize);
 			}
-			UE_LOG(LogShaderLibrary, Display, TEXT("8b839 format %s"), *Utf8FormatName.Get());
+			//UE_LOG(LogShaderLibrary, Display, TEXT("8b839 format %s"), *Utf8FormatName.Get());
 		}
 
 		Stats_TotalUncompressedMemory += CurrentGroupSize;
