Dump the current active breadcrumb in RDG task crashes
- One scope covers parallel RDG tasks. Another covers the render thread inside FRDGBuilder::Execute().
- Generates a file containing the list of breadcrumbs if a crash occurs within an RDG pass.

diff --git a/Engine/Source/Runtime/RenderCore/Private/RenderGraphBuilder.cpp b/Engine/Source/Runtime/RenderCore/Private/RenderGraphBuilder.cpp
--- a/Engine/Source/Runtime/RenderCore/Private/RenderGraphBuilder.cpp
+++ b/Engine/Source/Runtime/RenderCore/Private/RenderGraphBuilder.cpp
@@ -1760,7 +1760,19 @@ void FRDGBuilder::Execute()
 #if WITH_RHI_BREADCRUMBS
 	check(LocalCurrentBreadcrumb == FRHIBreadcrumbNode::Sentinel);
 	LocalCurrentBreadcrumb = RHICmdList.GetCurrentBreadcrumbRef();
-#endif
+
+#if WITH_ADDITIONAL_CRASH_CONTEXTS
+	FScopedAdditionalCrashContextProvider CrashContext([&RHICmdList = RHICmdList](FCrashContextExtendedWriter& Writer)
+	{
+		if (FRHIBreadcrumbNode* CurrentBreadcrumb = RHICmdList.GetCurrentBreadcrumbRef())
+		{
+			CurrentBreadcrumb->WriteCrashData(Writer, TEXT("RDGExecute_RenderThread"));
+		}
+	});
+#endif // WITH_ADDITIONAL_CRASH_CONTEXTS
+
+#endif // WITH_RHI_BREADCRUMBS
+
 
 	GRDGTransientResourceAllocator.ReleasePendingDeallocations();
 
@@ -3005,6 +3017,17 @@ void FRDGBuilder::SetupParallelExecute(TStaticArray<void*, MAX_NUM_GPUS> const&
 
 			PushPreScopes(*RHICmdListPass, ParallelPasses[0]);
 			{
+			#if WITH_RHI_BREADCRUMBS && WITH_ADDITIONAL_CRASH_CONTEXTS
+				FScopedAdditionalCrashContextProvider CrashContext([RHICmdListPass](FCrashContextExtendedWriter& Writer)
+				{
+					if (FRHIBreadcrumbNode* CurrentBreadcrumb = RHICmdListPass->GetCurrentBreadcrumbRef())
+					{
+						FString ThreadIDStr = FString::Printf(TEXT("RDGTask_%d"), FPlatformTLS::GetCurrentThreadId());
+						CurrentBreadcrumb->WriteCrashData(Writer, *ThreadIDStr);
+					}
+				});
+			#endif // WITH_ADDITIONAL_CRASH_CONTEXTS
+
 				for (FRDGPass* Pass : ParallelPasses)
 				{
 					ExecutePass(*RHICmdListPass, Pass);
