Fix trivial incorrect comment

diff --git a/Engine/Source/Runtime/RenderCore/Private/RenderGraphBuilder.cpp b/Engine/Source/Runtime/RenderCore/Private/RenderGraphBuilder.cpp
--- a/Engine/Source/Runtime/RenderCore/Private/RenderGraphBuilder.cpp
+++ b/Engine/Source/Runtime/RenderCore/Private/RenderGraphBuilder.cpp
@@ -3026,7 +3026,7 @@ void FRDGBuilder::SetupParallelExecute(TStaticArray<void*, MAX_NUM_GPUS> const&
 						CurrentBreadcrumb->WriteCrashData(Writer, *ThreadIDStr);
 					}
 				});
-			#endif // WITH_ADDITIONAL_CRASH_CONTEXTS
+			#endif
 
 				for (FRDGPass* Pass : ParallelPasses)
 				{
