PR #13659: Iterate OrderedCompileTasks in reverse to optimize cost of inserting and removing elements

#rnx

diff --git a/Engine/Source/Runtime/RenderCore/Private/ShaderPipelineCache.cpp b/Engine/Source/Runtime/RenderCore/Private/ShaderPipelineCache.cpp
--- a/Engine/Source/Runtime/RenderCore/Private/ShaderPipelineCache.cpp
+++ b/Engine/Source/Runtime/RenderCore/Private/ShaderPipelineCache.cpp
@@ -447,7 +447,6 @@ private:
 	TSet<uint64> CompletedMasks;
 	FShaderPipelineCache::FShaderCachePrecompileContext PrecompileContext;
 	bool bPreOptimizing = false;
-	TArray<FPipelineCachePSOHeader> PreFetchedTasks;
 	bool bReady = false;
 	bool bOpened = false;
 	bool bIsUserCache;
@@ -1418,7 +1417,7 @@ void FShaderPipelineCacheTask::PreparePipelineBatch(TDoubleLinkedList<FPipelineC
 					FPipelineCachePSOHeader Hdr;
 					Hdr.Hash = PSORead->Hash;
 					Hdr.Shaders = RequiredShaders;
-					OrderedCompileTasks.Insert(Hdr, 0);
+					OrderedCompileTasks.Add(Hdr);
 				}
 				else if (bCompatible)
 				{
@@ -1877,18 +1876,11 @@ void FShaderPipelineCacheTask::Tick(float DeltaTime)
 		TotalPrecompileWallTime = float(FPlatformTime::Seconds() - PrecompileStartTime);
 	}
 	
-	if (PrecompileStartTime == 0.0 && (PreFetchedTasks.Num() || FetchTasks.Num() || OrderedCompileTasks.Num()))
+	if (PrecompileStartTime == 0.0 && (FetchTasks.Num() || OrderedCompileTasks.Num()))
 	{
 		PrecompileStartTime = FPlatformTime::Seconds();
 	}
 	
-	// Copy any new items over to our 'internal' safe array
-	if (PreFetchedTasks.Num())
-	{
-		OrderedCompileTasks.Append(PreFetchedTasks);
-		PreFetchedTasks.Empty();
-	}
-	
 	if (ReadyForPrecompile())
 	{
 		SCOPE_SECONDS_ACCUMULATOR(STAT_PreCompileTotalTime);
@@ -1935,24 +1927,24 @@ void FShaderPipelineCacheTask::Tick(float DeltaTime)
 			TDoubleLinkedList<FPipelineCacheFileFormatPSORead*> NewBatch;
 		
 			Num -= FetchTasks.Num();
-            for (auto Iterator = OrderedCompileTasks.CreateIterator(); Iterator && Num > 0; ++Iterator)
+            for (int32 Index = OrderedCompileTasks.Num() - 1; Index >= 0 && Num > 0; --Index)
             {
                 bool bHasShaders = true;
-                for (FSHAHash const& Hash : Iterator->Shaders)
+                for (FSHAHash const& Hash : OrderedCompileTasks[Index].Shaders)
                 {
                     bHasShaders &= FShaderCodeLibrary::ContainsShaderCode(Hash);
                 }
             	if (bHasShaders)
             	{
             		FPipelineCacheFileFormatPSORead* Entry = new FPipelineCacheFileFormatPSORead;
-					Entry->Hash = Iterator->Hash;
+					Entry->Hash = OrderedCompileTasks[Index].Hash;
 					Entry->Ar = new FShaderPipelineCacheArchive;
 					
 					// Add to both new batch and fetch lists
 					NewBatch.AddTail(Entry);
 					FetchTasks.AddTail(Entry);
 					
-	            	Iterator.RemoveCurrent();
+            		OrderedCompileTasks.RemoveAtSwap(Index);
                     TotalActiveTasks++;
                     TotalWaitingTasks--;
                     --Num;
@@ -2127,10 +2119,14 @@ void FShaderPipelineCacheTask::BeginPrecompilingPipelineCache()
 		}
 
 		// Copy any new items over to our 'internal' safe array
-		check(PreFetchedTasks.IsEmpty());
 		check(OrderedCompileTasks.IsEmpty());
 
-		OrderedCompileTasks = MoveTemp(LocalPreFetchedTasks);
+		const int32 NumLocalPreFetchedTasks = LocalPreFetchedTasks.Num();
+		OrderedCompileTasks.Reserve(NumLocalPreFetchedTasks);
+		for (int32 Index = NumLocalPreFetchedTasks - 1; Index >= 0; --Index)
+		{
+			OrderedCompileTasks.Add(LocalPreFetchedTasks[Index]);
+		}
 
 		TotalPrecompileTasks = OrderedCompileTasks.Num();
 
