[Shaders] Allow to upload debug info from ShaderCodeLibrary as artifacts on build machines.

This adds CVar `r.ShaderCodeLibrary.UploadArtifacts` (disabled by default) which uploads the dumped debug information from ShaderCodeLibrary as build artifacts.
Also provide stub for deprecated function declaration DumpShaderCodeStats(), which was missing.

#rnx

diff --git a/Engine/Source/Runtime/Engine/Public/MaterialShared.h b/Engine/Source/Runtime/Engine/Public/MaterialShared.h
--- a/Engine/Source/Runtime/Engine/Public/MaterialShared.h
+++ b/Engine/Source/Runtime/Engine/Public/MaterialShared.h
@@ -2262,7 +2262,7 @@ public:
 	virtual void GetShaderMapId(EShaderPlatform Platform, const ITargetPlatform* TargetPlatform,
 		FMaterialShaderMapId& OutId) const
 	{
-		return BuildShaderMapId(OutId, TargetPlatform);
+		BuildShaderMapId(OutId, TargetPlatform);
 	}
 
 	/** Returns the ShaderMapId used in DDC and IncrementalCook for the ShaderMap created for this FMaterial. */
@@ -2276,7 +2276,7 @@ public:
 	void BuildShaderMapId(FMaterialShaderMapId& OutId, const ITargetPlatform* TargetPlatform,
 		EMaterialKeyInclude IncludeFlags = EMaterialKeyInclude::All) const
 	{
-		return BuildShaderMapIdOverride({
+		BuildShaderMapIdOverride({
 			.OutId = &OutId, .Platform = GetShaderPlatform(), .TargetPlatform = TargetPlatform, .IncludeFlags = IncludeFlags }
 			);
 	}
diff --git a/Engine/Source/Runtime/RenderCore/Private/ShaderCodeLibrary.cpp b/Engine/Source/Runtime/RenderCore/Private/ShaderCodeLibrary.cpp
--- a/Engine/Source/Runtime/RenderCore/Private/ShaderCodeLibrary.cpp
+++ b/Engine/Source/Runtime/RenderCore/Private/ShaderCodeLibrary.cpp
@@ -106,6 +106,20 @@ static FAutoConsoleVariableRef CVarShaderMapResourceRef(
 	ECVF_Default
 );
 
+bool GUploadShaderDebugArtifacts = false;
+static FAutoConsoleVariableRef CVarUploadShaderDebugArtifactsRef(
+	TEXT("r.ShaderCodeLibrary.UploadArtifacts"),
+	GUploadShaderDebugArtifacts,
+	TEXT("If true, uploads artifacts that were dumped to the Saved/Shaders/ folder when running on a build machine"),
+	ECVF_Default
+);
+
+static FString GetShaderCodeLibraryDebugInfoPath()
+{
+	const FString BaseDebugInfoPath = GIsBuildMachine && GUploadShaderDebugArtifacts ? GetBuildMachineArtifactBasePath() : FPaths::ProjectSavedDir();
+	return BaseDebugInfoPath / TEXT("Shaders");
+}
+
 class FShaderLibraryInstance;
 namespace UE
 {
@@ -2237,7 +2251,7 @@ struct FEditorShaderCodeArchive
 		static const FRegexPattern ShaderCodePattern(ShaderCodePatternStr);
 
 		const FString FormatNameStr = FormatName.ToString();
-		const FString ShaderIntermediateLocation = FPaths::ProjectSavedDir() / TEXT("Shaders") / FormatNameStr;
+		const FString ShaderIntermediateLocation = GetShaderCodeLibraryDebugInfoPath() / FormatNameStr;
 
 		TArray<FString> ShaderFiles;
 		ShaderFindFiles(ShaderFiles, *ShaderIntermediateLocation, *ShaderExtension);
@@ -2291,9 +2305,10 @@ struct FEditorShaderCodeArchive
 		if (bSuccess && SerializedShaders.GetNumShaderMaps() > 0)
 		{
 			// Write to a intermediate file
-			FString IntermediateFormatPath = GetShaderCodeFilename(FPaths::ProjectSavedDir() / TEXT("Shaders") / FormatName.ToString(), LibraryName, FormatName);
-			FString AssetInfoIntermediatePath = GetShaderAssetInfoFilename(FPaths::ProjectSavedDir() / TEXT("Shaders") / FormatName.ToString(), LibraryName, FormatName);
-			FString TypeInfoIntermediatePath = GetShaderTypeInfoFilename(FPaths::ProjectSavedDir() / TEXT("Shaders") / FormatName.ToString(), LibraryName, FormatName);
+			const FString DebugInfoPath = GetShaderCodeLibraryDebugInfoPath();
+			FString IntermediateFormatPath = GetShaderCodeFilename(DebugInfoPath / FormatName.ToString(), LibraryName, FormatName);
+			FString AssetInfoIntermediatePath = GetShaderAssetInfoFilename(DebugInfoPath / FormatName.ToString(), LibraryName, FormatName);
+			FString TypeInfoIntermediatePath = GetShaderTypeInfoFilename(DebugInfoPath / FormatName.ToString(), LibraryName, FormatName);
 
 			// save the actual shader code code
 			if (!bSaveOnlyAssetInfo)
@@ -2390,7 +2405,7 @@ struct FEditorShaderCodeArchive
 		}
 		checkf(!bHasCopiedAndCleared, TEXT("It is not valid to call PackageNativeShaderLibrary on an FEditorShaderCodeArchive that has sent its shaders to another process by calling CopyAndClear."));
 
-		FString IntermediateFormatPath = GetShaderDebugFolder(FPaths::ProjectSavedDir() / TEXT("Shaders") / FormatName.ToString(), LibraryName, FormatName);
+		FString IntermediateFormatPath = GetShaderDebugFolder(GetShaderCodeLibraryDebugInfoPath() / FormatName.ToString(), LibraryName, FormatName);
 		FString TempPath = IntermediateFormatPath / TEXT("NativeLibrary");
 
 		IFileManager::Get().MakeDirectory(*TempPath, true);
@@ -2564,18 +2579,21 @@ struct FEditorShaderCodeArchive
 				UE_LOG(LogShaderLibrary, Display, TEXT("No shader usage info is provided"));
 			}
 
-			FString DebugLibFolder = GetShaderDebugFolder(FPaths::ProjectSavedDir() / TEXT("Shaders") / FormatName.ToString(), LibraryName, FormatName);
+			const FString DebugInfoPath = GetShaderCodeLibraryDebugInfoPath();
+			const FString DebugLibFolder = GetShaderDebugFolder(DebugInfoPath / FormatName.ToString(), LibraryName, FormatName);
 			IFileManager::Get().MakeDirectory(*DebugLibFolder, true);
 
 			{
-				FString DumpFile = DebugLibFolder / TEXT("Dump.txt");
-				TUniquePtr<FArchive> DumpAr(IFileManager::Get().CreateFileWriter(*DumpFile));
-				FTCHARToUTF8 Converter(*ExtendedStats.TextualRepresentation);
-				DumpAr->Serialize(const_cast<UTF8CHAR*>(reinterpret_cast<const UTF8CHAR*>(Converter.Get())), Converter.Length());
-				UE_LOG(LogShaderLibrary, Display, TEXT("Textual dump saved to '%s'"), *DumpFile);
+				const FString DumpFilename = DebugLibFolder / TEXT("Dump.txt");
+				if (TUniquePtr<FArchive> DumpAr = TUniquePtr<FArchive>(IFileManager::Get().CreateFileWriter(*DumpFilename)))
+				{
+					FTCHARToUTF8 Converter(*ExtendedStats.TextualRepresentation);
+					DumpAr->Serialize(const_cast<UTF8CHAR*>(reinterpret_cast<const UTF8CHAR*>(Converter.Get())), Converter.Length());
+					UE_LOG(LogShaderLibrary, Display, TEXT("Textual dump saved to '%s'"), *DumpFilename);
+				}
 			}
 #if 0 // creating a graphviz graph - maybe one day we'll return to this
-			FString DebugGraphFolder = GetShaderDebugFolder(FPaths::ProjectSavedDir() / TEXT("Shaders") / FormatName.ToString(), LibraryName, FormatName);
+			FString DebugGraphFolder = GetShaderDebugFolder(DebugInfoPath / FormatName.ToString(), LibraryName, FormatName);
 			FString DebugGraphFile = DebugGraphFolder / TEXT("RelationshipGraph.gv");
 
 			IFileManager::Get().MakeDirectory(*DebugGraphFolder, true);
@@ -2705,7 +2723,7 @@ struct FEditorShaderStableInfo
 
 		bool bLibraryExistsInSavedShadersDir = false;
 		{
-			const FString ShaderIntermediateLocation = FPaths::ProjectSavedDir() / TEXT("Shaders") / FormatName.ToString();
+			const FString ShaderIntermediateLocation = GetShaderCodeLibraryDebugInfoPath() / FormatName.ToString();
 			TArray<FString> ShaderFiles;
 			ShaderFindFiles(ShaderFiles, *ShaderIntermediateLocation, *StableExtension);
 			FString ExpectedFileNameText = LibraryName + TEXT("-") + FormatName.ToString() + TEXT(".");
@@ -2750,7 +2768,7 @@ struct FEditorShaderStableInfo
 		if (bSuccess)
 		{
 			// Write to a intermediate file
-			FString IntermediateFormatPath = GetStableInfoArchiveFilename(FPaths::ProjectSavedDir() / TEXT("Shaders") / FormatName.ToString(), LibraryName, FormatName);
+			FString IntermediateFormatPath = GetStableInfoArchiveFilename(GetShaderCodeLibraryDebugInfoPath() / FormatName.ToString(), LibraryName, FormatName);
 
 			// Write directly to the file
 			{
@@ -3465,8 +3483,9 @@ public:
 	{
 		for (FName const& Format : ShaderFormats)
 		{
-			FString ShaderIntermediateLocation = FPaths::ProjectSavedDir() / TEXT("Shaders") / Format.ToString();
-			FString ShaderPlatformIntermediateLocation = FPaths::ProjectSavedDir() / TEXT("Shaders") / GetFormatAndPlatformName(Format);
+			const FString DebugInfoPath = GetShaderCodeLibraryDebugInfoPath();
+			FString ShaderIntermediateLocation = DebugInfoPath / Format.ToString();
+			FString ShaderPlatformIntermediateLocation = DebugInfoPath / GetFormatAndPlatformName(Format);
 			IFileManager::Get().DeleteDirectory(*ShaderIntermediateLocation, false, true);
 			IFileManager::Get().DeleteDirectory(*ShaderPlatformIntermediateLocation, false, true);
 		}
@@ -4429,6 +4448,11 @@ void FShaderLibraryCooker::DumpShaderTypeStats(const FString& DebugInfoDir, cons
 	}
 }
 
+void FShaderLibraryCooker::DumpShaderCodeStats()
+{
+	// stub - deprecated since UE 5.6
+}
+
 bool FShaderLibraryCooker::CreatePatchLibrary(TArray<FString> const& OldMetaDataDirs, FString const& NewMetaDataDir, FString const& OutDir, bool bNativeFormat, bool bNeedsDeterministicOrder)
 {
 	TMap<FName, TSet<FString>> FormatLibraryMap;
