Substrate - fix for settings causing issues in runtime game builds.

#rnx
#tests PC editor + game, fixed after rerpoduced.

diff --git a/Engine/Source/Runtime/Engine/Classes/Engine/RendererSettings.h b/Engine/Source/Runtime/Engine/Classes/Engine/RendererSettings.h
--- a/Engine/Source/Runtime/Engine/Classes/Engine/RendererSettings.h
+++ b/Engine/Source/Runtime/Engine/Classes/Engine/RendererSettings.h
@@ -1216,7 +1216,7 @@ class URendererSettings : public UDeveloperSettings
 	UPROPERTY(config, EditAnywhere, Category = Substrate, meta = (
 		ConsoleVariable = "r.Substrate.ProjectGBufferFormat",
 		DisplayName = "Substrate GBuffer Format (Project)",
-		ToolTip = "Select the GBuffer format used for this project.\n - Blendable GBuffer is supported on all platforms. \n - Adaptive GBuffer is supported on certains platforms. Unsupported platforms will fallback onto 'Blendable GBuffer' format.\nChanging this setting requires restarting the editor.",
+		ToolTip = "Select the GBuffer format used for this project.\n - Blendable GBuffer is supported on all platforms (=> 1 closure and 20 bytes per pixel maximum). \n - Adaptive GBuffer is supported on certains platforms. Unsupported platforms will fallback onto 'Blendable GBuffer' format.\nChanging this setting requires restarting the editor.",
 		ConfigRestartRequired = true))
 	TEnumAsByte<ESubstrateStorageFormat::Type> SubstrateGBufferFormat;
 
@@ -1566,7 +1566,6 @@ private:
 #endif // WITH_EDITOR
 
 	void SanitizeReflectionCaptureResolution();
-	void SanitizeSubstrateSettings(bool bSyncConsoleVariables);
 	void UpdateWorkingColorSpaceAndChromaticities();
 };
 
diff --git a/Engine/Source/Runtime/Engine/Private/RendererSettings.cpp b/Engine/Source/Runtime/Engine/Private/RendererSettings.cpp
--- a/Engine/Source/Runtime/Engine/Private/RendererSettings.cpp
+++ b/Engine/Source/Runtime/Engine/Private/RendererSettings.cpp
@@ -130,9 +130,6 @@ void URendererSettings::PostInitProperties()
 		ImportConsoleVariableValues();
 	}
 #endif // #if WITH_EDITOR
-
-	// Only sanitize Substrate value "after" the CVar import, to be sure default value are correct.
-	SanitizeSubstrateSettings(true /*bSyncConsoleVariables*/);
 }
 
 #if WITH_EDITOR
@@ -440,17 +437,6 @@ void URendererSettings::PostEditChangeProperty(FPropertyChangedEvent& PropertyCh
 		}
 
 		FixAntiAliasingOnShadingPathChange(PropertyChangedEvent);
-
-		// Update Substrate closure count if project's GBuffer format change
-		if (PropertyChangedEvent.Property->GetFName() == GET_MEMBER_NAME_CHECKED(URendererSettings, SubstrateGBufferFormat) ||
-			PropertyChangedEvent.Property->GetFName() == GET_MEMBER_NAME_CHECKED(URendererSettings, SubstrateProjectClosuresPerPixel) ||
-			PropertyChangedEvent.Property->GetFName() == GET_MEMBER_NAME_CHECKED(URendererSettings, SubstrateProjectBytesPerPixel))
-		{
-			// These settings require a restart
-			// Do not update the console variable to avoid a mismatch between the settings used for materials compilation, and the 'new' settings. 
-			// The 'new' settings will be applied on the next restart
-			SanitizeSubstrateSettings(false /*SanitizeSubstrateSettings*/);
-		}
 	}
 }
 
@@ -737,31 +723,6 @@ void URendererSettings::SanitizeReflectionCaptureResolution()
 #endif // WITH_EDITOR
 }
 
-void URendererSettings::SanitizeSubstrateSettings(bool bSyncConsoleVariables)
-{
-	// When blendable GBuffer is used, force project closure count to 1 for clarity
-	if (SubstrateGBufferFormat == ESubstrateStorageFormat::Type::BlendableGBuffer)
-	{
-		SubstrateProjectClosuresPerPixel = 1;
-		SubstrateProjectBytesPerPixel = 20u;
-	}
-
-	// Ensure project settings and the attached CVar are consistent
-	// Don't use ExportValuesToConsoleVariables() as it filters out ReadOnly CVars
-	if (bSyncConsoleVariables)
-	{
-		if (IConsoleVariable* CVar = IConsoleManager::Get().FindConsoleVariable(TEXT("r.Substrate.ProjectClosuresPerPixel")))
-		{
-			CVar->Set((int32)SubstrateProjectClosuresPerPixel, ECVF_SetByProjectSetting);
-		}
-
-		if (IConsoleVariable* CVar = IConsoleManager::Get().FindConsoleVariable(TEXT("r.Substrate.BytesPerPixel")))
-		{
-			CVar->Set((int32)SubstrateProjectBytesPerPixel, ECVF_SetByProjectSetting);
-		}
-	}
-}
-
 void URendererSettings::UpdateWorkingColorSpaceAndChromaticities()
 {
 	using namespace UE::Color;
diff --git a/Engine/Source/Runtime/RenderCore/Private/RenderUtils.cpp b/Engine/Source/Runtime/RenderCore/Private/RenderUtils.cpp
--- a/Engine/Source/Runtime/RenderCore/Private/RenderUtils.cpp
+++ b/Engine/Source/Runtime/RenderCore/Private/RenderUtils.cpp
@@ -2090,34 +2090,28 @@ namespace Substrate
 		return IsSubstrateEnabled() && CVarSubstrateStochasticLighting.GetValueOnAnyThread() > 0 && !IsSubstrateBlendableGBufferEnabled(InPlatform);
 	}
 
-	static uint32 InternalGetBytePerPixel(uint32 InBytePerPixel)
+	uint32 GetBytePerPixel(EShaderPlatform InPlatform)
 	{
+		// Variant for shader compilation per platform
+		static FShaderPlatformCachedIniValue<int32> CVarByteBudget(TEXT("r.Substrate.BytesPerPixel"));
+		uint32 CvarBytePerPixel = CVarByteBudget.Get(InPlatform);
+
 		auto RoundUpValueToUInt = [](uint32 Value)
-		{
-			const uint32 Divisor = sizeof(uint32);
-			return  FMath::DivideAndRoundUp<uint32>(Value, Divisor) * Divisor;
-		};
+			{
+				const uint32 Divisor = sizeof(uint32);
+				return  FMath::DivideAndRoundUp<uint32>(Value, Divisor) * Divisor;
+			};
+
+		const bool bBlendableGBufferFormat = IsSubstrateBlendableGBufferEnabled(InPlatform);
 
 		// We enforce at least 20 bytes per pixel because this is the minimal Substrate GBuffer footprint of the simplest material.
 		const uint32 MinSubstrateBytePerPixel = 20u;
-		const uint32 MaxSubstrateBytePerPixel = /*IsMobilePlatform(GMaxRHI InPlatform) ? 24u : */256u; // SUBSTRATE_TODO
+		const uint32 MaxSubstrateBytePerPixel = bBlendableGBufferFormat ? 20u : 256u;
 		// Align byte per pixel count to 4 because the read/write unit is uint32
 		check(MinSubstrateBytePerPixel == RoundUpValueToUInt(MinSubstrateBytePerPixel));
 		check(MaxSubstrateBytePerPixel == RoundUpValueToUInt(MaxSubstrateBytePerPixel));
-		InBytePerPixel = RoundUpValueToUInt(InBytePerPixel);
-		return FMath::Clamp(InBytePerPixel, MinSubstrateBytePerPixel, MaxSubstrateBytePerPixel);
-	}
-
-	uint32 GetBytePerPixel()
-	{
-		return InternalGetBytePerPixel(CVarSubstrateBytesPerPixel.GetValueOnAnyThread());
-	}
-
-	uint32 GetBytePerPixel(EShaderPlatform InPlatform)
-	{
-		// Variant for shader compilation per platform
-		static FShaderPlatformCachedIniValue<int32> CVarByteBudget(TEXT("r.Substrate.BytesPerPixel"));
-		return InternalGetBytePerPixel(CVarByteBudget.Get(InPlatform));
+		CvarBytePerPixel = RoundUpValueToUInt(CvarBytePerPixel);
+		return FMath::Clamp(CvarBytePerPixel, MinSubstrateBytePerPixel, MaxSubstrateBytePerPixel);
 	}
 
 	uint32 GetClosurePerPixel(EShaderPlatform InPlatform)
@@ -2207,7 +2201,9 @@ namespace Substrate
 		}
 		else
 		{
-			PayloadSizeBytes += Substrate::GetBytePerPixel();
+			// Not supported for now. Ray tracing payload size is not scalable this way.
+			// See GetRaytracingMaterialPayloadSizeFullySimplified()
+			check(false);
 		}
 
 		return PayloadSizeBytes;
diff --git a/Engine/Source/Runtime/RenderCore/Public/RenderUtils.h b/Engine/Source/Runtime/RenderCore/Public/RenderUtils.h
--- a/Engine/Source/Runtime/RenderCore/Public/RenderUtils.h
+++ b/Engine/Source/Runtime/RenderCore/Public/RenderUtils.h
@@ -646,7 +646,6 @@ namespace Substrate
 
 	RENDERCORE_API uint32 GetRayTracingMaterialPayloadSizeInBytes(bool bFullySimplifiedMaterial);
 
-	RENDERCORE_API uint32 GetBytePerPixel();
 	RENDERCORE_API uint32 GetBytePerPixel(EShaderPlatform InPlatform);
 	RENDERCORE_API uint32 GetClosurePerPixel(EShaderPlatform InPlatform);
 
