Shader library: add more logs about a situation when IoStore request's bufer goes away.

diff --git a/Engine/Source/Runtime/RenderCore/Private/ShaderCodeArchive.cpp b/Engine/Source/Runtime/RenderCore/Private/ShaderCodeArchive.cpp
--- a/Engine/Source/Runtime/RenderCore/Private/ShaderCodeArchive.cpp
+++ b/Engine/Source/Runtime/RenderCore/Private/ShaderCodeArchive.cpp
@@ -2451,6 +2451,22 @@ TRefCountPtr<FRHIShader> FIoStoreShaderCodeArchive::CreateShader(int32 ShaderInd
 		CSV_CUSTOM_STAT(Shaders, PreloadShaderWaitTime, WaitDuration * 1000.0f, ECsvCustomStatOp::Accumulate);
 	}
 
+	// we should not have this situation, but if we do, log a little more info about it before dying
+	if (UNLIKELY(PreloadEntryPtr->IoRequest.GetResult() == nullptr))
+	{
+		FString ShaderGroupIoHash;
+		Header.ShaderGroupIoHashes[GroupIndex].ToString(ShaderGroupIoHash);
+		UE_LOG(LogShaderLibrary, Warning, TEXT("IoRequest backing buffer is not there when we need it, about to crash. Library '%s', shader hash '%s' (freq=%d), group IoChunkId '%s', started preloading at frame %d, current frame is %d, preloading %s"),
+			*GetName(),
+			*Header.ShaderHashes[ShaderIndex].ToString(),
+			ShaderEntry.Frequency,
+			*ShaderGroupIoHash,
+			PreloadEntryPtr->FramePreloadStarted,
+			GFrameNumber,
+			PreloadEntryPtr->PreloadEvent.IsValid() ? (PreloadEntryPtr->PreloadEvent->IsComplete() ? TEXT("is complete") : TEXT("is NOT complete!")) : TEXT("event is NOT VALID!")
+		);
+	}
+
 	const uint8* ShaderCode = PreloadEntryPtr->IoRequest.GetResultOrDie().Data();
 
 	FMemStackBase& MemStack = FMemStack::Get();
