Listshaders improvements
* Consider shaders from from FShaderMapContent.Shader and FShaderMapContent.ShaderPipelines
* Add the originating shadermap hash for listshaders

diff --git a/Engine/Source/Runtime/Apple/MetalRHI/Private/Shaders/MetalShaderLibrary.h b/Engine/Source/Runtime/Apple/MetalRHI/Private/Shaders/MetalShaderLibrary.h
--- a/Engine/Source/Runtime/Apple/MetalRHI/Private/Shaders/MetalShaderLibrary.h
+++ b/Engine/Source/Runtime/Apple/MetalRHI/Private/Shaders/MetalShaderLibrary.h
@@ -85,6 +85,11 @@ public:
 		return SerializedShaders.GetShaderHashes()[GetShaderIndex(ShaderMapIndex, ShaderIndex)];
 	};
 
+	virtual FSHAHash GetShaderMapHash(int32 ShaderMapIndex) const override
+	{
+		return SerializedShaders.GetShaderMapHashes()[ShaderMapIndex];
+	};
+
 	virtual bool PreloadShader(int32 ShaderIndex, FGraphEventArray& OutCompletionEvents) override { return false; }
 	virtual bool PreloadShaderMap(int32 ShaderMapIndex, FGraphEventArray& OutCompletionEvents) override { return false; }
 
diff --git a/Engine/Source/Runtime/RHI/Public/RHIShaderLibrary.h b/Engine/Source/Runtime/RHI/Public/RHIShaderLibrary.h
--- a/Engine/Source/Runtime/RHI/Public/RHIShaderLibrary.h
+++ b/Engine/Source/Runtime/RHI/Public/RHIShaderLibrary.h
@@ -27,6 +27,7 @@ public:
 	virtual void GetAllShaderIndices(int32 ShaderMapIndex, TArray<int32>& ShaderIndices) {}
 	virtual uint32 GetSizeBytes() const = 0;
 	virtual FSHAHash GetShaderHash(int32 ShaderMapIndex, int32 ShaderIndex) = 0;
+	virtual FSHAHash GetShaderMapHash(int32 ShaderMapIndex) const = 0;
 	virtual int32 FindShaderMapIndex(const FSHAHash& Hash) = 0;
 	virtual int32 FindShaderIndex(const FSHAHash& Hash) = 0;
 	virtual uint32 GetShaderSizeBytes(int32 ShaderIndex) const { return 0; }
diff --git a/Engine/Source/Runtime/RenderCore/Private/ShaderCodeLibrary.cpp b/Engine/Source/Runtime/RenderCore/Private/ShaderCodeLibrary.cpp
--- a/Engine/Source/Runtime/RenderCore/Private/ShaderCodeLibrary.cpp
+++ b/Engine/Source/Runtime/RenderCore/Private/ShaderCodeLibrary.cpp
@@ -327,6 +327,8 @@ public:
 	virtual int32 GetLibraryId() const override;
 	virtual int32 GetLibraryShaderIndex(int32 ShaderIndex) const override;
 	virtual uint32 GetShaderSizeBytes(int32 ShaderIndex) const override;
+	/** Return shader hash for a particular shader without creating it. */
+	virtual FSHAHash GetShaderMapHash() const override;
 
 	class FShaderLibraryInstance* LibraryInstance;
 	int32 ShaderMapIndex;
@@ -1416,6 +1418,11 @@ uint32 FShaderMapResource_SharedCode::GetShaderSizeBytes(int32 ShaderIndex) cons
 	return LibraryInstance->Library->GetShaderSizeBytes(LibraryShaderIndex);
 }
 
+FSHAHash FShaderMapResource_SharedCode::GetShaderMapHash() const
+{
+	return LibraryInstance->Library->GetShaderMapHash(ShaderMapIndex);
+}
+
 void FShaderMapResource_SharedCode::ReleasePreloadedShaderCode(int32 ShaderIndex)
 {
 	SCOPED_LOADTIMER(FShaderMapResource_SharedCode_InitRHI);	// part of shader initialization in a way
@@ -4887,4 +4894,4 @@ FAutoConsoleCommandWithArgsAndOutputDevice GListShaderLibrariesCmd(
 			{
 				UE_LOG(LogShaderLibrary, Warning, TEXT("ShaderLibrariesCollection is not available."));
 			}
-		}));
+		}));
\ No newline at end of file
diff --git a/Engine/Source/Runtime/RenderCore/Private/ShaderMap.cpp b/Engine/Source/Runtime/RenderCore/Private/ShaderMap.cpp
--- a/Engine/Source/Runtime/RenderCore/Private/ShaderMap.cpp
+++ b/Engine/Source/Runtime/RenderCore/Private/ShaderMap.cpp
@@ -52,13 +52,19 @@ FAutoConsoleCommandWithArgsAndOutputDevice GListShaderCmd(
 				UE_LOG(LogShaders, Warning, TEXT("Enable r.TrackShaderMaps in ini config to enable the functionality"));
 				return;
 			}
-			Out.Logf(TEXT("FriendlyName,OwnerName,VertexFactoryName,ShaderTypeName,ShaderPermutation,UncompressedSizeKb,Hash,UsedForRenderingRHIPtr"));
+			Out.Logf(TEXT("FriendlyName,OwnerName,ShaderMapHash,VertexFactoryName,ShaderTypeName,ShaderPermutation,UncompressedSizeKb,Hash,UsedForRenderingRHIPtr"));
 			FScopeLock SMAccess(&GAllShaderMapsGuard);
 			for (FShaderMapBase* ShaderMap : GAllShaderMaps)
 			{
+				TCHAR ShaderMapHash[FSHAHash::GetStringLen()+1] = {0};
+				if (ShaderMap->GetResource())
+				{
+					ShaderMap->GetResource()->GetShaderMapHash().ToString(ShaderMapHash, true);
+				}
+
 				FShaderListReport Report;
 				Report.BindLambda(
-					[&Out](
+					[&Out, &ShaderMapHash](
 						const TCHAR* FriendlyName,
 						const TCHAR* OwnerName,
 						const TCHAR* VertexFactory,
@@ -68,7 +74,7 @@ FAutoConsoleCommandWithArgsAndOutputDevice GListShaderCmd(
 						const TCHAR* Hash,
 						void* Ptr)
 					{
-						Out.Logf(TEXT("%s,%s,%s,%s,%d,%.2f,%s,%p"), FriendlyName, OwnerName, VertexFactory, ShaderType, Permutation, Bytes / 1024.f, Hash, Ptr);
+						Out.Logf(TEXT("%s,%s,%s,%s,%s,%d,%.2f,%s,%p"), FriendlyName, OwnerName, ShaderMapHash, VertexFactory, ShaderType, Permutation, Bytes / 1024.f, Hash, Ptr);
 					}
 				);
 				ShaderMap->DumpShaderList(Report);
@@ -648,6 +654,21 @@ void FShaderMapContent::RemoveShaderPipelineType(const FShaderPipelineType* Shad
 }
 
 #if !UE_BUILD_SHIPPING
+
+static void RetrieveShaderInfo(uint32& OutShaderCodeSizeBytes, FRHIShader*& OutRHIShader, FString& OutHash, FShaderMapResource* Resource, FShader* Shader)
+{
+	if (Resource)
+	{
+		if (Resource->HasShader(Shader->GetResourceIndex()))
+		{
+			OutRHIShader = Resource->GetShader(Shader->GetResourceIndex());
+		}
+		const uint32 Index = Shader->GetResourceIndex();
+		OutShaderCodeSizeBytes = Resource->GetShaderSizeBytes(Index);
+		OutHash = Resource->GetShaderHash(Index).ToString();
+	}
+}
+
 void FShaderMapContent::DumpShaderList(const FShaderMapBase& InShaderMap, FShaderListReport& Out) const
 {
 	FString FriendlyName;
@@ -664,16 +685,7 @@ void FShaderMapContent::DumpShaderList(const FShaderMapBase& InShaderMap, FShade
 		uint32 ShaderCodeSizeBytes = 0;
 		FRHIShader* RHIShader = nullptr;
 		FString Hash;
-		if (Resource)
-		{
-			if (Resource->HasShader(Shader->GetResourceIndex()))
-			{
-				RHIShader = Resource->GetShader(Shader->GetResourceIndex());
-			}
-			const uint32 Index = Shader->GetResourceIndex();
-			ShaderCodeSizeBytes = Resource->GetShaderSizeBytes(Index);
-			Hash = Resource->GetShaderHash(Index).ToString();
-		}
+		RetrieveShaderInfo(ShaderCodeSizeBytes, RHIShader, Hash, Resource, Shader);
 
 		Out.ExecuteIfBound(
 			*FriendlyName, 
@@ -686,6 +698,33 @@ void FShaderMapContent::DumpShaderList(const FShaderMapBase& InShaderMap, FShade
 			RHIShader
 		);
 	}	
+
+	for (const FShaderPipeline* ShaderPipeline : ShaderPipelines)
+	{
+		for (uint32 Frequency = 0u; Frequency < SF_NumGraphicsFrequencies; ++Frequency)
+		{
+			FShader* Shader = ShaderPipeline->Shaders[Frequency].Get();
+			if (Shader)
+			{
+				uint32 ShaderCodeSizeBytes = 0;
+				FRHIShader* RHIShader = nullptr;
+				FString Hash;
+
+				RetrieveShaderInfo(ShaderCodeSizeBytes, RHIShader, Hash, Resource, Shader);
+
+				Out.ExecuteIfBound(
+					*FriendlyName, 
+					*OwnerName,
+					Shader->GetVertexFactoryType(InShaderMap.GetPointerTable()) ? Shader->GetVertexFactoryType(InShaderMap.GetPointerTable())->GetName() : TEXT(""),
+					Shader->GetType(InShaderMap.GetPointerTable()) ? Shader->GetType(InShaderMap.GetPointerTable())->GetName() : TEXT(""),
+					ShaderPipeline->PermutationIds[Frequency],
+					ShaderCodeSizeBytes,
+					*Hash,
+					RHIShader
+				);
+			}
+		}
+	}
 }
 #endif // !UE_BUILD_SHIPPING
 
diff --git a/Engine/Source/Runtime/RenderCore/Public/Shader.h b/Engine/Source/Runtime/RenderCore/Public/Shader.h
--- a/Engine/Source/Runtime/RenderCore/Public/Shader.h
+++ b/Engine/Source/Runtime/RenderCore/Public/Shader.h
@@ -413,6 +413,12 @@ public:
 	/** Return shader hash for a particular shader without creating it. */
 	virtual FSHAHash GetShaderHash(int32 ShaderIndex) = 0;
 
+	virtual FSHAHash GetShaderMapHash() const
+	{
+		/* return invalid index when not using shader library */
+		return FSHAHash();
+	}
+
 	RENDERCORE_API void BeginCreateAllShaders();
 
 #if RHI_RAYTRACING
diff --git a/Engine/Source/Runtime/RenderCore/Public/ShaderCodeArchive.h b/Engine/Source/Runtime/RenderCore/Public/ShaderCodeArchive.h
--- a/Engine/Source/Runtime/RenderCore/Public/ShaderCodeArchive.h
+++ b/Engine/Source/Runtime/RenderCore/Public/ShaderCodeArchive.h
@@ -421,6 +421,11 @@ public:
 		return SerializedShaders.GetShaderHashes()[GetShaderIndex(ShaderMapIndex, ShaderIndex)];
 	};
 
+	virtual FSHAHash GetShaderMapHash(int32 ShaderMapIndex) const override
+	{
+		return SerializedShaders.GetShaderMapHashes()[ShaderMapIndex];
+	};
+
 	virtual bool PreloadShader(int32 ShaderIndex, FGraphEventArray& OutCompletionEvents) override;
 
 	virtual bool PreloadShaderMap(int32 ShaderMapIndex, FGraphEventArray& OutCompletionEvents) override;
@@ -684,6 +689,11 @@ public:
 		return Header.ShaderHashes[GetShaderIndex(ShaderMapIndex, ShaderIndex)];
 	}
 
+	virtual FSHAHash GetShaderMapHash(int32 ShaderMapIndex) const override
+	{
+		return Header.ShaderMapHashes[ShaderMapIndex];
+	}
+
 	virtual bool IsPreloading(int32 ShaderIndex, FGraphEventArray& OutCompletionEvents) override;
 	virtual bool PreloadShader(int32 ShaderIndex, FGraphEventArray& OutCompletionEvents) override;
 	virtual void AddRefPreloadedShaderGroup(int32 ShaderGroupIndex) override;
