Fix misleading comments

diff --git a/Engine/Source/Runtime/RenderCore/Private/ShaderCodeLibrary.cpp b/Engine/Source/Runtime/RenderCore/Private/ShaderCodeLibrary.cpp
--- a/Engine/Source/Runtime/RenderCore/Private/ShaderCodeLibrary.cpp
+++ b/Engine/Source/Runtime/RenderCore/Private/ShaderCodeLibrary.cpp
@@ -327,7 +327,6 @@ public:
 	virtual int32 GetLibraryId() const override;
 	virtual int32 GetLibraryShaderIndex(int32 ShaderIndex) const override;
 	virtual uint32 GetShaderSizeBytes(int32 ShaderIndex) const override;
-	/** Return shader hash for a particular shader without creating it. */
 	virtual FSHAHash GetShaderMapHash() const override;
 
 	class FShaderLibraryInstance* LibraryInstance;
diff --git a/Engine/Source/Runtime/RenderCore/Public/Shader.h b/Engine/Source/Runtime/RenderCore/Public/Shader.h
--- a/Engine/Source/Runtime/RenderCore/Public/Shader.h
+++ b/Engine/Source/Runtime/RenderCore/Public/Shader.h
@@ -415,7 +415,7 @@ public:
 
 	virtual FSHAHash GetShaderMapHash() const
 	{
-		/* return invalid index when not using shader library */
+		/* return invalid hash when not using shader library */
 		return FSHAHash();
 	}
 
