Mobile: Fixed inconsistent test a deferred shading introduced in CL# 46462651

diff --git a/Engine/Source/Runtime/RenderCore/Private/ReadOnlyCVARCache.cpp b/Engine/Source/Runtime/RenderCore/Private/ReadOnlyCVARCache.cpp
--- a/Engine/Source/Runtime/RenderCore/Private/ReadOnlyCVARCache.cpp
+++ b/Engine/Source/Runtime/RenderCore/Private/ReadOnlyCVARCache.cpp
@@ -88,7 +88,7 @@ bool FReadOnlyCVARCache::MobileAllowFramebufferFetchIniValue(EShaderPlatform Pla
 	{
 		static FShaderPlatformCachedIniValue<bool> CVar(TEXT("r.Mobile.AllowFramebufferFetch"));
 		// Always allow in Forward, optional in Deferred
-		return !(bMobileDeferredShadingValue && bMobileHDR) || CVar.Get(Platform);
+		return !(MobileDeferredShading(Platform) && MobileHDR()) || CVar.Get(Platform);
 	}
 
 	// Only GLES, Vulkan and Metal have FBF
