Removed magic number clamping r.RenderTargetPoolMin to 2GB.

diff --git a/Engine/Source/Runtime/RenderCore/Private/RenderTargetPool.cpp b/Engine/Source/Runtime/RenderCore/Private/RenderTargetPool.cpp
--- a/Engine/Source/Runtime/RenderCore/Private/RenderTargetPool.cpp
+++ b/Engine/Source/Runtime/RenderCore/Private/RenderTargetPool.cpp
@@ -328,7 +328,7 @@ void FRenderTargetPool::TickPoolElements()
 	{
 		static const auto CVar = IConsoleManager::Get().FindTConsoleVariableDataInt(TEXT("r.RenderTargetPoolMin"));
 
-		MinimumPoolSizeInKB = FMath::Clamp(CVar->GetValueOnRenderThread(), 0, 2000) * 1024;
+		MinimumPoolSizeInKB = FMath::Max(CVar->GetValueOnRenderThread(), 0) * 1024;
 	}
 
 	CompactPool();
