Add clarification comment for r.Substrate and r.Substrate.ProjectGBufferFormat

diff --git a/Engine/Source/Runtime/RenderCore/Private/RenderUtils.cpp b/Engine/Source/Runtime/RenderCore/Private/RenderUtils.cpp
--- a/Engine/Source/Runtime/RenderCore/Private/RenderUtils.cpp
+++ b/Engine/Source/Runtime/RenderCore/Private/RenderUtils.cpp
@@ -1943,6 +1943,16 @@ static TAutoConsoleVariable<int32> CVarSubstrate(
 	TEXT("Enable Substrate materials."),
 	ECVF_ReadOnly | ECVF_RenderThreadSafe);
 
+// Summary:
+// * Existing project NOT using Substrate will continue to NOT use Substrate.
+// * Existing project using Substrate will continue to use Substrate with Adaptive GBuffer
+// * New project will using Substrate with Blendable GBuffer
+//
+// By default, Substrate uses Adaptive GBUffer. This allows existing project to continue using Adaptive format without updating their settings
+// New projects created with 5.7, will have Substrate enabled automatically:
+//  * With Blendable GBuffer for regular templates.
+//  * With Adaptive Gbuffer for advanced templates.
+// This is handled in GameProjectUtils.cpp.
 static TAutoConsoleVariable<int32> CVarSubstrateProjectGBufferFormat(
 	TEXT("r.Substrate.ProjectGBufferFormat"),
 	1,
