[shaders] fix handling of pipeline jobs which also have secondary outputs
- remove check that enforces pipeline interpolator removal cannot be used in conjunction with secondary compilation. this was overly pessimistic - what we actually need to enforce is that set of used attributes determined via reflection need to match on both primary and secondary compile of any given pipeline stage. this is due to the fact that later stages do not necessarily have secondary compiles, and in the case they do not their input attributes need to match the output attributes of both versions of the preceding stage to prevent interpolator mismatches.
- fix debug compile path for such jobs; the map storing debug source files didn't work for the case where we had two different preprocessed source files for a single stage.

diff --git a/Engine/Source/Developer/ShaderCompilerCommon/Private/ShaderCompilerCommon.cpp b/Engine/Source/Developer/ShaderCompilerCommon/Private/ShaderCompilerCommon.cpp
--- a/Engine/Source/Developer/ShaderCompilerCommon/Private/ShaderCompilerCommon.cpp
+++ b/Engine/Source/Developer/ShaderCompilerCommon/Private/ShaderCompilerCommon.cpp
@@ -2051,7 +2051,7 @@ namespace UE::ShaderCompilerCommon
 		FString DebugSourcePath = Options.GetDebugShaderPath(Input, Suffix);
 		FFileHelper::SaveStringToFile(Contents, *DebugSourcePath);
 
-		Ctx.DebugSourceFiles.Add(Input.Target.GetFrequency(), MoveTemp(DebugSourcePath));
+		Ctx.DebugSourceFiles.Add({ Input.Target.GetFrequency(), MoveTemp(DebugSourcePath) });
 	}
 
 	void DumpDebugShaderData(const FShaderCompilerInput& Input, const FString& PreprocessedSource, const FDebugShaderDataOptions& Options)
diff --git a/Engine/Source/Programs/ShaderCompileWorker/Private/ShaderCompileWorker.cpp b/Engine/Source/Programs/ShaderCompileWorker/Private/ShaderCompileWorker.cpp
--- a/Engine/Source/Programs/ShaderCompileWorker/Private/ShaderCompileWorker.cpp
+++ b/Engine/Source/Programs/ShaderCompileWorker/Private/ShaderCompileWorker.cpp
@@ -829,9 +829,7 @@ private:
 				InputFile << NumStages;
 				FShaderPipelineCompileJob& PipelineJob = OutPipelineJobs.Emplace_GetRef(NumStages);
 
-				// if debug source files are specified ensure the number of files matches the number of pipeline stages
-				check(!bIsDebugCompile || DebugSourceFiles.Num() == NumStages);
-
+				int32 DebugSourceFileIndex = 0, DebugSourceFileCount = DebugSourceFiles.Num();
 				for (int32 StageIndex = 0; StageIndex < NumStages; ++StageIndex)
 				{
 					// Deserialize the job's inputs.
@@ -841,7 +839,9 @@ private:
 
 					if (bIsDebugCompile)
 					{
-						SetupDebugCompile(*Job, *DebugSourceFiles[StageIndex]);
+						const TCHAR* DebugSourceFile = DebugSourceFileIndex < DebugSourceFileCount ? *DebugSourceFiles[DebugSourceFileIndex++] : nullptr;
+						const TCHAR* DebugSecondarySourceFile = (Job->SecondaryPreprocessOutput.IsValid() && (DebugSourceFileIndex < DebugSourceFileCount)) ? *DebugSourceFiles[DebugSourceFileIndex++] : nullptr;
+						SetupDebugCompile(*Job, DebugSourceFile, DebugSecondarySourceFile);
 					}
 
 					// SCW doesn't run DDPI, GShaderHasCache Initialize is run  at start with no knowledge of the CustomPlatforms
diff --git a/Engine/Source/Runtime/RenderCore/Private/ShaderCore.cpp b/Engine/Source/Runtime/RenderCore/Private/ShaderCore.cpp
--- a/Engine/Source/Runtime/RenderCore/Private/ShaderCore.cpp
+++ b/Engine/Source/Runtime/RenderCore/Private/ShaderCore.cpp
@@ -1820,9 +1820,6 @@ public:
 		if (Job.PreprocessOutput.bSucceeded && Backend->RequiresSecondaryCompile(Job.Input, Environment, Job.PreprocessOutput))
 		{
 			PreprocessStartTime = FPlatformTime::Seconds();
-			// check that we're not compiling for a pipeline; this would be theretically possible to support but we'd have to enforce that all jobs in the pipeline
-			// have a secondary compile if any do, to ensure interpolators match across pipeline stages. currently there's no such guarantee so this is a safety measure.
-			checkf(!Job.Input.bCompilingForShaderPipeline, TEXT("Cannot combine shader pipeline interpolator removal optimization with secondary compilation at this time"));
 			Job.SecondaryPreprocessOutput = MakeUnique<FShaderPreprocessOutput>();
 			Job.SecondaryPreprocessOutput->bIsSecondary = true;
 			Job.SecondaryPreprocessOutput->bSucceeded = Backend->PreprocessShader(Job.Input, Environment, *Job.SecondaryPreprocessOutput);
@@ -1904,7 +1901,7 @@ public:
 		return bSucceeded;
 	}
 
-	static void CombineCodeOutputs(const IShaderFormat* Compiler, FShaderCompileJob& Job)
+	static bool CombineCodeOutputs(const IShaderFormat* Compiler, FShaderCompileJob& Job)
 	{
 		check(Job.SecondaryOutput.IsValid());
 
@@ -1937,6 +1934,17 @@ public:
 		// note: we intentionally keep the secondary output object around so duplicate/cache hit jobs can also dump debug info. it's a bit of memory overhead but preferable
 		// to requiring that enabling/disabling debug information mutates the DDC key (which would be necessary if we cleared it, since we output separate debug information 
 		// for the secondary jobs).
+
+		// validate that the UsedAttributes array, if populated, matches for both jobs. this is necessary since it's used for interpolator removal on later stages, and such
+		// stages may not have a secondary compilation (and so would need to have matching attributes to be run in a pipeline with two different versions of preceding stages)
+		// if we end up in cases where this check is being hit and it's not resolvable by fixing the code to ensure the same attributes are used in both primary and secondary compilation,
+		// then we could potentially merge these here instead, but that _seems_ unlikely to be the case in practice.
+		if (Job.Output.UsedAttributes != Job.SecondaryOutput->UsedAttributes)
+		{
+			Job.Output.Errors.Add(FShaderCompilerError(TEXT("Mismatch in UsedAttributes between primary and secondary compilation; these must match since later stages may not have a secondary compile and so require both versions of the earlier stage to output the same set of attributes.")));
+			return false;
+		}
+		return true;
 	}
 
 	static void InvokeCompile(const IShaderFormat* Compiler, FShaderCompileJob& Job, const FString& WorkingDirectory, FString& OutExceptionMsg, FString& OutExceptionCallstack)
@@ -1997,7 +2005,7 @@ public:
 			{
 				Job.SecondaryOutput->GenerateOutputHash();
 			}
-			CombineCodeOutputs(Compiler, Job);
+			Job.Output.bSucceeded &= CombineCodeOutputs(Compiler, Job);
 		}
 
 		// ensure the target field is set on the job output struct as we use it for validation during serialization
@@ -4775,21 +4783,30 @@ void FShaderPipelineCompileJob::OnComplete(FShaderDebugDataContext& Ctx)
 	bool bShouldDump = !Ctx.DebugSourceFiles.IsEmpty() && (CVarDumpDebugInfoForCacheHits.GetValueOnAnyThread() || !JobStatusPtr->WasCompilationSkipped());
 	if (bShouldDump)
 	{
-		if (Ctx.DebugSourceFiles.Num() == StageJobs.Num())
+		// validate that all expected debug source files for each pipeline job have been written out and in the expected order
+		int32 DebugSourceIndex = 0, DebugSourceCount = Ctx.DebugSourceFiles.Num();
+		TArray<const FShaderCompileJob*> MissingStages;
+		for (int32 Index = 0, Num = StageJobs.Num(); Index < Num; ++Index)
+		{
+			FShaderCompileJob* StageJob = StageJobs[Index];
+			if (DebugSourceIndex >= DebugSourceCount || (Ctx.DebugSourceFiles[DebugSourceIndex++].Key != StageJob->Input.Target.GetFrequency()))
+			{
+				MissingStages.Add(StageJob);
+				continue;
+			}
+
+			if (StageJob->SecondaryPreprocessOutput.IsValid() && (DebugSourceIndex >= DebugSourceCount || (Ctx.DebugSourceFiles[DebugSourceIndex++].Key != StageJob->Input.Target.GetFrequency())))
+			{
+				MissingStages.Add(StageJob);
+			}
+		}
+
+		if (MissingStages.IsEmpty())
 		{
 			FShaderCompileWorkerUtil::DumpDebugCompileInput(*this, Ctx);
 		}
 		else
 		{
-			TArray<const FShaderCompileJob*> MissingStages;
-			for (int32 Index = 0, Num = StageJobs.Num(); Index < Num; ++Index)
-			{
-				FShaderCompileJob* StageJob = StageJobs[Index];
-				if (!Ctx.DebugSourceFiles.Contains(StageJob->Input.Target.GetFrequency()))
-				{
-					MissingStages.Add(StageJob);
-				}
-			}
 			FString MissingStagesStr = FString::JoinBy(MissingStages, TEXT("\n"), [](const FShaderCompileJob* StageJob)
 				{
 					bool bDebugInfoEnabled = StageJob->Input.DumpDebugInfoEnabled();
diff --git a/Engine/Source/Runtime/RenderCore/Public/ShaderCompilerCore.h b/Engine/Source/Runtime/RenderCore/Public/ShaderCompilerCore.h
--- a/Engine/Source/Runtime/RenderCore/Public/ShaderCompilerCore.h
+++ b/Engine/Source/Runtime/RenderCore/Public/ShaderCompilerCore.h
@@ -202,7 +202,7 @@ using FShaderCompilerInputHash = FBlake3Hash;
 struct FShaderDebugDataContext
 {
 	bool bIsPipeline = false;
-	TMap<EShaderFrequency, FString> DebugSourceFiles;
+	TArray<TPair<EShaderFrequency, FString>> DebugSourceFiles;
 };
 
 /** Struct that gathers all readonly inputs needed for the compilation of a single shader. */
