Fixed render thread fence not submitting to the actual pipe.

#14110

diff --git a/Engine/Source/Runtime/RenderCore/Private/RenderingThread.cpp b/Engine/Source/Runtime/RenderCore/Private/RenderingThread.cpp
--- a/Engine/Source/Runtime/RenderCore/Private/RenderingThread.cpp
+++ b/Engine/Source/Runtime/RenderCore/Private/RenderingThread.cpp
@@ -1015,7 +1015,7 @@ void FRenderCommandFence::BeginFence(ESyncDepth SyncDepth)
 				UE::Tasks::FTaskEvent PipeEvent { UE_SOURCE_LOCATION };
 				Event.AddPrerequisites(PipeEvent);
 
-				ENQUEUE_RENDER_COMMAND(BeginFence)([PipeEvent = MoveTemp(PipeEvent)](FRHICommandList&) mutable
+				ENQUEUE_RENDER_COMMAND(BeginFence)(Pipe, [PipeEvent = MoveTemp(PipeEvent)](FRHICommandList&) mutable
 				{
 					PipeEvent.Trigger();
 				});
