[shaders] correct oversight in serialization of symbol buffers; should not occur unless we are serializing editoronly data.

diff --git a/Engine/Source/Runtime/RenderCore/Private/ShaderResource.cpp b/Engine/Source/Runtime/RenderCore/Private/ShaderResource.cpp
--- a/Engine/Source/Runtime/RenderCore/Private/ShaderResource.cpp
+++ b/Engine/Source/Runtime/RenderCore/Private/ShaderResource.cpp
@@ -443,15 +443,15 @@ void FShaderMapResourceCode::Serialize(FShaderSerializeContext& Ctx)
 	if (bSerializeEditorOnlyData)
 	{
 		Ar << ShaderEditorOnlyDataEntries;
-	}
-	if (Ar.IsLoading())
-	{
-		// only need to set SymbolBuffer when loading, saving is handled automatically by the serialization of ShaderCodeResources
-		// (this just sets the FSharedBuffer reference in a second place to avoid needing to fix up a bunch of code which retrieves
-		// symbols from the editoronly data)
-		for (int32 CodeIndex = 0; CodeIndex < ShaderCodeResources.Num(); ++CodeIndex)
+		if (Ar.IsLoading())
 		{
-			ShaderEditorOnlyDataEntries[CodeIndex].ConditionalSetSymbolBuffer(ShaderCodeResources[CodeIndex].GetSymbolsBuffer());
+			// only need to set SymbolBuffer when loading, saving is handled automatically by the serialization of ShaderCodeResources
+			// (this just sets the FSharedBuffer reference in a second place to avoid needing to fix up a bunch of code which retrieves
+			// symbols from the editoronly data)
+			for (int32 CodeIndex = 0; CodeIndex < ShaderCodeResources.Num(); ++CodeIndex)
+			{
+				ShaderEditorOnlyDataEntries[CodeIndex].ConditionalSetSymbolBuffer(ShaderCodeResources[CodeIndex].GetSymbolsBuffer());
+			}
 		}
 	}
 #endif // WITH_EDITORONLY_DATA
