ShaderTypeRegistration should not be inlined

Mainly because of the Add() which itself has a bunch of inlined code.
It's one-time registration so inlining is a waste of space.

diff --git a/Engine/Source/Runtime/RenderCore/Private/Shader.cpp b/Engine/Source/Runtime/RenderCore/Private/Shader.cpp
--- a/Engine/Source/Runtime/RenderCore/Private/Shader.cpp
+++ b/Engine/Source/Runtime/RenderCore/Private/Shader.cpp
@@ -315,6 +315,14 @@ FShaderType::~FShaderType()
 
 bool FShaderTypeRegistration::bShaderTypesInitialized = false;
 static TArray<const FShaderTypeRegistration*>* GShaderTypeRegistrationInstances = nullptr;
+
+FShaderTypeRegistration::FShaderTypeRegistration(TFunctionRef<FShaderType& ()> LazyShaderTypeAccessor)
+	: LazyShaderTypeAccessor(LazyShaderTypeAccessor)
+{
+	checkf(!AreShaderTypesInitialized(), TEXT("Shader type was loaded too late, use ELoadingPhase::PostConfigInit on your module to cause it to load earlier. This shader will not be compiled or function.  The module and shader will appear in the callstack of this assert."));
+	GetInstances().Add(this);
+}
+
 TArray<const FShaderTypeRegistration*>& FShaderTypeRegistration::GetInstances()
 {
 	if (GShaderTypeRegistrationInstances == nullptr)
diff --git a/Engine/Source/Runtime/RenderCore/Public/Shader.h b/Engine/Source/Runtime/RenderCore/Public/Shader.h
--- a/Engine/Source/Runtime/RenderCore/Public/Shader.h
+++ b/Engine/Source/Runtime/RenderCore/Public/Shader.h
@@ -1582,12 +1582,7 @@ inline FString LexToString(FShaderType::EShaderTypeForDynamicCast ShaderType)
 class FShaderTypeRegistration
 {
 public:
-	FShaderTypeRegistration(TFunctionRef<FShaderType& ()> LazyShaderTypeAccessor)
-		: LazyShaderTypeAccessor(LazyShaderTypeAccessor)
-	{
-		checkf(!AreShaderTypesInitialized(), TEXT("Shader type was loaded too late, use ELoadingPhase::PostConfigInit on your module to cause it to load earlier. This shader will not be compiled or function.  The module and shader will appear in the callstack of this assert."));
-		GetInstances().Add(this);
-	}
+	RENDERCORE_API FShaderTypeRegistration(TFunctionRef<FShaderType& ()> LazyShaderTypeAccessor);
 
 	static RENDERCORE_API TArray<const FShaderTypeRegistration*>& GetInstances();
 
