Removing unused backwards compatibility.

diff --git a/Engine/Source/Runtime/RenderCore/Private/PipelineCacheUtilities.cpp b/Engine/Source/Runtime/RenderCore/Private/PipelineCacheUtilities.cpp
--- a/Engine/Source/Runtime/RenderCore/Private/PipelineCacheUtilities.cpp
+++ b/Engine/Source/Runtime/RenderCore/Private/PipelineCacheUtilities.cpp
@@ -654,14 +654,7 @@ bool UE::PipelineCacheUtilities::LoadStablePipelineCacheFile(const FString& File
 	UE::PipelineCacheUtilities::Private::FIndexedSHAHashAndFNameProxyArchive MemReader(PlainMemReader);
 
 	uint32 CacheVersion = FPipelineCacheFileFormatCurrentVersion;
-	if (Header.Version >= UE::PipelineCacheUtilities::Private::FStablePipelineCacheSerializedHeader::EVersion::AddingPipelineCacheVersion)
-	{
-		MemReader << CacheVersion;
-	}
-	else
-	{
-		CacheVersion = 26; //EPipelineCacheFileFormatVersions::BeforeStableCacheVersioning
-	}
+	MemReader << CacheVersion;
 
 	MemReader.SetGameNetVer(CacheVersion);
 
