[MRG] Better filtering for DumpGPU debug node.

Problem: We don't have a way to filter multiple passes with DumpGPU in the MRG debug node

Solutions:
- Add a pass path filters parameter the Debug node DumpGPU section of the MovieRenderGraph.
- Extend DumpGPU to be able to filter from multiple pass paths.

#rnx
#tests win64 editor : generate DumpGPU via MRG

diff --git a/Engine/Plugins/MovieScene/MovieRenderPipeline/Source/MovieRenderPipelineCore/Private/Graph/MovieGraphDefaultRenderer.cpp b/Engine/Plugins/MovieScene/MovieRenderPipeline/Source/MovieRenderPipelineCore/Private/Graph/MovieGraphDefaultRenderer.cpp
--- a/Engine/Plugins/MovieScene/MovieRenderPipeline/Source/MovieRenderPipelineCore/Private/Graph/MovieGraphDefaultRenderer.cpp
+++ b/Engine/Plugins/MovieScene/MovieRenderPipeline/Source/MovieRenderPipelineCore/Private/Graph/MovieGraphDefaultRenderer.cpp
@@ -357,7 +357,7 @@ void UMovieGraphDefaultRenderer::Render(const FMovieGraphTimeStepData& InTimeSte
 		if (bIsRendering && bWantsDumpGpuDebugCapture && bIsDumpGPUCorrectFrame)
 		{
 			const FString FinalPath = UE::MovieGraph::Capture::PrepareCapturePath(GetOwningGraph(), InTimeStepData.EvaluatedConfig, DebugSettings->DumpGpuCaptureFilePathFormat, TEXT(""), true);
-			const FString Cmd = FString::Printf(TEXT("-noexplorer -path=%s"), *FinalPath);
+			const FString Cmd = FString::Printf(TEXT("-noexplorer -path=\"%s\" -pass_path_filters=\"%s\""), *FinalPath, *DebugSettings->DumpGpuCapturePassPathFilters);
 			const FString Context = FString::Printf(TEXT("MRQ Frame: %d"), InTimeStepData.RootFrameNumber.Value);
 			ScopedDumpGPUCapture = MakeUnique<UE::RenderCore::DumpGPU::FDumpScope>(true, *Cmd, *Context);
 		}
diff --git a/Engine/Plugins/MovieScene/MovieRenderPipeline/Source/MovieRenderPipelineCore/Public/Graph/Nodes/MovieGraphDebugNode.h b/Engine/Plugins/MovieScene/MovieRenderPipeline/Source/MovieRenderPipelineCore/Public/Graph/Nodes/MovieGraphDebugNode.h
--- a/Engine/Plugins/MovieScene/MovieRenderPipeline/Source/MovieRenderPipelineCore/Public/Graph/Nodes/MovieGraphDebugNode.h
+++ b/Engine/Plugins/MovieScene/MovieRenderPipeline/Source/MovieRenderPipelineCore/Public/Graph/Nodes/MovieGraphDebugNode.h
@@ -50,6 +50,9 @@ public:
 	UPROPERTY(EditAnywhere, BlueprintReadWrite, Category = Overrides, meta = (InlineEditConditionToggle))
 	uint8 bOverride_DumpGpuCaptureFilePathFormat : 1;
 
+	UPROPERTY(EditAnywhere, BlueprintReadWrite, Category = Overrides, meta = (InlineEditConditionToggle))
+	uint8 bOverride_DumpGpuCapturePassPathFilters : 1;
+
 	/** 
 	* If true, automatically trigger RenderDoc to capture rendering information. RenderDoc plugin must be enabled, 
 	* and the editor must have been launched with -AttachRenderDoc. Resulting capture will be in /Saved/RenderDocCaptures.
@@ -100,6 +103,13 @@ public:
 	*/
 	UPROPERTY(EditAnywhere, BlueprintReadWrite, Category = "DumpGPU", meta = (EditCondition = "bOverride_DumpGpuCaptureFilePathFormat"))
 	FString DumpGpuCaptureFilePathFormat = TEXT("{sequence_name}.{frame_number}_DumpGPU");
+
+	/**
+	* If bCaptureFramesWithDumpGpu is true, wildcard enabled filters to apply to the path of the passes to dump.
+	* Support multiple filters separated by "|". For example : *PostProcessing/DOF*|*PostProcessing/MotionBlur*
+	*/
+	UPROPERTY(EditAnywhere, BlueprintReadWrite, Category = "DumpGPU", meta = (EditCondition = "bOverride_DumpGpuCapturePassPathFilters"))
+	FString DumpGpuCapturePassPathFilters = TEXT("*");
 };
 
 #undef UE_API
diff --git a/Engine/Source/Runtime/RenderCore/Private/DumpGPU.cpp b/Engine/Source/Runtime/RenderCore/Private/DumpGPU.cpp
--- a/Engine/Source/Runtime/RenderCore/Private/DumpGPU.cpp
+++ b/Engine/Source/Runtime/RenderCore/Private/DumpGPU.cpp
@@ -245,6 +245,7 @@ public:
 	int32 DumpedFrameId = 0;
 	int32 FrameCount = 1;
 	FName UploadResourceCompressionName;
+	TArray<FString> PassPathFilters;
 	FString DumpingDirectoryPath;
 	FString Context;
 	FDateTime Time;
@@ -1916,26 +1917,38 @@ public:
 	bool IsDumpingPass(const FRDGPass* Pass)
 	{
 		FString RootWildcardString = GDumpGPURootCVar.GetValueOnRenderThread();
-		FWildcardString WildcardFilter(RootWildcardString);
+		FWildcardString WildcardRootFilter(RootWildcardString);
 
 		bool bDumpPass = (RootWildcardString == TEXT("*"));
 
 		if (!bDumpPass)
 		{
-			bDumpPass = WildcardFilter.IsMatch(Pass->GetEventName().GetTCHAR());
+			bDumpPass = WildcardRootFilter.IsMatch(Pass->GetEventName().GetTCHAR());
 		}
 
+		FString PassPath = Pass->GetEventName().GetTCHAR();
+
 #if RDG_EVENTS
-		for (FRDGScope const* ParentScope = Pass->GetScope(); !bDumpPass && ParentScope; ParentScope = ParentScope->Parent)
+		for (FRDGScope const* ParentScope = Pass->GetScope(); ParentScope; ParentScope = ParentScope->Parent)
 		{
 			if (FRDGScope_RHI const* RHIScope = ParentScope->Get<FRDGScope_RHI>())
 			{
 				FRHIBreadcrumb::FBuffer Buffer;
-				bDumpPass = WildcardFilter.IsMatch(RHIScope->GetTCHAR(Buffer));
+				bDumpPass |= WildcardRootFilter.IsMatch(RHIScope->GetTCHAR(Buffer));
+				PassPath = RHIScope->GetTCHAR(Buffer) / PassPath;
 			}
 		}
 #endif // RDG_EVENTS
 
+		bool bPathPassAllowed = false;
+
+		for (const FString& PassPathFilter : PassPathFilters)
+		{
+			FWildcardString WildcardPassPathFilter(PassPathFilter);
+			bPathPassAllowed |= WildcardPassPathFilter.IsMatch(*PassPath);
+		}
+		bDumpPass &= bPathPassAllowed;
+
 		return bDumpPass;
 	}
 
@@ -2043,6 +2056,11 @@ FString FRDGBuilder::BeginResourceDump(const TCHAR* Cmd, const TCHAR* Context)
 		NewResourceDumpContext->DumpingDirectoryPath = Parameters[TEXT("path")];
 	}
 
+	if (Parameters.Contains(TEXT("pass_path_filters")))
+	{
+		Parameters[TEXT("pass_path_filters")].ParseIntoArray(NewResourceDumpContext->PassPathFilters, TEXT("|"));
+	}
+
 	NewResourceDumpContext->bStream = GDumpGPUStream.GetValueOnGameThread() != 0;
 	NewResourceDumpContext->bEnableDiskWrite = GDumpTestEnableDiskWrite.GetValueOnGameThread() != 0;
 	NewResourceDumpContext->DeltaTime = FApp::GetDeltaTime();
