Fix default call to DumpGPU following CL 46209999

Problem: new filtering option would filter out everying by default, we want the absence of the option to let everything in (same as before it was introduced).
Solutions: add an exception when option is empty

#rnx
#tests win64 editor : generate DumpGPU with and without MRG, with and without the -pass_path_filters parameter

diff --git a/Engine/Source/Runtime/RenderCore/Private/DumpGPU.cpp b/Engine/Source/Runtime/RenderCore/Private/DumpGPU.cpp
--- a/Engine/Source/Runtime/RenderCore/Private/DumpGPU.cpp
+++ b/Engine/Source/Runtime/RenderCore/Private/DumpGPU.cpp
@@ -1940,14 +1940,17 @@ public:
 		}
 #endif // RDG_EVENTS
 
-		bool bPathPassAllowed = false;
-
-		for (const FString& PassPathFilter : PassPathFilters)
+		if (!PassPathFilters.IsEmpty())
 		{
-			FWildcardString WildcardPassPathFilter(PassPathFilter);
-			bPathPassAllowed |= WildcardPassPathFilter.IsMatch(*PassPath);
+			bool bPathPassAllowed = false;
+
+			for (const FString& PassPathFilter : PassPathFilters)
+			{
+				FWildcardString WildcardPassPathFilter(PassPathFilter);
+				bPathPassAllowed |= WildcardPassPathFilter.IsMatch(*PassPath);
+			}
+			bDumpPass &= bPathPassAllowed;
 		}
-		bDumpPass &= bPathPassAllowed;
 
 		return bDumpPass;
 	}
