Bubble up display information from RHI implementations

diff --git a/Engine/Source/Runtime/Core/Private/ColorManagement/ColorSpace.cpp b/Engine/Source/Runtime/Core/Private/ColorManagement/ColorSpace.cpp
--- a/Engine/Source/Runtime/Core/Private/ColorManagement/ColorSpace.cpp
+++ b/Engine/Source/Runtime/Core/Private/ColorManagement/ColorSpace.cpp
@@ -76,6 +76,13 @@ const FColorSpace& FColorSpace::GetSRGB()
 	return ColorSpaceSRGB;
 }
 
+const FColorSpace& FColorSpace::GetRec2020()
+{
+	static const FColorSpace ColorSpaceRec2020 = FColorSpace(EColorSpace::Rec2020);
+
+	return ColorSpaceRec2020;
+}
+
 static bool IsSRGBChromaticities(const TStaticArray<FVector2d, 4>& Chromaticities, double Tolerance = 1.e-7)
 {
 	return	Chromaticities[0].Equals(FVector2d(0.64, 0.33), Tolerance) &&
diff --git a/Engine/Source/Runtime/Core/Public/ColorManagement/ColorSpace.h b/Engine/Source/Runtime/Core/Public/ColorManagement/ColorSpace.h
--- a/Engine/Source/Runtime/Core/Public/ColorManagement/ColorSpace.h
+++ b/Engine/Source/Runtime/Core/Public/ColorManagement/ColorSpace.h
@@ -118,6 +118,13 @@ public:
 	 */
 	static CORE_API const FColorSpace& GetSRGB();
 
+	/**
+	 * Get a (statically cached) Rec. 2020 color space.
+	 *
+	 * @return FColorSpace Rec. 2020 color space
+	 */
+	static CORE_API const FColorSpace& GetRec2020();
+
 	/** Constructor */
 	FColorSpace() {}
 
diff --git a/Engine/Source/Runtime/D3D12RHI/Private/Windows/WindowsD3D12Device.cpp b/Engine/Source/Runtime/D3D12RHI/Private/Windows/WindowsD3D12Device.cpp
--- a/Engine/Source/Runtime/D3D12RHI/Private/Windows/WindowsD3D12Device.cpp
+++ b/Engine/Source/Runtime/D3D12RHI/Private/Windows/WindowsD3D12Device.cpp
@@ -555,6 +555,15 @@ static bool SupportsDepthBoundsTest(FD3D12DynamicRHI* D3DRHI)
 	return D3DRHI->GetAdapter().IsDepthBoundsTestSupported();
 }
 
+static UE::Color::FColorSpace MakeColorSpaceFromDXGI(const DXGI_OUTPUT_DESC1& OutputDesc)
+{
+	return UE::Color::FColorSpace(
+		FVector2d(OutputDesc.RedPrimary[0], OutputDesc.RedPrimary[1]),
+		FVector2d(OutputDesc.GreenPrimary[0], OutputDesc.GreenPrimary[1]),
+		FVector2d(OutputDesc.BluePrimary[0], OutputDesc.BluePrimary[1]),
+		FVector2d(OutputDesc.WhitePoint[0], OutputDesc.WhitePoint[1]));
+}
+
 bool FD3D12DynamicRHI::SetupDisplayHDRMetaData()
 {
 	// Determines if any displays support HDR
@@ -629,11 +638,27 @@ bool FD3D12DynamicRHI::SetupDisplayHDRMetaData()
 
 					bSupportsHDROutput = true;
 				}
-				FDisplayInformation DisplayInformation{};
+
+				FDisplayInformation& DisplayInformation = DisplayList.Emplace_GetRef();
 				DisplayInformation.bHDRSupported = bDisplaySupportsHDROutput;
 				const RECT& DisplayCoords = OutputDesc.DesktopCoordinates;
 				DisplayInformation.DesktopCoordinates = FIntRect(DisplayCoords.left, DisplayCoords.top, DisplayCoords.right, DisplayCoords.bottom);
-				DisplayList.Add(DisplayInformation);
+				DisplayInformation.MinimumLuminanceInNits = OutputDesc.MinLuminance;
+				DisplayInformation.MaximumLuminanceInNits = OutputDesc.MaxLuminance ?
+					OutputDesc.MaxLuminance :
+					(bDisplaySupportsHDROutput ? 1499.f : 100.f);
+				DisplayInformation.MaximumFullFrameLuminanceInNits = OutputDesc.MaxFullFrameLuminance ?
+					OutputDesc.MaxFullFrameLuminance :
+					(bDisplaySupportsHDROutput ? 799.f : 100.f);
+				const bool bIsGamutValid = OutputDesc.RedPrimary[0] && OutputDesc.RedPrimary[1] &&
+					OutputDesc.GreenPrimary[0] && OutputDesc.GreenPrimary[1] &&
+					OutputDesc.BluePrimary[0] && OutputDesc.BluePrimary[1] &&
+					OutputDesc.WhitePoint[0] && OutputDesc.WhitePoint[1];
+				DisplayInformation.LimitingColorSpace = bIsGamutValid ?
+					MakeColorSpaceFromDXGI(OutputDesc) :
+					(bDisplaySupportsHDROutput ?
+						UE::Color::FColorSpace::GetRec2020() :
+						UE::Color::FColorSpace::GetSRGB());
 			}
 		}
 	}
diff --git a/Engine/Source/Runtime/Engine/Private/Slate/SceneViewport.cpp b/Engine/Source/Runtime/Engine/Private/Slate/SceneViewport.cpp
--- a/Engine/Source/Runtime/Engine/Private/Slate/SceneViewport.cpp
+++ b/Engine/Source/Runtime/Engine/Private/Slate/SceneViewport.cpp
@@ -63,6 +63,10 @@ FSceneViewport::FSceneViewport( FViewportClient* InViewportClient, TSharedPtr<SV
 	, NumTouches(0)
 	, DisplayColorGamut(EDisplayColorGamut::sRGB_D65)
 	, DisplayOutputFormat(EDisplayOutputFormat::SDR_sRGB)
+	, MinimumLuminanceInNits(0.f)
+	, MaximumLuminanceInNits(100.f)
+	, MaximumFullFrameLuminanceInNits(100.f)
+	, LimitingColorSpace(UE::Color::EColorSpace::sRGB)
 {
 	bIsSlateViewport = true;
 	ViewportType = NAME_SceneViewport;
@@ -377,19 +381,28 @@ IStereoRenderTargetManager* RetrieveStereoRenderTargetManager(bool bIsStereoRend
 
 }
 
-void ComputeSceneViewportHDRMetaData(EDisplayOutputFormat& OutDisplayOutputFormat, EDisplayColorGamut& OutDisplayColorGamut, bool& OutbHDRSupported, const FVector2D& WindowTopLeft, const FVector2D& WindowBottomRight, void* OSWindow, bool bIsStereoRenderingAllowed)
+static void ComputeSceneViewportHDRMetaData(FHDRMetaData& OutHDRMetaData, const FVector2D& WindowTopLeft, const FVector2D& WindowBottomRight, void* OSWindow, bool bIsStereoRenderingAllowed)
 {
 	// @todo vreditor switch: This code needs to be called when switching between stereo/non when going immersive.  Seems to always work out that way anyway though? (Probably due to resize)
 	IStereoRenderTargetManager* const StereoRenderTargetManager = RetrieveStereoRenderTargetManager(bIsStereoRenderingAllowed);
 	if (StereoRenderTargetManager != nullptr)
 	{
-		if (StereoRenderTargetManager->HDRGetMetaDataForStereo(OutDisplayOutputFormat, OutDisplayColorGamut, OutbHDRSupported))
+		// TODO: Get the real values for stereo.
+		OutHDRMetaData.MinimumLuminanceInNits = 0.f;
+		OutHDRMetaData.MaximumLuminanceInNits = 100.f;
+		OutHDRMetaData.MaximumFullFrameLuminanceInNits = 100.f;
+		OutHDRMetaData.LimitingColorSpace = UE::Color::FColorSpace::GetSRGB();
+		if (StereoRenderTargetManager->HDRGetMetaDataForStereo(OutHDRMetaData.DisplayOutputFormat, OutHDRMetaData.DisplayColorGamut, OutHDRMetaData.bHDRSupported))
 		{
 			return;
 		}
 	}
-	
-	HDRGetMetaData(OutDisplayOutputFormat, OutDisplayColorGamut, OutbHDRSupported, WindowTopLeft, WindowBottomRight, OSWindow);
+
+	HDRGetMetaDataForWindow(
+		OutHDRMetaData,
+		WindowTopLeft,
+		WindowBottomRight,
+		OSWindow);
 }
 
 void FSceneViewport::OnDrawViewport( const FGeometry& AllottedGeometry, const FSlateRect& MyCullingRect, FSlateWindowElementList& OutDrawElements, int32 LayerId, const FWidgetStyle& InWidgetStyle, bool bParentEnabled )
@@ -412,14 +425,21 @@ void FSceneViewport::OnDrawViewport( const FGeometry& AllottedGeometry, const FS
 		bool bHDRStale = (bIsHDREnabled != bHDRViewport);
 		if (PaintWindow)
 		{
-			bool bNewHDREnabled;
-			EDisplayColorGamut NewDisplayColorGamut;
-			EDisplayOutputFormat NewDisplayOutputFormat;
-			ComputeSceneViewportHDRMetaData(NewDisplayOutputFormat, NewDisplayColorGamut, bNewHDREnabled, PaintWindow->GetPositionInScreen(), PaintWindow->GetPositionInScreen() + PaintWindow->GetSizeInScreen(), PaintWindow->GetNativeWindow()->GetOSWindowHandle(), IsStereoRenderingAllowed());
+			FHDRMetaData OutHDRMetaData;
+			ComputeSceneViewportHDRMetaData(
+				OutHDRMetaData,
+				PaintWindow->GetPositionInScreen(),
+				PaintWindow->GetPositionInScreen() + PaintWindow->GetSizeInScreen(),
+				PaintWindow->GetNativeWindow()->GetOSWindowHandle(),
+				IsStereoRenderingAllowed());
 			// if we manage to get data for the window, we can ignore the global toggle IsHDREnabled since HDRGetMetaData will take both the global flag and the monitor properties
-			bHDRStale = DisplayOutputFormat != NewDisplayOutputFormat;
-			bHDRStale |= DisplayColorGamut != NewDisplayColorGamut;
-			bHDRStale |= bHDRViewport != bNewHDREnabled;
+			bHDRStale = DisplayOutputFormat != OutHDRMetaData.DisplayOutputFormat;
+			bHDRStale |= DisplayColorGamut != OutHDRMetaData.DisplayColorGamut;
+			bHDRStale |= MinimumLuminanceInNits != OutHDRMetaData.MinimumLuminanceInNits;
+			bHDRStale |= MaximumLuminanceInNits != OutHDRMetaData.MaximumLuminanceInNits;
+			bHDRStale |= MaximumFullFrameLuminanceInNits != OutHDRMetaData.MaximumFullFrameLuminanceInNits;
+			bHDRStale |= LimitingColorSpace != OutHDRMetaData.LimitingColorSpace;
+			bHDRStale |= bHDRViewport != OutHDRMetaData.bHDRSupported;
 		}
 
 	    if (GetSizeXY() != DrawSize || bHDRStale)
@@ -1313,6 +1333,26 @@ bool FSceneViewport::GetSceneHDREnabled() const
 	return bHDRViewport;
 }
 
+float FSceneViewport::GetMinimumLuminanceInNits() const
+{
+	return MinimumLuminanceInNits;
+}
+
+float FSceneViewport::GetMaximumLuminanceInNits() const
+{
+	return MaximumLuminanceInNits;
+}
+
+float FSceneViewport::GetMaximumFullFrameLuminanceInNits() const
+{
+	return MaximumFullFrameLuminanceInNits;
+}
+
+const UE::Color::FColorSpace& FSceneViewport::GetLimitingColorSpace() const
+{
+	return LimitingColorSpace;
+}
+
 ESlateViewportDynamicRange FSceneViewport::GetViewportDynamicRange() const
 {
 	return bHDRViewport ? ESlateViewportDynamicRange::HDR : ESlateViewportDynamicRange::SDR;
@@ -1771,7 +1811,15 @@ void FSceneViewport::UpdateViewportRHI(bool bDestroyed, uint32 NewSizeX, uint32
 					WindowBottomRight = Window->GetPositionInScreen() + Window->GetSizeInScreen();
 				}
 			}
-			ComputeSceneViewportHDRMetaData(DisplayOutputFormat, DisplayColorGamut, bHDRViewport, WindowTopLeft, WindowBottomRight, OSWindow, IsStereoRenderingAllowed());
+			FHDRMetaData OutHDRMetaData;
+			ComputeSceneViewportHDRMetaData(OutHDRMetaData, WindowTopLeft, WindowBottomRight, OSWindow, IsStereoRenderingAllowed());
+			DisplayOutputFormat = OutHDRMetaData.DisplayOutputFormat;
+			DisplayColorGamut = OutHDRMetaData.DisplayColorGamut;
+			MinimumLuminanceInNits = OutHDRMetaData.MinimumLuminanceInNits;
+			MaximumLuminanceInNits = OutHDRMetaData.MaximumLuminanceInNits;
+			MaximumFullFrameLuminanceInNits = OutHDRMetaData.MaximumFullFrameLuminanceInNits;
+			LimitingColorSpace = OutHDRMetaData.LimitingColorSpace;
+			bHDRViewport = OutHDRMetaData.bHDRSupported;
 
 			BeginInitResource(this);
 			
@@ -1990,7 +2038,7 @@ void FSceneViewport::OnPlayWorldViewportSwapped( const FSceneViewport& OtherView
 		Invalidate();
 	}
 
-	// Play world viewports should transfer active stats so it doesn't appear like a seperate viewport
+	// Play world viewports should transfer active stats so it doesn't appear like a separate viewport
 	SwapStatCommands(OtherViewport);
 }
 
diff --git a/Engine/Source/Runtime/Engine/Public/Slate/SceneViewport.h b/Engine/Source/Runtime/Engine/Public/Slate/SceneViewport.h
--- a/Engine/Source/Runtime/Engine/Public/Slate/SceneViewport.h
+++ b/Engine/Source/Runtime/Engine/Public/Slate/SceneViewport.h
@@ -245,6 +245,10 @@ public:
 	ENGINE_API virtual EDisplayColorGamut GetDisplayColorGamut() const override;
 	ENGINE_API virtual EDisplayOutputFormat GetDisplayOutputFormat() const override;
 	ENGINE_API virtual bool GetSceneHDREnabled() const override;
+	ENGINE_API virtual float GetMinimumLuminanceInNits() const override;
+	ENGINE_API virtual float GetMaximumLuminanceInNits() const override;
+	ENGINE_API virtual float GetMaximumFullFrameLuminanceInNits() const override;
+	ENGINE_API virtual const UE::Color::FColorSpace& GetLimitingColorSpace() const override;
 	ENGINE_API virtual ESlateViewportDynamicRange GetViewportDynamicRange() const override;
 
 	ENGINE_API void SetViewportSize(uint32 NewSizeX,uint32 NewSizeY);
@@ -450,6 +454,10 @@ private:
 
 	EDisplayColorGamut DisplayColorGamut;
 	EDisplayOutputFormat DisplayOutputFormat;
+	float MinimumLuminanceInNits;
+	float MaximumLuminanceInNits;
+	float MaximumFullFrameLuminanceInNits;
+	UE::Color::FColorSpace LimitingColorSpace;
 
 	/** The optional gamma value to use for this viewport */
 	TOptional<float> ViewportGammaOverride;
diff --git a/Engine/Source/Runtime/Engine/Public/UnrealClient.h b/Engine/Source/Runtime/Engine/Public/UnrealClient.h
--- a/Engine/Source/Runtime/Engine/Public/UnrealClient.h
+++ b/Engine/Source/Runtime/Engine/Public/UnrealClient.h
@@ -8,6 +8,7 @@
 
 #include "CoreMinimal.h"
 #include "InputKeyEventArgs.h"
+#include "ColorManagement/ColorSpace.h"
 #include "Engine/EngineBaseTypes.h"
 #include "UObject/GCObject.h"
 #include "RenderResource.h"
@@ -72,6 +73,10 @@ public:
 	virtual EDisplayColorGamut GetDisplayColorGamut() const { return EDisplayColorGamut::sRGB_D65; }
 	virtual EDisplayOutputFormat GetDisplayOutputFormat() const { return EDisplayOutputFormat::SDR_sRGB; }
 	virtual bool GetSceneHDREnabled() const { return false; }
+	virtual float GetMinimumLuminanceInNits() const { return 0.f; }
+	virtual float GetMaximumLuminanceInNits() const { return 100.f; }
+	virtual float GetMaximumFullFrameLuminanceInNits() const { return 100.f; }
+	virtual const UE::Color::FColorSpace& GetLimitingColorSpace() const { return UE::Color::FColorSpace::GetSRGB(); }
 
 	/**
 	* Handles freezing/unfreezing of rendering
diff --git a/Engine/Source/Runtime/RenderCore/Private/RenderCore.cpp b/Engine/Source/Runtime/RenderCore/Private/RenderCore.cpp
--- a/Engine/Source/Runtime/RenderCore/Private/RenderCore.cpp
+++ b/Engine/Source/Runtime/RenderCore/Private/RenderCore.cpp
@@ -409,6 +409,16 @@ static TAutoConsoleVariable<int32> CVarHDRDisplayMaxLuminance(
 	ECVF_RenderThreadSafe
 );
 
+static TAutoConsoleVariable<bool> CVarHDROverrideOSMaxLuminance(
+	TEXT("r.HDR.OverrideOSMaxLuminance"),
+#if PLATFORM_USES_FIXED_HDR_SETTING
+	false,
+#else
+	true,
+#endif
+	TEXT("Whether to override the maximum display luminance provided by the OS."),
+	ECVF_RenderThreadSafe);
+
 static TAutoConsoleVariable<float> CVarHDRAcesColorMultiplier(
 	TEXT("r.HDR.Aces.SceneColorMultiplier"),
 	1.5f,
@@ -494,23 +504,25 @@ RENDERCORE_API EDisplayColorGamut HDRGetDefaultDisplayColorGamut()
 	return static_cast<EDisplayColorGamut>(FMath::Clamp(CVarDisplayColorGamut.GetValueOnAnyThread(), 0, static_cast<int32>(EDisplayColorGamut::MAX) - 1));
 }
 
-struct FHDRMetaData
-{
-	EDisplayOutputFormat DisplayOutputFormat;
-	EDisplayColorGamut DisplayColorGamut;
-	uint32 MaximumLuminanceInNits;
-	bool bHDRSupported;
-};
-
-
 inline FHDRMetaData HDRGetDefaultMetaData()
 {
-	FHDRMetaData HDRMetaData{};
-	HDRMetaData.DisplayOutputFormat = HDRGetDefaultDisplayOutputFormat();
-	HDRMetaData.DisplayColorGamut = HDRGetDefaultDisplayColorGamut();
-	HDRMetaData.bHDRSupported = IsHDREnabled() && GRHISupportsHDROutput;
-	HDRMetaData.MaximumLuminanceInNits = CVarHDRDisplayMaxLuminance.GetValueOnAnyThread();
-	return HDRMetaData;
+	const bool bHDRSupported = IsHDREnabled() && GRHISupportsHDROutput;
+	return
+	{
+		.DisplayOutputFormat = HDRGetDefaultDisplayOutputFormat(),
+		.DisplayColorGamut = HDRGetDefaultDisplayColorGamut(),
+		.MinimumLuminanceInNits = bHDRSupported ?
+			FMath::Pow(10.f, static_cast<float>(CVarHDRDisplayMinLuminanceLog10.GetValueOnAnyThread())) :
+			0.f,
+		.MaximumLuminanceInNits = bHDRSupported ?
+			static_cast<float>(CVarHDRDisplayMaxLuminance.GetValueOnAnyThread()) :
+			100.f,
+		.MaximumFullFrameLuminanceInNits = bHDRSupported ?
+			static_cast<float>(CVarHDRDisplayMaxLuminance.GetValueOnAnyThread()) :
+			100.f,
+		.LimitingColorSpace = bHDRSupported ? UE::Color::FColorSpace::GetRec2020() : UE::Color::FColorSpace::GetSRGB(),
+		.bHDRSupported = bHDRSupported,
+	};
 }
 
 inline int32 WindowDisplayIntersectionArea(FIntRect WindowRect, FIntRect DisplayRect)
@@ -529,14 +541,24 @@ RENDERCORE_API void HDRAddCustomMetaData(void* OSWindow, EDisplayOutputFormat Di
 	{
 		return;
 	}
-	FHDRMetaData HDRMetaData{};
+	FScopeLock Lock(&GWindowsWithDefaultParamsCS);
+	FHDRMetaData& HDRMetaData = GWindowsWithDefaultParams.Emplace(OSWindow);
 	HDRMetaData.DisplayOutputFormat = DisplayOutputFormat;
 	HDRMetaData.DisplayColorGamut = DisplayColorGamut;
+	// TODO: Plumb color space.
+	HDRMetaData.LimitingColorSpace = bHDREnabled ?
+		UE::Color::FColorSpace::GetRec2020() :
+		UE::Color::FColorSpace::GetSRGB();
+	HDRMetaData.MinimumLuminanceInNits = bHDREnabled ?
+		FMath::Pow(10.f, static_cast<float>(CVarHDRDisplayMinLuminanceLog10.GetValueOnAnyThread())) :
+		0.f,
+	HDRMetaData.MaximumLuminanceInNits = bHDREnabled ?
+		CVarHDRDisplayMaxLuminance.GetValueOnAnyThread() :
+		100.f;
+	HDRMetaData.MaximumFullFrameLuminanceInNits = bHDREnabled ?
+		CVarHDRDisplayMaxLuminance.GetValueOnAnyThread() :
+		100.f;
 	HDRMetaData.bHDRSupported = bHDREnabled;
-	HDRMetaData.MaximumLuminanceInNits = CVarHDRDisplayMaxLuminance.GetValueOnAnyThread();
-
-	FScopeLock Lock(&GWindowsWithDefaultParamsCS);
-	GWindowsWithDefaultParams.Add(OSWindow, HDRMetaData);
 }
 
 RENDERCORE_API void HDRRemoveCustomMetaData(void* OSWindow)
@@ -570,21 +592,25 @@ bool HdrHasWindowParamsFromCVars(void* OSWindow, FHDRMetaData& HDRMetaData)
 RENDERCORE_API void HDRGetMetaData(EDisplayOutputFormat& OutDisplayOutputFormat, EDisplayColorGamut& OutDisplayColorGamut, bool& OutbHDRSupported, 
 								   const FVector2D& WindowTopLeft, const FVector2D& WindowBottomRight, void* OSWindow)
 {
-	FHDRMetaData HDRMetaData;
+	FHDRMetaData OutHDRMetaData;
+	HDRGetMetaDataForWindow(OutHDRMetaData, WindowTopLeft, WindowBottomRight, OSWindow);
+	OutDisplayOutputFormat = OutHDRMetaData.DisplayOutputFormat;
+	OutDisplayColorGamut = OutHDRMetaData.DisplayColorGamut;
+	OutbHDRSupported = OutHDRMetaData.bHDRSupported;
+}
 
+RENDERCORE_API void HDRGetMetaDataForWindow(FHDRMetaData& OutHDRMetaData, const FVector2D& WindowTopLeft, const FVector2D& WindowBottomRight, void* OSWindow)
+{
 #if WITH_EDITOR
 	// this has priority over IsHDREnabled because MovieSceneCapture might request custom parameters
-	if (HdrHasWindowParamsFromCVars(OSWindow, HDRMetaData))
+	if (HdrHasWindowParamsFromCVars(OSWindow, OutHDRMetaData))
 	{
 		return;
 	}
 #endif
 	
-	HDRMetaData = HDRGetDefaultMetaData();
+	OutHDRMetaData = HDRGetDefaultMetaData();
 
-	OutDisplayOutputFormat = HDRMetaData.DisplayOutputFormat;
-	OutDisplayColorGamut = HDRMetaData.DisplayColorGamut;
-	OutbHDRSupported = HDRMetaData.bHDRSupported;
 #if PLATFORM_USES_FIXED_HDR_SETTING
 	// IsHDREnabled() is controlled by code, depending on system and display HDR availability
 	if (!IsHDRAllowed())
@@ -617,15 +643,18 @@ RENDERCORE_API void HDRGetMetaData(EDisplayOutputFormat& OutDisplayOutputFormat,
 		}
 	}
 
-	OutbHDRSupported = DisplayList[BestDisplayForWindow].bHDRSupported;
-	OutDisplayOutputFormat = EDisplayOutputFormat::SDR_sRGB;
-	OutDisplayColorGamut = EDisplayColorGamut::sRGB_D65;
+	OutHDRMetaData.bHDRSupported = DisplayList[BestDisplayForWindow].bHDRSupported;
+	OutHDRMetaData.DisplayOutputFormat = EDisplayOutputFormat::SDR_sRGB;
+	OutHDRMetaData.DisplayColorGamut = EDisplayColorGamut::sRGB_D65;
+	OutHDRMetaData.MinimumLuminanceInNits = OutHDRMetaData.bHDRSupported ? DisplayList[BestDisplayForWindow].MinimumLuminanceInNits : 0.f;
+	OutHDRMetaData.MaximumLuminanceInNits = OutHDRMetaData.bHDRSupported ? DisplayList[BestDisplayForWindow].MaximumLuminanceInNits : 100.f;
+	OutHDRMetaData.MaximumFullFrameLuminanceInNits = OutHDRMetaData.bHDRSupported ? DisplayList[BestDisplayForWindow].MaximumFullFrameLuminanceInNits : 100.f;
+	OutHDRMetaData.LimitingColorSpace = OutHDRMetaData.bHDRSupported ? DisplayList[BestDisplayForWindow].LimitingColorSpace : UE::Color::FColorSpace::GetSRGB();
 
-	if (OutbHDRSupported)
+	if (OutHDRMetaData.bHDRSupported)
 	{
-		FPlatformMisc::ChooseHDRDeviceAndColorGamut(GRHIVendorId, CVarHDRDisplayMaxLuminance.GetValueOnAnyThread(), OutDisplayOutputFormat, OutDisplayColorGamut);
+		FPlatformMisc::ChooseHDRDeviceAndColorGamut(GRHIVendorId, CVarHDRDisplayMaxLuminance.GetValueOnAnyThread(), OutHDRMetaData.DisplayOutputFormat, OutHDRMetaData.DisplayColorGamut);
 	}
-
 }
 
 RENDERCORE_API float HDRGetDisplayMaximumLuminance()
diff --git a/Engine/Source/Runtime/RenderCore/Public/HDRHelper.h b/Engine/Source/Runtime/RenderCore/Public/HDRHelper.h
--- a/Engine/Source/Runtime/RenderCore/Public/HDRHelper.h
+++ b/Engine/Source/Runtime/RenderCore/Public/HDRHelper.h
@@ -3,6 +3,7 @@
 #pragma once
 
 #include "CoreMinimal.h"
+#include "ColorManagement/ColorSpace.h"
 
 RENDERCORE_API bool IsHDREnabled();
 RENDERCORE_API bool IsHDRAllowed();
@@ -12,6 +13,18 @@ RENDERCORE_API bool IsDefaultBackBufferLinearSDR();
 RENDERCORE_API void HDRGetMetaData(EDisplayOutputFormat& OutDisplayOutputFormat, EDisplayColorGamut& OutDisplayColorGamut, bool& OutbHDRSupported,
 								   const FVector2D& WindowTopLeft, const FVector2D& WindowBottomRight, void* OSWindow);
 
+struct FHDRMetaData
+{
+	EDisplayOutputFormat DisplayOutputFormat;
+	EDisplayColorGamut DisplayColorGamut;
+	float MinimumLuminanceInNits;
+	float MaximumLuminanceInNits;
+	float MaximumFullFrameLuminanceInNits;
+	UE::Color::FColorSpace LimitingColorSpace;
+	bool bHDRSupported;
+};
+RENDERCORE_API void HDRGetMetaDataForWindow(FHDRMetaData& OutHDRMetaData, const FVector2D& WindowTopLeft, const FVector2D& WindowBottomRight, void* OSWindow);
+
 RENDERCORE_API void HDRConfigureCVars(bool bIsHDREnabled, uint32 DisplayNits, bool bFromGameSettings);
 RENDERCORE_API EDisplayOutputFormat HDRGetDefaultDisplayOutputFormat();
 RENDERCORE_API EDisplayColorGamut HDRGetDefaultDisplayColorGamut();
@@ -24,13 +37,12 @@ RENDERCORE_API void ConvertPixelDataToSCRGB(TArray<FLinearColor>& InOutRawPixels
 
 struct FDisplayInformation
 {
-	FDisplayInformation()
-		: DesktopCoordinates(0, 0, 0, 0)
-		, bHDRSupported(false)
-	{}
-
-	FIntRect DesktopCoordinates;
-	bool   bHDRSupported;
+	FIntRect DesktopCoordinates = FIntRect(0, 0, 0, 0);
+	bool bHDRSupported = false;
+	float MinimumLuminanceInNits = 0.f;
+	float MaximumLuminanceInNits = 100.f;
+	float MaximumFullFrameLuminanceInNits = 100.f;
+	UE::Color::FColorSpace LimitingColorSpace = UE::Color::FColorSpace::GetSRGB();
 };
 
 using FDisplayInformationArray = TArray<FDisplayInformation>;
@@ -78,4 +90,4 @@ inline void GetACESTonemapParameters(FACESTonemapParametersType& ACESTonemapPara
 	ACESTonemapParameters.ACESCoefsHigh_4 = TonemapperParams.ACESCoefsHigh_4;
 	ACESTonemapParameters.ACESSceneColorMultiplier = TonemapperParams.ACESSceneColorMultiplier;
 	ACESTonemapParameters.ACESGamutCompression = FMath::Clamp(TonemapperParams.ACESGamutCompression, 0.0f, 1.0f);
-}
\ No newline at end of file
+}
diff --git a/Engine/Source/Runtime/Renderer/Private/PostProcess/PostProcessDeviceEncodingOnly.cpp b/Engine/Source/Runtime/Renderer/Private/PostProcess/PostProcessDeviceEncodingOnly.cpp
--- a/Engine/Source/Runtime/Renderer/Private/PostProcess/PostProcessDeviceEncodingOnly.cpp
+++ b/Engine/Source/Runtime/Renderer/Private/PostProcess/PostProcessDeviceEncodingOnly.cpp
@@ -40,18 +40,22 @@ FDeviceEncodingOnlyOutputDeviceParameters GetDeviceEncodingOnlyOutputDeviceParam
 	static TConsoleVariableData<float>* CVarOutputGamma = IConsoleManager::Get().FindTConsoleVariableDataFloat(TEXT("r.TonemapperGamma"));
 
 	EDisplayOutputFormat OutputDeviceValue;
+	bool bIsHDR;
 
 	if (Family.SceneCaptureSource == SCS_FinalColorHDR)
 	{
 		OutputDeviceValue = EDisplayOutputFormat::HDR_LinearNoToneCurve;
+		bIsHDR = true;
 	}
 	else if (Family.SceneCaptureSource == SCS_FinalToneCurveHDR)
 	{
 		OutputDeviceValue = EDisplayOutputFormat::HDR_LinearWithToneCurve;
+		bIsHDR = true;
 	}
 	else
 	{
 		OutputDeviceValue = Family.RenderTarget->GetDisplayOutputFormat();
+		bIsHDR = Family.RenderTarget->GetSceneHDREnabled();
 	}
 
 	float Gamma = CVarOutputGamma->GetValueOnRenderThread();
@@ -66,6 +70,7 @@ FDeviceEncodingOnlyOutputDeviceParameters GetDeviceEncodingOnlyOutputDeviceParam
 	if (Gamma > 0.0f && (OutputDeviceValue == EDisplayOutputFormat::SDR_sRGB || OutputDeviceValue == EDisplayOutputFormat::SDR_Rec709))
 	{
 		OutputDeviceValue = EDisplayOutputFormat::SDR_ExplicitGammaMapping;
+		bIsHDR = false;
 	}
 
 	FVector InvDisplayGammaValue;
@@ -77,26 +82,10 @@ FDeviceEncodingOnlyOutputDeviceParameters GetDeviceEncodingOnlyOutputDeviceParam
 	Parameters.InverseGamma = (FVector3f)InvDisplayGammaValue;
 	Parameters.OutputDevice = static_cast<uint32>(OutputDeviceValue);
 	Parameters.OutputGamut = static_cast<uint32>(Family.RenderTarget->GetDisplayColorGamut());
-	switch (OutputDeviceValue)
-	{
-	case EDisplayOutputFormat::SDR_sRGB:
-	case EDisplayOutputFormat::SDR_Rec709:
-	case EDisplayOutputFormat::SDR_ExplicitGammaMapping:
-		Parameters.OutputMaxLuminance = 100.f;
-		break;
-	case EDisplayOutputFormat::HDR_ACES_1000nit_ST2084:
-	case EDisplayOutputFormat::HDR_ACES_2000nit_ST2084:
-	case EDisplayOutputFormat::HDR_ACES_1000nit_ScRGB:
-	case EDisplayOutputFormat::HDR_ACES_2000nit_ScRGB:
-	case EDisplayOutputFormat::HDR_LinearEXR:
-	case EDisplayOutputFormat::HDR_LinearNoToneCurve:
-	case EDisplayOutputFormat::HDR_LinearWithToneCurve:
-		// TODO: Plumb luminance through FSceneViewFamily or related.
-		Parameters.OutputMaxLuminance = HDRGetDisplayMaximumLuminance();
-		break;
-	default:
-		check(false);
-	}
+	static const IConsoleVariable* const CVarHDROverrideOSMaxLuminance = IConsoleManager::Get().FindConsoleVariable(TEXT("r.HDR.OverrideOSMaxLuminance"));
+	Parameters.OutputMaxLuminance = CVarHDROverrideOSMaxLuminance->GetBool() ?
+		(bIsHDR ? HDRGetDisplayMaximumLuminance() : 100.f) :
+		Family.RenderTarget->GetMaximumLuminanceInNits();
 	return Parameters;
 }
 
diff --git a/Engine/Source/Runtime/Renderer/Private/PostProcess/PostProcessTonemap.cpp b/Engine/Source/Runtime/Renderer/Private/PostProcess/PostProcessTonemap.cpp
--- a/Engine/Source/Runtime/Renderer/Private/PostProcess/PostProcessTonemap.cpp
+++ b/Engine/Source/Runtime/Renderer/Private/PostProcess/PostProcessTonemap.cpp
@@ -229,18 +229,22 @@ FTonemapperOutputDeviceParameters GetTonemapperOutputDeviceParameters(const FSce
 	static TConsoleVariableData<float>* CVarOutputGamma = IConsoleManager::Get().FindTConsoleVariableDataFloat(TEXT("r.TonemapperGamma"));
 
 	EDisplayOutputFormat OutputDeviceValue;
+	bool bIsHDR;
 
 	if (Family.SceneCaptureSource == SCS_FinalColorHDR)
 	{
 		OutputDeviceValue = EDisplayOutputFormat::HDR_LinearNoToneCurve;
+		bIsHDR = true;
 	}
 	else if (Family.SceneCaptureSource == SCS_FinalToneCurveHDR)
 	{
 		OutputDeviceValue = EDisplayOutputFormat::HDR_LinearWithToneCurve;
+		bIsHDR = true;
 	}
 	else
 	{
 		OutputDeviceValue = Family.RenderTarget->GetDisplayOutputFormat();
+		bIsHDR = Family.RenderTarget->GetSceneHDREnabled();
 	}
 
 	float Gamma = CVarOutputGamma->GetValueOnRenderThread();
@@ -255,6 +259,7 @@ FTonemapperOutputDeviceParameters GetTonemapperOutputDeviceParameters(const FSce
 	if (Gamma > 0.0f && (OutputDeviceValue == EDisplayOutputFormat::SDR_sRGB || OutputDeviceValue == EDisplayOutputFormat::SDR_Rec709))
 	{
 		OutputDeviceValue = EDisplayOutputFormat::SDR_ExplicitGammaMapping;
+		bIsHDR = false;
 	}
 
 	FVector3f InvDisplayGammaValue;
@@ -266,26 +271,10 @@ FTonemapperOutputDeviceParameters GetTonemapperOutputDeviceParameters(const FSce
 	Parameters.InverseGamma = InvDisplayGammaValue;
 	Parameters.OutputDevice = static_cast<uint32>(OutputDeviceValue);
 	Parameters.OutputGamut = static_cast<uint32>(Family.RenderTarget->GetDisplayColorGamut());
-	switch (OutputDeviceValue)
-	{
-	case EDisplayOutputFormat::SDR_sRGB:
-	case EDisplayOutputFormat::SDR_Rec709:
-	case EDisplayOutputFormat::SDR_ExplicitGammaMapping:
-		Parameters.OutputMaxLuminance = 100.f;
-		break;
-	case EDisplayOutputFormat::HDR_ACES_1000nit_ST2084:
-	case EDisplayOutputFormat::HDR_ACES_2000nit_ST2084:
-	case EDisplayOutputFormat::HDR_ACES_1000nit_ScRGB:
-	case EDisplayOutputFormat::HDR_ACES_2000nit_ScRGB:
-	case EDisplayOutputFormat::HDR_LinearEXR:
-	case EDisplayOutputFormat::HDR_LinearNoToneCurve:
-	case EDisplayOutputFormat::HDR_LinearWithToneCurve:
-		// TODO: Plumb luminance through FSceneViewFamily or related.
-		Parameters.OutputMaxLuminance = HDRGetDisplayMaximumLuminance();
-		break;
-	default:
-		check(false);
-	}
+	static const IConsoleVariable* const CVarHDROverrideOSMaxLuminance = IConsoleManager::Get().FindConsoleVariable(TEXT("r.HDR.OverrideOSMaxLuminance"));
+	Parameters.OutputMaxLuminance = CVarHDROverrideOSMaxLuminance->GetBool() ?
+		(bIsHDR ? HDRGetDisplayMaximumLuminance() : 100.f) :
+		Family.RenderTarget->GetMaximumLuminanceInNits();
 	return Parameters;
 }
 
diff --git a/Engine/Source/Runtime/Windows/D3D11RHI/Private/Windows/WindowsD3D11Device.cpp b/Engine/Source/Runtime/Windows/D3D11RHI/Private/Windows/WindowsD3D11Device.cpp
--- a/Engine/Source/Runtime/Windows/D3D11RHI/Private/Windows/WindowsD3D11Device.cpp
+++ b/Engine/Source/Runtime/Windows/D3D11RHI/Private/Windows/WindowsD3D11Device.cpp
@@ -468,6 +468,15 @@ const DisplayChromacities DisplayChromacityList[] =
 	{ 0.71300f, 0.29300f, 0.16500f, 0.83000f, 0.12800f, 0.04400f, 0.32168f, 0.33767f }, // EDisplayColorGamut::ACEScg_D60
 };
 
+static UE::Color::FColorSpace MakeColorSpaceFromDXGI(const DXGI_OUTPUT_DESC1& OutputDesc)
+{
+	return UE::Color::FColorSpace(
+		FVector2d(OutputDesc.RedPrimary[0], OutputDesc.RedPrimary[1]),
+		FVector2d(OutputDesc.GreenPrimary[0], OutputDesc.GreenPrimary[1]),
+		FVector2d(OutputDesc.BluePrimary[0], OutputDesc.BluePrimary[1]),
+		FVector2d(OutputDesc.WhitePoint[0], OutputDesc.WhitePoint[1]));
+}
+
 bool FD3D11DynamicRHI::SetupDisplayHDRMetaData()
 {
 	check(GetDevice());
@@ -550,11 +559,26 @@ bool FD3D11DynamicRHI::SetupDisplayHDRMetaData()
 					bSupportsHDROutput = true;
 				}
 
-				FDisplayInformation DisplayInformation{};
+				FDisplayInformation& DisplayInformation = DisplayList.Emplace_GetRef();
 				DisplayInformation.bHDRSupported = bDisplaySupportsHDROutput;
 				const RECT& DisplayCoords = OutputDesc.DesktopCoordinates;
 				DisplayInformation.DesktopCoordinates = FIntRect(DisplayCoords.left, DisplayCoords.top, DisplayCoords.right, DisplayCoords.bottom);
-				DisplayList.Add(DisplayInformation);
+				DisplayInformation.MinimumLuminanceInNits = OutputDesc.MinLuminance;
+				DisplayInformation.MaximumLuminanceInNits = OutputDesc.MaxLuminance ?
+					OutputDesc.MaxLuminance :
+					(bDisplaySupportsHDROutput ? 1499.f : 100.f);
+				DisplayInformation.MaximumFullFrameLuminanceInNits = OutputDesc.MaxFullFrameLuminance ?
+					OutputDesc.MaxFullFrameLuminance :
+					(bDisplaySupportsHDROutput ? 799.f : 100.f);
+				const bool bIsGamutValid = OutputDesc.RedPrimary[0] && OutputDesc.RedPrimary[1] &&
+					OutputDesc.GreenPrimary[0] && OutputDesc.GreenPrimary[1] &&
+					OutputDesc.BluePrimary[0] && OutputDesc.BluePrimary[1] &&
+					OutputDesc.WhitePoint[0] && OutputDesc.WhitePoint[1];
+				DisplayInformation.LimitingColorSpace = bIsGamutValid ?
+					MakeColorSpaceFromDXGI(OutputDesc) :
+					(bDisplaySupportsHDROutput ?
+						UE::Color::FColorSpace::GetRec2020() :
+						UE::Color::FColorSpace::GetSRGB());
 			}
 		}
 	}
