Add dimensions to FRenderTargetPool_CreateTexture debug information.

diff --git a/Engine/Source/Runtime/RenderCore/Private/RenderTargetPool.cpp b/Engine/Source/Runtime/RenderCore/Private/RenderTargetPool.cpp
--- a/Engine/Source/Runtime/RenderCore/Private/RenderTargetPool.cpp
+++ b/Engine/Source/Runtime/RenderCore/Private/RenderTargetPool.cpp
@@ -21,6 +21,8 @@ TRACE_DECLARE_MEMORY_COUNTER(RenderTargetPoolSize, TEXT("RenderTargetPool/Size")
 
 UE_TRACE_EVENT_BEGIN(Cpu, FRenderTargetPool_CreateTexture, NoSync)
 	UE_TRACE_EVENT_FIELD(UE::Trace::WideString, Name)
+	UE_TRACE_EVENT_FIELD(uint32, Width)
+	UE_TRACE_EVENT_FIELD(uint32, Height)
 UE_TRACE_EVENT_END()
 
 static int32 GRenderTargetPoolLogCreationSizes = 0;
@@ -75,7 +77,9 @@ FPooledRenderTarget* FRenderTargetPool::CreateRenderTarget(FRHICommandListBase&
 {
 #if CPUPROFILERTRACE_ENABLED
 	UE_TRACE_LOG_SCOPED_T(Cpu, FRenderTargetPool_CreateTexture, CpuChannel)
-		<< FRenderTargetPool_CreateTexture.Name(Name);
+		<< FRenderTargetPool_CreateTexture.Name(Name)
+		<< FRenderTargetPool_CreateTexture.Width(Desc.Extent.X)
+		<< FRenderTargetPool_CreateTexture.Height(Desc.Extent.Y);
 #endif
 
 	if (GRenderTargetPoolLogCreationSizes > 0)
