Change default HDR display color space from Rec. 2020 to P3-D65

diff --git a/Engine/Source/Runtime/Core/Private/ColorManagement/ColorSpace.cpp b/Engine/Source/Runtime/Core/Private/ColorManagement/ColorSpace.cpp
--- a/Engine/Source/Runtime/Core/Private/ColorManagement/ColorSpace.cpp
+++ b/Engine/Source/Runtime/Core/Private/ColorManagement/ColorSpace.cpp
@@ -76,11 +76,11 @@ const FColorSpace& FColorSpace::GetSRGB()
 	return ColorSpaceSRGB;
 }
 
-const FColorSpace& FColorSpace::GetRec2020()
+const FColorSpace& FColorSpace::GetP3D65()
 {
-	static const FColorSpace ColorSpaceRec2020 = FColorSpace(EColorSpace::Rec2020);
+	static const FColorSpace ColorSpaceP3D65 = FColorSpace(EColorSpace::P3D65);
 
-	return ColorSpaceRec2020;
+	return ColorSpaceP3D65;
 }
 
 static bool IsSRGBChromaticities(const TStaticArray<FVector2d, 4>& Chromaticities, double Tolerance = 1.e-7)
diff --git a/Engine/Source/Runtime/Core/Public/ColorManagement/ColorSpace.h b/Engine/Source/Runtime/Core/Public/ColorManagement/ColorSpace.h
--- a/Engine/Source/Runtime/Core/Public/ColorManagement/ColorSpace.h
+++ b/Engine/Source/Runtime/Core/Public/ColorManagement/ColorSpace.h
@@ -119,11 +119,11 @@ public:
 	static CORE_API const FColorSpace& GetSRGB();
 
 	/**
-	 * Get a (statically cached) Rec. 2020 color space.
+	 * Get a (statically cached) P3-D65 color space.
 	 *
-	 * @return FColorSpace Rec. 2020 color space
+	 * @return FColorSpace P3-D65 color space
 	 */
-	static CORE_API const FColorSpace& GetRec2020();
+	static CORE_API const FColorSpace& GetP3D65();
 
 	/** Constructor */
 	FColorSpace() {}
diff --git a/Engine/Source/Runtime/D3D12RHI/Private/Windows/WindowsD3D12Device.cpp b/Engine/Source/Runtime/D3D12RHI/Private/Windows/WindowsD3D12Device.cpp
--- a/Engine/Source/Runtime/D3D12RHI/Private/Windows/WindowsD3D12Device.cpp
+++ b/Engine/Source/Runtime/D3D12RHI/Private/Windows/WindowsD3D12Device.cpp
@@ -657,7 +657,7 @@ bool FD3D12DynamicRHI::SetupDisplayHDRMetaData()
 				DisplayInformation.LimitingColorSpace = bIsGamutValid ?
 					MakeColorSpaceFromDXGI(OutputDesc) :
 					(bDisplaySupportsHDROutput ?
-						UE::Color::FColorSpace::GetRec2020() :
+						UE::Color::FColorSpace::GetP3D65() :
 						UE::Color::FColorSpace::GetSRGB());
 			}
 		}
diff --git a/Engine/Source/Runtime/RenderCore/Private/RenderCore.cpp b/Engine/Source/Runtime/RenderCore/Private/RenderCore.cpp
--- a/Engine/Source/Runtime/RenderCore/Private/RenderCore.cpp
+++ b/Engine/Source/Runtime/RenderCore/Private/RenderCore.cpp
@@ -520,7 +520,7 @@ inline FHDRMetaData HDRGetDefaultMetaData()
 		.MaximumFullFrameLuminanceInNits = bHDRSupported ?
 			static_cast<float>(CVarHDRDisplayMaxLuminance.GetValueOnAnyThread()) :
 			100.f,
-		.LimitingColorSpace = bHDRSupported ? UE::Color::FColorSpace::GetRec2020() : UE::Color::FColorSpace::GetSRGB(),
+		.LimitingColorSpace = bHDRSupported ? UE::Color::FColorSpace::GetP3D65() : UE::Color::FColorSpace::GetSRGB(),
 		.bHDRSupported = bHDRSupported,
 	};
 }
@@ -547,7 +547,7 @@ RENDERCORE_API void HDRAddCustomMetaData(void* OSWindow, EDisplayOutputFormat Di
 	HDRMetaData.DisplayColorGamut = DisplayColorGamut;
 	// TODO: Plumb color space.
 	HDRMetaData.LimitingColorSpace = bHDREnabled ?
-		UE::Color::FColorSpace::GetRec2020() :
+		UE::Color::FColorSpace::GetP3D65() :
 		UE::Color::FColorSpace::GetSRGB();
 	HDRMetaData.MinimumLuminanceInNits = bHDREnabled ?
 		FMath::Pow(10.f, static_cast<float>(CVarHDRDisplayMinLuminanceLog10.GetValueOnAnyThread())) :
diff --git a/Engine/Source/Runtime/Windows/D3D11RHI/Private/Windows/WindowsD3D11Device.cpp b/Engine/Source/Runtime/Windows/D3D11RHI/Private/Windows/WindowsD3D11Device.cpp
--- a/Engine/Source/Runtime/Windows/D3D11RHI/Private/Windows/WindowsD3D11Device.cpp
+++ b/Engine/Source/Runtime/Windows/D3D11RHI/Private/Windows/WindowsD3D11Device.cpp
@@ -577,7 +577,7 @@ bool FD3D11DynamicRHI::SetupDisplayHDRMetaData()
 				DisplayInformation.LimitingColorSpace = bIsGamutValid ?
 					MakeColorSpaceFromDXGI(OutputDesc) :
 					(bDisplaySupportsHDROutput ?
-						UE::Color::FColorSpace::GetRec2020() :
+						UE::Color::FColorSpace::GetP3D65() :
 						UE::Color::FColorSpace::GetSRGB());
 			}
 		}
