Remove root constants flag from predicate shader bundle CS (unused)

diff --git a/Engine/Source/Runtime/RenderCore/Private/ShaderBundles.cpp b/Engine/Source/Runtime/RenderCore/Private/ShaderBundles.cpp
--- a/Engine/Source/Runtime/RenderCore/Private/ShaderBundles.cpp
+++ b/Engine/Source/Runtime/RenderCore/Private/ShaderBundles.cpp
@@ -47,11 +47,8 @@ void FPredicateShaderBundleCS::ModifyCompilationEnvironment(const FGlobalShaderP
 	{
 		OutEnvironment.CompilerFlags.Add(CFLAG_ForceBindful);
 	}
-	OutEnvironment.CompilerFlags.Add(CFLAG_RootConstants);
-	OutEnvironment.CompilerFlags.Add(CFLAG_SupportsMinimalBindless);
 
 	OutEnvironment.SetDefine(TEXT("THREADGROUP_SIZEX"), ThreadGroupSizeX);
-	OutEnvironment.SetDefine(TEXT("USE_SHADER_ROOT_CONSTANTS"), RHISupportsShaderRootConstants(Parameters.Platform) ? 1 : 0);
 	OutEnvironment.SetDefine(TEXT("BUNDLE_MODE_CS"), int32(ERHIShaderBundleMode::CS));
 	OutEnvironment.SetDefine(TEXT("BUNDLE_MODE_MSPS"), int32(ERHIShaderBundleMode::MSPS));
 	OutEnvironment.SetDefine(TEXT("BUNDLE_MODE_VSPS"), int32(ERHIShaderBundleMode::VSPS));
