Use r.Mobile.VirtualTextures=0 for the Windows client to remove VirtualTextureUpdate overhead since VTs aren't used in performance mode anyway

This attempt uses WindowsDeviceProfiles to only apply to performance mode in the client to avoid breaking the cook.

diff --git a/Engine/Source/Runtime/RenderCore/Private/RenderUtils.cpp b/Engine/Source/Runtime/RenderCore/Private/RenderUtils.cpp
--- a/Engine/Source/Runtime/RenderCore/Private/RenderUtils.cpp
+++ b/Engine/Source/Runtime/RenderCore/Private/RenderUtils.cpp
@@ -1188,7 +1188,8 @@ RENDERCORE_API bool UseVirtualTexturing(const FStaticShaderPlatform InShaderPlat
 		// We treat r.Mobile.VirtualTextures as per project, and _not_ per platform.
 		// To change virtual texture support per platform use r.VirtualTextures.
 		// This makes management simpler and avoids a case where virtual texture support for platforms that support both desktop and mobile renderers could diverge.
-		// The values can't diverge per platform because we only cook one version of the texture resources.
+		// The value generally can't diverge on the same platform because we only cook one version of the texture resources.
+		// But we can still disable this value from client-specific device profiles if we want to remove the overhead case-by-case.
 		static const auto CVarMobileVirtualTextures = IConsoleManager::Get().FindTConsoleVariableDataInt(TEXT("r.Mobile.VirtualTextures"));
 		static  bool bMobileVirtualTextures = CVarMobileVirtualTextures->GetValueOnAnyThread() != 0;
 		if (!bMobileVirtualTextures)
