Add option to enable shader stable keys in DLC cooks
Move LoadMultipleStableKeysFiles to PipelineCacheUtilities for loading and merging of multiple shk files

#tests local Music cook and MPB Join

#FailedValidations Preflight

diff --git a/Engine/Source/Editor/UnrealEd/Private/Commandlets/ShaderPipelineCacheToolsCommandlet.cpp b/Engine/Source/Editor/UnrealEd/Private/Commandlets/ShaderPipelineCacheToolsCommandlet.cpp
--- a/Engine/Source/Editor/UnrealEd/Private/Commandlets/ShaderPipelineCacheToolsCommandlet.cpp
+++ b/Engine/Source/Editor/UnrealEd/Private/Commandlets/ShaderPipelineCacheToolsCommandlet.cpp
@@ -149,42 +149,8 @@ static void LoadStableShaderKeys(TArray<FStableShaderKeyAndValue>& StableArray,
 
 static void LoadStableShaderKeysMultiple(FStableShaderSet& StableMap, TArrayView<const FStringView> FileNames)
 {
-	TArray<TArray<FStableShaderKeyAndValue>> StableArrays;
-	StableArrays.AddDefaulted(FileNames.Num());
-	ParallelFor(FileNames.Num(), [&StableArrays, &FileNames](int32 Index) { LoadStableShaderKeys(StableArrays[Index], FileNames[Index]); });
-
-	if (StableArrays.Num() > 0)
-	{
-		// Since stable keys are saved from a TSet, we assume that a single array does not have non-unique members, so add the largest one first in the most optimal manner
-		StableArrays.Sort([](const TArray<FStableShaderKeyAndValue>& A, const TArray<FStableShaderKeyAndValue>& B) { return (A.Num() > B.Num()); });
-
-		const int32 StableArrayCount = Algo::TransformAccumulate(StableArrays, &TArray<FStableShaderKeyAndValue>::Num, 0);
-		StableMap.Reserve(StableMap.Num() + StableArrayCount);
-
-		bool bFirst = true;
-		for (const TArray<FStableShaderKeyAndValue>& StableArray : StableArrays)
-		{
-			if (bFirst)
-			{
-				bFirst = false;
-				for (const FStableShaderKeyAndValue& Shader : StableArray)
-				{
-					StableMap.Add(MakeUnique<FStableShaderKeyAndValue>(Shader));
-				}
-			}
-			else
-			{
-				for (const FStableShaderKeyAndValue& Shader : StableArray)
-				{
-					// Double lookup to avoid MakeUnique for duplicates
-					if (!StableMap.ContainsByHash(GetTypeHash(Shader), Shader))
-					{
-						StableMap.Add(MakeUnique<FStableShaderKeyAndValue>(Shader));
-					}
-				}
-			}
-		}
-	}
+	// This is now available in PipelineCacheUtilities
+	UE::PipelineCacheUtilities::LoadMultipleStableKeysFiles(FileNames, StableMap);
 }
 
 // Version optimized for ExpandPSOSC
diff --git a/Engine/Source/Editor/UnrealEd/Private/Cooker/ShaderLibraryCooking.cpp b/Engine/Source/Editor/UnrealEd/Private/Cooker/ShaderLibraryCooking.cpp
--- a/Engine/Source/Editor/UnrealEd/Private/Cooker/ShaderLibraryCooking.cpp
+++ b/Engine/Source/Editor/UnrealEd/Private/Cooker/ShaderLibraryCooking.cpp
@@ -254,9 +254,11 @@ void UCookOnTheFlyServer::BeginCookStartShaderCodeLibrary(FBeginCookContext& Beg
 		GConfig->GetBool(TEXT("DevOptions.Shaders"), TEXT("NeedsShaderStableKeys"), bAllPlatformsNeedStableKeys, GEngineIni);
 		GConfig->GetBool(TEXT("DevOptions.Shaders"), TEXT("bNeedsShaderStableKeys"), bAllPlatformsNeedStableKeys, GEngineIni);
 
-		// PSO manual cache in DLC is not currently supported. Although stable keys can have other uses, disable this for DLC as it will
-		// also make it cook faster.
-		bAllPlatformsNeedStableKeys &= !IsCookingDLC();
+		// See if we desire stable keys in DLC, PSO manual cache in DLC is not supported, so they are disabled by default
+		// They have other uses so allow them to be optionally enabled
+		bool bAllowShaderStableKeysInDLC = false;
+		GConfig->GetBool(TEXT("DevOptions.Shaders"), TEXT("bAllowShaderStableKeysInDLC"), bAllowShaderStableKeysInDLC, GEngineIni);
+		bAllPlatformsNeedStableKeys = bAllPlatformsNeedStableKeys && (!IsCookingDLC() || bAllowShaderStableKeysInDLC);
 
 		for (const ITargetPlatform* TargetPlatform : BeginContext.TargetPlatforms)
 		{
diff --git a/Engine/Source/Runtime/RenderCore/Private/PipelineCacheUtilities.cpp b/Engine/Source/Runtime/RenderCore/Private/PipelineCacheUtilities.cpp
--- a/Engine/Source/Runtime/RenderCore/Private/PipelineCacheUtilities.cpp
+++ b/Engine/Source/Runtime/RenderCore/Private/PipelineCacheUtilities.cpp
@@ -2,6 +2,8 @@
 #include "PipelineCacheUtilities.h"
 
 #if UE_WITH_PIPELINE_CACHE_UTILITIES
+#include "Algo/Accumulate.h"
+#include "Async/ParallelFor.h"
 #include "Async/TaskGraphInterfaces.h"
 #include "Misc/Compression.h"
 #include "Misc/SecureHash.h"
@@ -390,6 +392,52 @@ bool UE::PipelineCacheUtilities::LoadStableKeysFile(const FStringView& Filename,
 
 #if WITH_EDITOR
 
+uint32 UE::PipelineCacheUtilities::LoadMultipleStableKeysFiles(TArrayView<const FStringView> FileNames, FStableShaderSet& StableMap)
+{
+	TArray<TArray<FStableShaderKeyAndValue>> StableArrays;
+	StableArrays.AddDefaulted(FileNames.Num());
+	ParallelFor(FileNames.Num(), [&StableArrays, &FileNames](int32 Index) { LoadStableKeysFile(FileNames[Index], StableArrays[Index]); });
+	uint32 MergedCount = 0;
+
+	if (StableArrays.Num() > 0)
+	{
+		// Since stable keys are saved from a TSet, we assume that a single array does not have non-unique members, so add the largest one first in the most optimal manner
+		StableArrays.Sort([](const TArray<FStableShaderKeyAndValue>& A, const TArray<FStableShaderKeyAndValue>& B) { return (A.Num() > B.Num()); });
+
+		const int32 StableArrayCount = Algo::TransformAccumulate(StableArrays, &TArray<FStableShaderKeyAndValue>::Num, 0);
+		StableMap.Reserve(StableMap.Num() + StableArrayCount);
+
+		bool bFirst = true;
+		for (const TArray<FStableShaderKeyAndValue>& StableArray : StableArrays)
+		{
+			// If this file had any data, increase merged file counter
+			MergedCount += StableArray.IsEmpty() ? 0 : 1;
+
+			if (bFirst)
+			{
+				bFirst = false;
+				for (const FStableShaderKeyAndValue& Shader : StableArray)
+				{
+					StableMap.Add(MakeUnique<FStableShaderKeyAndValue>(Shader));
+				}
+			}
+			else
+			{
+				for (const FStableShaderKeyAndValue& Shader : StableArray)
+				{
+					// Double lookup to avoid MakeUnique for duplicates
+					if (!StableMap.ContainsByHash(GetTypeHash(Shader), Shader))
+					{
+						StableMap.Add(MakeUnique<FStableShaderKeyAndValue>(Shader));
+					}
+				}
+			}
+		}
+	}
+
+	return MergedCount;
+}
+
 bool UE::PipelineCacheUtilities::SaveStableKeysFile(const FStringView& Filename, const FStableShaderSet& Values)
 {
 	TUniquePtr<FArchive> FileArchiveInner(IFileManager::Get().CreateFileWriter(*FString::ConstructFromPtrSize(Filename.GetData(), Filename.Len())));
diff --git a/Engine/Source/Runtime/RenderCore/Public/PipelineCacheUtilities.h b/Engine/Source/Runtime/RenderCore/Public/PipelineCacheUtilities.h
--- a/Engine/Source/Runtime/RenderCore/Public/PipelineCacheUtilities.h
+++ b/Engine/Source/Runtime/RenderCore/Public/PipelineCacheUtilities.h
@@ -70,6 +70,15 @@ namespace PipelineCacheUtilities
 
 #if WITH_EDITOR	// limit what's compiled for the cooked games
 
+	/** 
+	* Load Multiple stable shader keys file (using a proprietary format). Stable key is a way to identify a shader independently of its output hash
+	* 
+	* @param FileNames array of filenames (with path if needed)
+	* @param Values destination for merged file values
+	* @return Count of files read in and merged
+	*/
+	RENDERCORE_API uint32 LoadMultipleStableKeysFiles(TArrayView<const FStringView> FileNames, FStableShaderSet& Values);
+
 	/** 
 	 * Saves stable shader keys file (using a proprietary format). Stable key is a way to identify a shader independently of its output hash
 	 * 
