Fixed crash when a refcount of 0 is applied to FRenderCommandList::ReleasePipeRefs.

diff --git a/Engine/Source/Runtime/RenderCore/Public/RenderingThread.h b/Engine/Source/Runtime/RenderCore/Public/RenderingThread.h
--- a/Engine/Source/Runtime/RenderCore/Public/RenderingThread.h
+++ b/Engine/Source/Runtime/RenderCore/Public/RenderingThread.h
@@ -971,20 +971,22 @@ private:
 		return DispatchTaskEvent.GetPtrOrNull();
 	}
 
-	int32 ReleasePipeRefs(int32 InNumRefs)
+	void ReleasePipeRefs(int32 InNumRefs)
 	{
-		int32 NumRefs = NumPipeRefs.fetch_sub(InNumRefs, std::memory_order_acq_rel) - InNumRefs;
-		check(NumRefs >= 0);
-		if (!NumRefs)
+		if (InNumRefs > 0)
 		{
-			delete this;
+			int32 NumRefs = NumPipeRefs.fetch_sub(InNumRefs, std::memory_order_acq_rel) - InNumRefs;
+			check(NumRefs >= 0);
+			if (!NumRefs)
+			{
+				delete this;
+			}
 		}
-		return NumRefs;
 	}
 
-	int32 ReleasePipeRef()
+	void ReleasePipeRef()
 	{
-		return ReleasePipeRefs(1);
+		ReleasePipeRefs(1);
 	}
 
 	bool HasDispatchTask()
