[Insights] Fixed the overlapping between CPU scope for the "BeginFrame" RHI command and the "Frame" RHI breadcrumb scope. This resulted in "BeginFrame" to extend incorrectly over entire "Frame" scope and the "Frame" scope to be incorrectly shortened.
Added ENQUEUE_RENDER_COMMAND_NOMARKER(Type) and DECLARE_RENDER_COMMAND_TAG_NOMARKER. Now "BeginFrame" RHI command is using _NOMARKER, so it will not create a CPU scope anymore.

diff --git a/Engine/Source/Runtime/Core/Public/Stats/LightweightStats.h b/Engine/Source/Runtime/Core/Public/Stats/LightweightStats.h
--- a/Engine/Source/Runtime/Core/Public/Stats/LightweightStats.h
+++ b/Engine/Source/Runtime/Core/Public/Stats/LightweightStats.h
@@ -409,11 +409,15 @@ namespace UE::Stats::Private
 #define UE_INTERNAL_GET_QUICK_STAT_WITH_GROUP_IF_ENABLED(Stat, Group)\
 	UE::Stats::Private::GetStatIfEnabled<UE_INTERNAL_GET_STATDATA_ALLOW_DEFAULT(Stat), UE_INTERNAL_GET_STATGROUP_TYPE(Group)>(ANSI_TO_PROFILING(#Stat), UE_STATS_HASH_NAME(Stat), UE_STATS_HASH_NAME(Group))
 
-
 #elif !STATS // UE_USE_LIGHTWEIGHT_STATS
 
 struct TStatId
 {
+	inline constexpr bool IsValidStat() const
+	{
+		return false;
+	}
+
 	static void AutoRTFMAssignFromOpenToClosed(TStatId& Closed, const TStatId& Open)
 	{
 		Closed = Open;
@@ -430,4 +434,5 @@ public:
 	{
 	}
 };
+
 #endif // UE_USE_LIGHTWEIGHT_STATS
diff --git a/Engine/Source/Runtime/Launch/Private/LaunchEngineLoop.cpp b/Engine/Source/Runtime/Launch/Private/LaunchEngineLoop.cpp
--- a/Engine/Source/Runtime/Launch/Private/LaunchEngineLoop.cpp
+++ b/Engine/Source/Runtime/Launch/Private/LaunchEngineLoop.cpp
@@ -5493,13 +5493,13 @@ static inline void BeginFrameRenderThread(FRHICommandListImmediate& RHICmdList,
 		));
 	}
 	else
-#endif
+#endif // CSV_PROFILER
 	{
 		GRHIFrameBreadcrumb.Emplace(RHICmdList, RHI_BREADCRUMB_DESC_FORWARD_VALUES(TEXT("Frame"), TEXT("Frame %d"), RHI_GPU_STAT_ARGS_NONE)(
 			RHI_BREADCRUMB_FIELD("Frame Number", CurrentFrameCounter)
 		));
 	}
-#endif
+#endif // WITH_RHI_BREADCRUMBS
 
 	GPU_STATS_BEGINFRAME(RHICmdList);
 
@@ -5632,7 +5632,7 @@ void FEngineLoop::Tick()
 	bTraceCpuChannelEnabled = UE_TRACE_CHANNELEXPR_IS_ENABLED(CpuChannel);
 	static FAutoNamedEventsToggler TraceNamedEventsToggler;
 	TraceNamedEventsToggler.Update(bTraceCpuChannelEnabled && UE::Trace::IsTracing());
-#endif
+#endif // CPUPROFILERTRACE_ENABLED
 
 	// Generate the Frame event string
 	TCHAR IndexedFrameString[32] = { 0 };
@@ -5643,7 +5643,7 @@ void FEngineLoop::Tick()
 		FCString::Snprintf(IndexedFrameString, 32, TEXT("Csv%s%d"), FrameStringPrefix, FCsvProfiler::Get()->GetCaptureFrameNumber());
 	}
 	else
-#endif
+#endif // CSV_PROFILER
 	{
 		FCString::Snprintf(IndexedFrameString, 32, TEXT("%s%d"), FrameStringPrefix, CurrentFrameCounter);
 	}
@@ -5656,8 +5656,8 @@ void FEngineLoop::Tick()
 #if PLATFORM_LIMIT_PROFILER_UNIQUE_NAMED_EVENTS
 	FrameString = TEXT("FEngineLoop");
 #endif
-	SCOPED_NAMED_EVENT_TCHAR_CONDITIONAL(FrameString, FrameColor, !bTraceCpuChannelEnabled); 
-#endif //ENABLE_NAMED_EVENTS
+	SCOPED_NAMED_EVENT_TCHAR_CONDITIONAL(FrameString, FrameColor, !bTraceCpuChannelEnabled);
+#endif // ENABLE_NAMED_EVENTS
 
 	// execute callbacks for cvar changes
 	{
@@ -5717,10 +5717,19 @@ void FEngineLoop::Tick()
 		// beginning of RHI frame
 		{
 			SCOPED_NAMED_EVENT(Render_BeginFrame, SystemColor);
+
+#if WITH_RHI_BREADCRUMBS
+			// Uses _NOMARKER to avoid intersecting "BeginFrame" scope with the "Frame" breadcrumb from BeginFrameRenderThread().
+			ENQUEUE_RENDER_COMMAND_NOMARKER(BeginFrame)([CurrentFrameCounter](FRHICommandListImmediate& RHICmdList)
+			{
+				BeginFrameRenderThread(RHICmdList, CurrentFrameCounter);
+			});
+#else // WITH_RHI_BREADCRUMBS
 			ENQUEUE_RENDER_COMMAND(BeginFrame)([CurrentFrameCounter](FRHICommandListImmediate& RHICmdList)
 			{
 				BeginFrameRenderThread(RHICmdList, CurrentFrameCounter);
 			});
+#endif // WITH_RHI_BREADCRUMBS
 
 			for (FSceneInterface* Scene : GetRendererModule().GetAllocatedScenes())
 			{
diff --git a/Engine/Source/Runtime/RenderCore/Private/RenderingThread.cpp b/Engine/Source/Runtime/RenderCore/Private/RenderingThread.cpp
--- a/Engine/Source/Runtime/RenderCore/Private/RenderingThread.cpp
+++ b/Engine/Source/Runtime/RenderCore/Private/RenderingThread.cpp
@@ -1518,11 +1518,8 @@ void FCommandList::Release()
 
 //////////////////////////////////////////////////////////////////////////
 
-static void ExecuteCommand(FRHICommandList* RHICmdList, const FRenderCommandFunctionVariant& Function, const FRenderCommandTag& Tag)
+static void ExecuteCommand_NoMarker(FRHICommandList* RHICmdList, const FRenderCommandFunctionVariant& Function)
 {
-	TRACE_CPUPROFILER_EVENT_SCOPE_USE_ON_CHANNEL(Tag.GetSpecId(), Tag.GetName(), EventScope, RenderCommandsChannel, true);
-	FScopeCycleCounter Scope(Tag.GetStatId(), true);
-
 	switch (Function.GetIndex())
 	{
 	case FRenderCommandFunctionVariant::IndexOfType<TUniqueFunction<void(FRHICommandListImmediate&)>>():
@@ -1541,6 +1538,20 @@ static void ExecuteCommand(FRHICommandList* RHICmdList, const FRenderCommandFunc
 	}
 }
 
+static void ExecuteCommand(FRHICommandList* RHICmdList, const FRenderCommandFunctionVariant& Function, const FRenderCommandTag& Tag)
+{
+	if (Tag.GetStatId().IsValidStat())
+	{
+		TRACE_CPUPROFILER_EVENT_SCOPE_USE_ON_CHANNEL(Tag.GetSpecId(), Tag.GetName(), EventScope, RenderCommandsChannel, !GCycleStatsShouldEmitNamedEvents);
+		FScopeCycleCounter Scope(Tag.GetStatId(), true);
+		ExecuteCommand_NoMarker(RHICmdList, Function);
+	}
+	else
+	{
+		ExecuteCommand_NoMarker(RHICmdList, Function);
+	}
+}
+
 //////////////////////////////////////////////////////////////////////////
 
 thread_local FRenderCommandList* FRenderCommandList::InstanceTLS = nullptr;
diff --git a/Engine/Source/Runtime/RenderCore/Public/RenderingThread.h b/Engine/Source/Runtime/RenderCore/Public/RenderingThread.h
--- a/Engine/Source/Runtime/RenderCore/Public/RenderingThread.h
+++ b/Engine/Source/Runtime/RenderCore/Public/RenderingThread.h
@@ -236,6 +236,25 @@ private:
 	{}
 };
 
+/** Type that contains profiler data necessary to mark up render commands for various profilers. */
+template <typename TSTR>
+class TRenderCommandTagNoMarker : public FRenderCommandTag
+{
+public:
+	static const TRenderCommandTagNoMarker& Get()
+	{
+		static TRenderCommandTagNoMarker Tag(TSTR::TStr(), TStatId());
+		return Tag;
+	}
+
+private:
+	TRenderCommandTagNoMarker(const TCHAR* InName, TStatId InStatId)
+		: FRenderCommandTag(InName, TStatId())
+	{
+	}
+};
+
+
 /** Declares a new render command tag type from a name. */
 #define DECLARE_RENDER_COMMAND_TAG(Type, Name) \
 	struct UE_JOIN(TSTR_, Name, __LINE__) \
@@ -245,6 +264,14 @@ private:
 	}; \
 	using Type = TRenderCommandTag<UE_JOIN(TSTR_, Name, __LINE__)>;
 
+#define DECLARE_RENDER_COMMAND_TAG_NOMARKER(Type, Name) \
+	struct UE_JOIN(TSTR_, Name, __LINE__) \
+	{  \
+		static const char* CStr() { return #Name; } \
+		static const TCHAR* TStr() { return TEXT(#Name); } \
+	}; \
+	using Type = TRenderCommandTagNoMarker<UE_JOIN(TSTR_, Name, __LINE__)>;
+
 /** Describes which pipes are configured to use the render command pipe system. */
 enum class ERenderCommandPipeMode : uint8
 {
@@ -594,7 +621,6 @@ public:
 	static void Enqueue(LambdaType&& Lambda)
 	{
 		const FRenderCommandTag& Tag = RenderCommandTag::Get();
-		TRACE_CPUPROFILER_EVENT_SCOPE_USE_ON_CHANNEL(Tag.GetSpecId(), Tag.GetName(), EventScope, RenderCommandsChannel, true);
 
 		if (!IsInRenderingThread() && ShouldExecuteOnRenderThread())
 		{
@@ -609,11 +635,16 @@ public:
 				TGraphTask<TRenderCommandTask<LambdaType>>::CreateTask().ConstructAndDispatchWhenReady(MoveTemp(Lambda), Tag);
 			}
 		}
-		else
+		else if (Tag.GetStatId().IsValidStat())
 		{
+			TRACE_CPUPROFILER_EVENT_SCOPE_USE_ON_CHANNEL(Tag.GetSpecId(), Tag.GetName(), EventScope, RenderCommandsChannel, !GCycleStatsShouldEmitNamedEvents);
 			FScopeCycleCounter CycleScope(Tag.GetStatId());
 			Lambda(GetImmediateCommandList_ForRenderCommand());
 		}
+		else
+		{
+			Lambda(GetImmediateCommandList_ForRenderCommand());
+		}
 	}
 
 	static void Enqueue(FRenderCommandList* CommandList)
@@ -1170,6 +1201,10 @@ public:
 	DECLARE_RENDER_COMMAND_TAG(UE_JOIN(FRenderCommandTag_, Type, __LINE__), Type) \
 	FRenderCommandDispatcher::Enqueue<UE_JOIN(FRenderCommandTag_, Type, __LINE__)>
 
+#define ENQUEUE_RENDER_COMMAND_NOMARKER(Type) \
+	DECLARE_RENDER_COMMAND_TAG_NOMARKER(UE_JOIN(FRenderCommandTag_, Type, __LINE__), Type) \
+	FRenderCommandDispatcher::Enqueue<UE_JOIN(FRenderCommandTag_, Type, __LINE__)>
+
 ////////////////////////////////////
 // RenderThread scoped work
 ////////////////////////////////////
