[MRG+DumpGPU] Add Resource Filtering to DumpGPU

Improve DumpGPU filtering to only dump specified resources .
Fixes Dumping of Buffers which were previously incorrectly tied to TextureDumping cvar.

#rnx

diff --git a/Engine/Plugins/MovieScene/MovieRenderPipeline/Source/MovieRenderPipelineCore/Private/Graph/MovieGraphDefaultRenderer.cpp b/Engine/Plugins/MovieScene/MovieRenderPipeline/Source/MovieRenderPipelineCore/Private/Graph/MovieGraphDefaultRenderer.cpp
--- a/Engine/Plugins/MovieScene/MovieRenderPipeline/Source/MovieRenderPipelineCore/Private/Graph/MovieGraphDefaultRenderer.cpp
+++ b/Engine/Plugins/MovieScene/MovieRenderPipeline/Source/MovieRenderPipelineCore/Private/Graph/MovieGraphDefaultRenderer.cpp
@@ -356,7 +356,7 @@ void UMovieGraphDefaultRenderer::Render(const FMovieGraphTimeStepData& InTimeSte
 		if (bIsRendering && bWantsDumpGpuDebugCapture && bIsDumpGPUCorrectFrame)
 		{
 			const FString FinalPath = UE::MovieGraph::Capture::PrepareCapturePath(GetOwningGraph(), InTimeStepData.EvaluatedConfig, DebugSettings->DumpGpuCaptureFilePathFormat, TEXT(""), true);
-			const FString Cmd = FString::Printf(TEXT("-noexplorer -path=\"%s\" -pass_path_filters=\"%s\""), *FinalPath, *DebugSettings->DumpGpuCapturePassPathFilters);
+			const FString Cmd = FString::Printf(TEXT("-noexplorer -path=\"%s\" -pass_path_filters=\"%s\" -resource_name_filters=\"%s\""), *FinalPath, *DebugSettings->DumpGpuCapturePassPathFilters, *DebugSettings->DumpGpuCaptureResourcesFilters);
 			const FString Context = FString::Printf(TEXT("MRQ Frame: %d"), InTimeStepData.RootFrameNumber.Value);
 			ScopedDumpGPUCapture = MakeUnique<UE::RenderCore::DumpGPU::FDumpScope>(true, *Cmd, *Context);
 		}
diff --git a/Engine/Plugins/MovieScene/MovieRenderPipeline/Source/MovieRenderPipelineCore/Public/Graph/Nodes/MovieGraphDebugNode.h b/Engine/Plugins/MovieScene/MovieRenderPipeline/Source/MovieRenderPipelineCore/Public/Graph/Nodes/MovieGraphDebugNode.h
--- a/Engine/Plugins/MovieScene/MovieRenderPipeline/Source/MovieRenderPipelineCore/Public/Graph/Nodes/MovieGraphDebugNode.h
+++ b/Engine/Plugins/MovieScene/MovieRenderPipeline/Source/MovieRenderPipelineCore/Public/Graph/Nodes/MovieGraphDebugNode.h
@@ -53,6 +53,9 @@ public:
 	UPROPERTY(EditAnywhere, BlueprintReadWrite, Category = Overrides, meta = (InlineEditConditionToggle))
 	uint8 bOverride_DumpGpuCapturePassPathFilters : 1;
 
+	UPROPERTY(EditAnywhere, BlueprintReadWrite, Category = Overrides, meta = (InlineEditConditionToggle))
+	uint8 bOverride_DumpGpuCaptureResourcesFilters : 1;
+
 	/** 
 	* If true, automatically trigger RenderDoc to capture rendering information. RenderDoc plugin must be enabled, 
 	* and the editor must have been launched with -AttachRenderDoc. Resulting capture will be in /Saved/RenderDocCaptures.
@@ -109,6 +112,13 @@ public:
 	*/
 	UPROPERTY(EditAnywhere, BlueprintReadWrite, Category = "DumpGPU", meta = (EditCondition = "bOverride_DumpGpuCapturePassPathFilters"))
 	FString DumpGpuCapturePassPathFilters = TEXT("*");
+
+	/**
+	* If bCaptureFramesWithDumpGpu is true, wildcard enabled filters to apply to the resources to dump.
+	* Support multiple filters separated by "|". For example : MegaLights.ResolvedDiffuseLighting|MegaLights.ResolvedSpecularLighting
+	*/
+	UPROPERTY(EditAnywhere, BlueprintReadWrite, Category = "DumpGPU", meta = (EditCondition = "bOverride_DumpGpuCaptureResourcesFilters"))
+	FString DumpGpuCaptureResourcesFilters = TEXT("*");
 };
 
 #undef UE_API
diff --git a/Engine/Source/Runtime/RenderCore/Private/DumpGPU.cpp b/Engine/Source/Runtime/RenderCore/Private/DumpGPU.cpp
--- a/Engine/Source/Runtime/RenderCore/Private/DumpGPU.cpp
+++ b/Engine/Source/Runtime/RenderCore/Private/DumpGPU.cpp
@@ -246,6 +246,7 @@ public:
 	int32 FrameCount = 1;
 	FName UploadResourceCompressionName;
 	TArray<FString> PassPathFilters;
+	TArray<FString> ResourceNameFilters;
 	FString DumpingDirectoryPath;
 	FString Context;
 	FDateTime Time;
@@ -1495,6 +1496,11 @@ public:
 		const FRDGTextureSRVDesc& SubresourceRangeDesc,
 		bool bIsOutputResource)
 	{
+		if (!IsDumpingResource(SubresourceRangeDesc.Texture->Name))
+		{
+			return;
+		}
+
 		bool bAllowDumpBinary = true;
 		if (!HasResourceEverBeenDumped(SubresourceRangeDesc.Texture))
 		{
@@ -1576,7 +1582,12 @@ public:
 		FRDGBuffer* Buffer,
 		bool bIsOutputResource)
 	{
-		int32 DumpBufferMode = GDumpTextureCVar.GetValueOnRenderThread();
+		if (!IsDumpingResource(Buffer->Name))
+		{
+			return;
+		}
+
+		int32 DumpBufferMode = GDumpBufferCVar.GetValueOnRenderThread();
 
 		if (DumpBufferMode == 0)
 		{
@@ -1960,6 +1971,26 @@ public:
 		return bDumpPass;
 	}
 
+	bool IsDumpingResource(const TCHAR* ResourceName) const
+	{
+		if (ResourceNameFilters.IsEmpty())
+		{
+			return true;
+		}
+
+		for (const FString& ResourceFilter : ResourceNameFilters)
+		{
+			FWildcardString WildcardResourceFilter(ResourceFilter);
+			if (WildcardResourceFilter.IsMatch(ResourceName))
+			{
+				return true;
+			}
+		}
+
+
+		return false;
+	}
+
 	void OnCrash()
 	{
 		if (GLog)
@@ -2069,6 +2100,11 @@ FString FRDGBuilder::BeginResourceDump(const TCHAR* Cmd, const TCHAR* Context)
 		Parameters[TEXT("pass_path_filters")].ParseIntoArray(NewResourceDumpContext->PassPathFilters, TEXT("|"));
 	}
 
+	if (Parameters.Contains(TEXT("resource_name_filters")))
+	{
+		Parameters[TEXT("resource_name_filters")].ParseIntoArray(NewResourceDumpContext->ResourceNameFilters, TEXT("|"));
+	}
+
 	NewResourceDumpContext->bStream = GDumpGPUStream.GetValueOnGameThread() != 0;
 	NewResourceDumpContext->bEnableDiskWrite = GDumpTestEnableDiskWrite.GetValueOnGameThread() != 0;
 	NewResourceDumpContext->DeltaTime = FApp::GetDeltaTime();
