Added AddBarrierPass helper to give some control over flushing of barriers on platforms where split barriers are not supported.

diff --git a/Engine/Source/Runtime/RHI/Public/RHIGlobals.h b/Engine/Source/Runtime/RHI/Public/RHIGlobals.h
--- a/Engine/Source/Runtime/RHI/Public/RHIGlobals.h
+++ b/Engine/Source/Runtime/RHI/Public/RHIGlobals.h
@@ -749,6 +749,9 @@ PRAGMA_ENABLE_DEPRECATION_WARNINGS
 	/** True if RHI needs to manually apply an RHI transition when a resource is flag as SRVGraphicsNonPixel. */
 	bool NeedsSRVGraphicsNonPixelWorkaround = false;
 
+	/** True if the platform supports split barriers. */
+	bool SupportsSplitBarriers = false;
+
 	/** Preferred format of FShaderCode (binary code or microcode). Only used for preview platforms with more than one shader code being supported such as D3D12 (supports DXIL and DXBC). */
 	FName PreferredPreviewShaderCodeFormat = NAME_None;
 };
diff --git a/Engine/Source/Runtime/RenderCore/Private/RenderGraphUtils.cpp b/Engine/Source/Runtime/RenderCore/Private/RenderGraphUtils.cpp
--- a/Engine/Source/Runtime/RenderCore/Private/RenderGraphUtils.cpp
+++ b/Engine/Source/Runtime/RenderCore/Private/RenderGraphUtils.cpp
@@ -1022,3 +1022,20 @@ TRefCountPtr<FRDGPooledBuffer> AllocatePooledBufferCurrentLLMTag(FRHICommandList
 	return RDGPooledBuffer;
 }
 #endif
+
+BEGIN_SHADER_PARAMETER_STRUCT(FBarrierPassParameters, )
+	RDG_TEXTURE_ACCESS_ARRAY(Textures)
+	RDG_BUFFER_ACCESS_ARRAY(Buffers)
+END_SHADER_PARAMETER_STRUCT()
+
+void AddBarrierPass(FRDGBuilder& GraphBuilder, ERDGPassFlags PassFlags, FRDGTextureAccessArray Textures, FRDGBufferAccessArray Buffers)
+{
+	// Skip adding the barrier pass when split barriers are supported, as they intrinsically support overlap.
+	if (!GRHIGlobals.SupportsSplitBarriers)
+	{
+		FBarrierPassParameters* Parameters = GraphBuilder.AllocParameters<FBarrierPassParameters>();
+		Parameters->Textures = MoveTemp(Textures);
+		Parameters->Buffers  = MoveTemp(Buffers);
+		GraphBuilder.AddPass(RDG_EVENT_NAME("BarrierPass"), Parameters, ERDGPassFlags::NeverCull | PassFlags, [&Parameters](FRDGAsyncTask, FRHICommandList& RHICmdList) {});
+	}
+}
\ No newline at end of file
diff --git a/Engine/Source/Runtime/RenderCore/Public/RenderGraphUtils.h b/Engine/Source/Runtime/RenderCore/Public/RenderGraphUtils.h
--- a/Engine/Source/Runtime/RenderCore/Public/RenderGraphUtils.h
+++ b/Engine/Source/Runtime/RenderCore/Public/RenderGraphUtils.h
@@ -1428,3 +1428,32 @@ RENDERCORE_API bool AllocatePooledTexture(
 	const TCHAR* Name);
 
 RENDERCORE_API TRefCountPtr<IPooledRenderTarget> AllocatePooledTexture(const FRDGTextureDesc& Desc, const TCHAR* Name);
+
+/*
+	Inserts a pass that transitions the resources into the requested state early. This can be used to move barriers to occur earlier in the graph.
+	This works if the state matches the state in a later pass. For example:
+
+	Prepare(A)
+	Prepare(B)
+	AddBarrierPass(A, B) <- Barriers are flushed here.
+	Dispatch(A)
+	Dispatch(B)
+
+	On platforms without split barriers, the barriers would otherwise be interleaved:
+
+	Prepare(A)
+	Prepare(B)
+	Dispatch(A) <- Wait for A
+	Dispatch(B) <- Wait for B
+*/
+RENDERCORE_API void AddBarrierPass(FRDGBuilder& GraphBuilder, ERDGPassFlags PassFlags, FRDGTextureAccessArray Textures, FRDGBufferAccessArray Buffers);
+
+inline void AddBarrierPass(FRDGBuilder& GraphBuilder, ERDGPassFlags PassFlags, FRDGBufferAccessArray Buffers)
+{
+	AddBarrierPass(GraphBuilder, PassFlags, {}, MoveTemp(Buffers));
+}
+
+inline void AddBarrierPass(FRDGBuilder& GraphBuilder, ERDGPassFlags PassFlags, FRDGTextureAccessArray Textures)
+{
+	AddBarrierPass(GraphBuilder, PassFlags, MoveTemp(Textures), {});
+}
\ No newline at end of file
