Cleaning up UniformBuffer code generation.
- Renaming  UB_CB_DECL_* to UB_CB_REMAP_* and StructMembers to Remappings to clarify that these are for the UniformBuffer remapping step.
- Deduplicating a few code blocks that generate parameter names the same way. Doing the same for the constant buffer padding calculations.

diff --git a/Engine/Shaders/Public/Platform.ush b/Engine/Shaders/Public/Platform.ush
--- a/Engine/Shaders/Public/Platform.ush
+++ b/Engine/Shaders/Public/Platform.ush
@@ -520,23 +520,25 @@
 #define UB_CB_MEMBER_ACCESS(         UBName,         MemberName) UBName##_##MemberName
 #define UB_CB_PREFIXED_MEMBER_ACCESS(UBName, Prefix, MemberName) Prefix##UBName##_##MemberName
 
-#define UB_CB_DECL_PARAMETER(UBName, StructName, GlobalName) UBName##.##StructName = UBName##_##GlobalName
-#define UB_CB_DECL_RESOURCE(UBName, StructName, GlobalName)  UBName##.##StructName = UBName##_##GlobalName
+#define UB_CB_REMAP_PARAMETER(UBName, StructName, GlobalName) UBName##.##StructName = UBName##_##GlobalName
+#define UB_CB_REMAP_RESOURCE(UBName, StructName, GlobalName)  UBName##.##StructName = UBName##_##GlobalName
 
 #define UB_CB_UNIFORM_BLOCK(UBName, MemberName) UB_FLOAT(4) UB_CB_MEMBER_NAME(UBName,MemberName)[PLATFORM_MAX_UNIFORM_BUFFER_RANGE_FLOAT4]
 
 #if PLATFORM_SUPPORTS_CONSTANTBUFFER_OBJECT
-	#define UB_CBO_DEFINITION_START(      UBName                    ) struct F##UBName##Constants {
+	#define UB_CBO_TYPE(UBName) F##UBName##Constants
+
+	#define UB_CBO_DEFINITION_START(      UBName                    ) struct UB_CBO_TYPE(UBName) {
 	#define UB_CBO_DEFINITION_START_REG(  UBName, Register, Space   ) UB_CBO_DEFINITION_START(UBName)
 	#define UB_CBO_MEMBER_NAME(           UBName,         MemberName) MemberName
 	#define UB_CBO_PREFIXED_MEMBER_NAME(  UBName, Prefix, MemberName) Prefix##MemberName
-	#define UB_CBO_DEFINITION_END(        UBName                    ) }; ConstantBuffer<F##UBName##Constants> UBName;
-	#define UB_CBO_DEFINITION_END_REG(    UBName, Register, Space   ) }; ConstantBuffer<F##UBName##Constants> UBName : register(b##Register, space##Space);
+	#define UB_CBO_DEFINITION_END(        UBName                    ) }; ConstantBuffer<UB_CBO_TYPE(UBName)> UBName;
+	#define UB_CBO_DEFINITION_END_REG(    UBName, Register, Space   ) }; ConstantBuffer<UB_CBO_TYPE(UBName)> UBName : register(b##Register, space##Space);
 	#define UB_CBO_MEMBER_ACCESS(         UBName,         MemberName) UBName##.##MemberName
 	#define UB_CBO_PREFIXED_MEMBER_ACCESS(UBName, Prefix, MemberName) UBName##.##Prefix##MemberName
 
-	#define UB_CBO_DECL_PARAMETER(UBName, StructName, GlobalName) UBName##.##StructName = UBName##.##GlobalName
-	#define UB_CBO_DECL_RESOURCE(UBName, StructName, GlobalName)  UBName##.##StructName = UBName##_##GlobalName
+	#define UB_CBO_REMAP_PARAMETER(UBName, StructName, GlobalName) UBName##.##StructName = UBName##.##GlobalName
+	#define UB_CBO_REMAP_RESOURCE(UBName, StructName, GlobalName)  UBName##.##StructName = UBName##_##GlobalName
 
 	#define UB_CBO_UNIFORM_BLOCK(UBName, MemberName) UB_FLOAT(4) UB_CBO_MEMBER_NAME(UBName,MemberName)[PLATFORM_MAX_UNIFORM_BUFFER_RANGE_FLOAT4]
 
@@ -550,8 +552,8 @@
 	#define UB_CBO_MEMBER_ACCESS(         UBName,         MemberName) UB_CB_MEMBER_ACCESS(UBName,MemberName)
 	#define UB_CBO_PREFIXED_MEMBER_ACCESS(UBName, Prefix, MemberName) UB_CB_PREFIXED_MEMBER_ACCESS(UBName,Prefix,MemberName)
 
-	#define UB_CBO_DECL_PARAMETER(UBName, StructName, GlobalName) UB_CB_DECL_PARAMETER(UBName,StructName,GlobalName)
-	#define UB_CBO_DECL_RESOURCE( UBName, StructName, GlobalName) UB_CB_DECL_RESOURCE(UBName,StructName,GlobalName)
+	#define UB_CBO_REMAP_PARAMETER(UBName, StructName, GlobalName) UB_CB_REMAP_PARAMETER(UBName,StructName,GlobalName)
+	#define UB_CBO_REMAP_RESOURCE( UBName, StructName, GlobalName) UB_CB_REMAP_RESOURCE(UBName,StructName,GlobalName)
 
 	#define UB_CBO_UNIFORM_BLOCK(UBName, MemberName) UB_CB_UNIFORM_BLOCK(UBName, MemberName)
 #endif
diff --git a/Engine/Source/Runtime/RenderCore/Private/ShaderParameters.cpp b/Engine/Source/Runtime/RenderCore/Private/ShaderParameters.cpp
--- a/Engine/Source/Runtime/RenderCore/Private/ShaderParameters.cpp
+++ b/Engine/Source/Runtime/RenderCore/Private/ShaderParameters.cpp
@@ -132,12 +132,71 @@ struct FUniformBufferDecl
 {
 	/** Members to place in the constant buffer. */
 	FStringBuilderBase ConstantBufferMembers;
+
 	/** Members to place in the resource table. */
 	FStringBuilderBase ResourceMembers;
-	/** Members in the struct HLSL shader code will access. */
-	FStringBuilderBase StructMembers;
+
+	/** Remappings to apply in the UniformBuffer { } block */
+	FStringBuilderBase Remappings;
+};
+
+inline EShaderParameterType GetShaderParameterTypeFromMemberResourceType(EUniformBufferBaseType InMemberType)
+{
+	if (InMemberType == UBMT_SAMPLER)
+	{
+		return EShaderParameterType::Sampler;
+	}
+
+	if (InMemberType == UBMT_UAV || InMemberType == UBMT_RDG_TEXTURE_UAV || InMemberType == UBMT_RDG_BUFFER_UAV)
+	{
+		return EShaderParameterType::UAV;
+	}
+
+	return EShaderParameterType::SRV;
+}
+
+inline const TCHAR* GetShaderParameterTypeSuffix(EShaderParameterType ParameterType)
+{
+	if (ParameterType == EShaderParameterType::SRV)
+	{
+		return TEXT("SRV");
+	}
+
+	if (ParameterType == EShaderParameterType::UAV)
+	{
+		return TEXT("UAV");
+	}
+
+	if (ParameterType == EShaderParameterType::Sampler)
+	{
+		return TEXT("SAMPLER");
+	}
+
+	return TEXT("PARAMETER");
+}
+
+inline bool CanParameterTypeBeBindless(EShaderParameterType ParameterType)
+{
+	return ParameterType != EShaderParameterType::LooseData;
+}
+
+inline const TCHAR* GetBindlessPrefix(EUniformBufferBaseType InBaseType)
+{
+	if (InBaseType == UBMT_SAMPLER)
+	{
+		return FShaderParameterParser::kBindlessSamplerPrefix;
+	}
+
+	if (InBaseType == UBMT_UAV || InBaseType == UBMT_RDG_TEXTURE_UAV || InBaseType == UBMT_RDG_BUFFER_UAV)
+	{
+		return FShaderParameterParser::kBindlessUAVPrefix;
+	}
+
+	return FShaderParameterParser::kBindlessSRVPrefix;
 };
 
+using FTempParameterStringBuilder = TStringBuilder<256>;
+
 /** Generates a HLSL struct declaration for a uniform buffer struct. */
 static void CreateHLSLUniformBufferStructMembersDeclaration(
 	const FShaderParametersMetadata& UniformBufferStruct,
@@ -150,21 +209,53 @@ static void CreateHLSLUniformBufferStructMembersDeclaration(
 {
 	const TArray<FShaderParametersMetadata::FMember>& StructMembers = UniformBufferStruct.GetMembers();
 
-	auto AddStructMember = [&](const FShaderParametersMetadata::FMember& Member, bool bResource)
+	auto GenerateGlobalMemberName = [&GlobalPrefix](FTempParameterStringBuilder& Builder, const FShaderParametersMetadata::FMember& Member)
 	{
-		// UB_CB_DECL_PARAMETER(UniformBufferName, StructPrefix_Name, GlobalPrefix_Name)
+		if (!GlobalPrefix.IsEmpty())
+		{
+			Builder << GlobalPrefix << TEXT("_");
+		}
+		Builder << Member.GetName();
+	};
 
-		Decl.StructMembers << (bResource ? TEXT("UB_CB_DECL_RESOURCE(") : TEXT("UB_CB_DECL_PARAMETER(")) << UniformBufferName << TEXT(",");
+	auto GenerateStructMemberName = [&StructPrefix](FStringBuilderBase& Builder, const FShaderParametersMetadata::FMember& Member)
+	{
 		if (!StructPrefix.IsEmpty())
 		{
-			Decl.StructMembers << StructPrefix << TEXT(".");
+			Builder << StructPrefix << TEXT(".");
 		}
-		Decl.StructMembers << Member.GetName() << TEXT(",");
-		if (!GlobalPrefix.IsEmpty())
+		Builder << Member.GetName();
+	};
+
+	auto AdjustConstantBufferPadding = [&HLSLBaseOffset, &Decl, StructOffset, &UniformBufferName](const FShaderParametersMetadata::FMember& Member, const TCHAR* PreviousBaseTypeName, uint32 HLSLMemberSize)
+	{
+		const uint32 AbsoluteMemberOffset = StructOffset + Member.GetOffset();
+
+		// If the HLSL offset doesn't match the C++ offset, generate padding to fix it.
+		if (HLSLBaseOffset != AbsoluteMemberOffset)
 		{
-			Decl.StructMembers << GlobalPrefix << TEXT("_");
+			check(HLSLBaseOffset < AbsoluteMemberOffset);
+			while (HLSLBaseOffset < AbsoluteMemberOffset)
+			{
+				Decl.ConstantBufferMembers.Appendf(TEXT("\t%s() UB_CB_MEMBER_NAME(%s, Padding%u);\n"), PreviousBaseTypeName, *UniformBufferName, HLSLBaseOffset);
+				HLSLBaseOffset += 4;
+			};
+			check(HLSLBaseOffset == AbsoluteMemberOffset);
 		}
-		Decl.StructMembers << Member.GetName() << TEXT(");\n");
+		HLSLBaseOffset = AbsoluteMemberOffset + HLSLMemberSize;
+	};
+
+	auto AddRemapping = [&](const FShaderParametersMetadata::FMember& Member, const FTempParameterStringBuilder& GlobalName, EShaderParameterType ParameterType)
+	{
+		const bool bResource = (ParameterType != EShaderParameterType::LooseData);
+		// Members for the UniformBuffer { ... } code that is used to remap struct accesses to global resources
+
+		// UB_CB_REMAP_PARAMETER(UniformBufferName, StructPrefix_Name, GlobalPrefix_Name)
+		Decl.Remappings << (bResource ? TEXT("UB_CB_REMAP_RESOURCE(") : TEXT("UB_CB_REMAP_PARAMETER(")) << UniformBufferName << TEXT(",");
+
+		GenerateStructMemberName(Decl.Remappings, Member);
+
+		Decl.Remappings << TEXT(",") << GlobalName << TEXT(");\n");
 	};
 
 	if (EnumHasAnyFlags(UniformBufferStruct.GetUsageFlags(), FShaderParametersMetadata::EUsageFlags::UniformView))
@@ -174,16 +265,13 @@ static void CreateHLSLUniformBufferStructMembersDeclaration(
 		const FShaderParametersMetadata::FMember& Member = StructMembers[0];
 		check(Member.GetBaseType() == UBMT_SRV || Member.GetBaseType() == UBMT_RDG_BUFFER_SRV);
 
-		TStringBuilder<256> ParameterName;
-		if (!GlobalPrefix.IsEmpty())
-		{
-			ParameterName << GlobalPrefix << TEXT("_");
-		}
-		ParameterName << Member.GetName();
+		FTempParameterStringBuilder GlobalParameterName;
+		GenerateGlobalMemberName(GlobalParameterName, Member);
+
+		Decl.ConstantBufferMembers.Appendf(TEXT("UB_CB_UNIFORM_BLOCK(%s, %s);\n"), *UniformBufferName, *GlobalParameterName);
 
-		Decl.ConstantBufferMembers.Appendf(TEXT("UB_CB_UNIFORM_BLOCK(%s, %s);\n"), *UniformBufferName, *ParameterName);
+		AddRemapping(Member, GlobalParameterName, EShaderParameterType::LooseData);
 
-		AddStructMember(Member, false);
 		return;
 	}
 
@@ -214,51 +302,19 @@ static void CreateHLSLUniformBufferStructMembersDeclaration(
 		else if (IsShaderParameterTypeForUniformBufferLayout(Member.GetBaseType()))
 		{
 			// Add the constant buffer entry for bindless resource indices
-
 			constexpr uint32 HLSLMemberSize = 4;
 
-			const uint32 AbsoluteMemberOffset = StructOffset + Member.GetOffset();
+			AdjustConstantBufferPadding(Member, PreviousBaseTypeName, HLSLMemberSize);
 
-			// If the HLSL offset doesn't match the C++ offset, generate padding to fix it.
-			if (HLSLBaseOffset != AbsoluteMemberOffset)
-			{
-				check(HLSLBaseOffset < AbsoluteMemberOffset);
-				while (HLSLBaseOffset < AbsoluteMemberOffset)
-				{
-					Decl.ConstantBufferMembers.Appendf(TEXT("\t%s() UB_CB_MEMBER_NAME(%s, Padding%u);\n"), PreviousBaseTypeName, *UniformBufferName, HLSLBaseOffset);
-					HLSLBaseOffset += 4;
-				};
-				check(HLSLBaseOffset == AbsoluteMemberOffset);
-			}
 			PreviousBaseTypeName = TEXT("UB_UINT");
-			HLSLBaseOffset = AbsoluteMemberOffset + HLSLMemberSize;
-
-			auto GetBindlessPrefix = [](EUniformBufferBaseType InBaseType)
-				{
-					if (InBaseType == UBMT_SAMPLER)
-					{
-						return FShaderParameterParser::kBindlessSamplerPrefix;
-					}
-
-					if (InBaseType == UBMT_UAV || InBaseType == UBMT_RDG_TEXTURE_UAV || InBaseType == UBMT_RDG_BUFFER_UAV)
-					{
-						return FShaderParameterParser::kBindlessUAVPrefix;
-					}
-
-					return FShaderParameterParser::kBindlessSRVPrefix;
-				};
 
 			// Generate the member declaration.
 			const TCHAR* MemberPrefix = GetBindlessPrefix(Member.GetBaseType());
 
-			TStringBuilder<256> ParameterName;
-			if (!GlobalPrefix.IsEmpty())
-			{
-				ParameterName << GlobalPrefix << TEXT("_");
-			}
-			ParameterName << Member.GetName();
+			FTempParameterStringBuilder GlobalParameterName;
+			GenerateGlobalMemberName(GlobalParameterName, Member);
 
-			Decl.ConstantBufferMembers.Appendf(TEXT("\tUB_UINT() UB_CB_PREFIXED_MEMBER_NAME(%s, %s, %s);\n"), *UniformBufferName, MemberPrefix, *ParameterName);
+			Decl.ConstantBufferMembers.Appendf(TEXT("\tUB_UINT() UB_CB_PREFIXED_MEMBER_NAME(%s, %s, %s);\n"), *UniformBufferName, MemberPrefix, *GlobalParameterName);
 		}
 		else
 		{
@@ -307,32 +363,16 @@ static void CreateHLSLUniformBufferStructMembersDeclaration(
 				HLSLMemberSize = (Member.GetNumElements() - 1) * Align(HLSLMemberSize, 16) + HLSLMemberSize;
 			}
 
-			const uint32 AbsoluteMemberOffset = StructOffset + Member.GetOffset();
+			AdjustConstantBufferPadding(Member, PreviousBaseTypeName, HLSLMemberSize);
 
-			// If the HLSL offset doesn't match the C++ offset, generate padding to fix it.
-			if (HLSLBaseOffset != AbsoluteMemberOffset)
-			{
-				check(HLSLBaseOffset < AbsoluteMemberOffset);
-				while (HLSLBaseOffset < AbsoluteMemberOffset)
-				{
-					Decl.ConstantBufferMembers.Appendf(TEXT("\t%s() UB_CB_MEMBER_NAME(%s, Padding%u);\n"), PreviousBaseTypeName, *UniformBufferName, HLSLBaseOffset);
-					HLSLBaseOffset += 4;
-				};
-				check(HLSLBaseOffset == AbsoluteMemberOffset);
-			}
 			PreviousBaseTypeName = BaseTypeName;
-			HLSLBaseOffset = AbsoluteMemberOffset + HLSLMemberSize;
 
-			TStringBuilder<256> ParameterName;
-			if (!GlobalPrefix.IsEmpty())
-			{
-				ParameterName << GlobalPrefix << TEXT("_");
-			}
-			ParameterName << Member.GetName();
+			FTempParameterStringBuilder GlobalParameterName;
+			GenerateGlobalMemberName(GlobalParameterName, Member);
 
-			Decl.ConstantBufferMembers.Appendf(TEXT("\t%s(%s) UB_CB_MEMBER_NAME(%s, %s%s);\n"), BaseTypeName, *TypeDim, *UniformBufferName, *ParameterName, *ArrayDim);
+			Decl.ConstantBufferMembers.Appendf(TEXT("\t%s(%s) UB_CB_MEMBER_NAME(%s, %s%s);\n"), BaseTypeName, *TypeDim, *UniformBufferName, *GlobalParameterName, *ArrayDim);
 
-			AddStructMember(Member, false);
+			AddRemapping(Member, GlobalParameterName, EShaderParameterType::LooseData);
 		}
 	}
 
@@ -345,28 +385,15 @@ static void CreateHLSLUniformBufferStructMembersDeclaration(
 			// TODO: handle arrays?
 			checkf(!IsRDGResourceAccessType(Member.GetBaseType()), TEXT("RDG access parameter types (e.g. RDG_TEXTURE_ACCESS) are not allowed in uniform buffers."));
 
-			TStringBuilder<256> ParameterName;
-			if (!GlobalPrefix.IsEmpty())
-			{
-				ParameterName << GlobalPrefix << TEXT("_");
-			}
-			ParameterName << Member.GetName();
+			FTempParameterStringBuilder GlobalParameterName;
+			GenerateGlobalMemberName(GlobalParameterName, Member);
 
-			if (Member.GetBaseType() == UBMT_SAMPLER)
-			{
-				Decl.ResourceMembers.Appendf(TEXT("UB_RESOURCE_MEMBER_SAMPLER(%s, %s, %s);\n"), Member.GetShaderType(), *UniformBufferName, *ParameterName);
-				AddStructMember(Member, true);
-			}
-			else if (Member.GetBaseType() == UBMT_UAV || Member.GetBaseType() == UBMT_RDG_TEXTURE_UAV || Member.GetBaseType() == UBMT_RDG_BUFFER_UAV)
-			{
-				Decl.ResourceMembers.Appendf(TEXT("UB_RESOURCE_MEMBER_UAV(%s, %s, %s);\n"), Member.GetShaderType(), *UniformBufferName, *ParameterName);
-				AddStructMember(Member, true);
-			}
-			else
-			{
-				Decl.ResourceMembers.Appendf(TEXT("UB_RESOURCE_MEMBER_SRV(%s, %s, %s);\n"), Member.GetShaderType(), *UniformBufferName, *ParameterName);
-				AddStructMember(Member, true);
-			}
+			const EShaderParameterType ShaderParameterType = GetShaderParameterTypeFromMemberResourceType(Member.GetBaseType());
+			const TCHAR* MemberSuffix = GetShaderParameterTypeSuffix(ShaderParameterType);
+
+			Decl.ResourceMembers.Appendf(TEXT("UB_RESOURCE_MEMBER_%s(%s, %s, %s);\n"), MemberSuffix, Member.GetShaderType(), *UniformBufferName, *GlobalParameterName);
+
+			AddRemapping(Member, GlobalParameterName, ShaderParameterType);
 		}
 	}
 }
@@ -403,7 +430,7 @@ static FString CreateHLSLUniformBufferDeclaration(const TCHAR* UniformBufferName
 			*StartEndSuffix,
 			*Decl.ResourceMembers,
 			UniformBufferName,
-			*Decl.StructMembers
+			*Decl.Remappings
 		);
 
 		if (Decl.ResourceMembers.Len() == 0)
