[DumpGPU] Dump referenced structs.

Problem: DumpGPU can only dump pass parameters but not parameters inside referenced structs. Thus, important parameters like the view uniforms containing all matrices are not dumped.

Solution: DumpGPU accesses the cpu copy of FUniformBuffer to store contents of referenced structs.

#rnx

diff --git a/Engine/Source/Runtime/RenderCore/Private/DumpGPU.cpp b/Engine/Source/Runtime/RenderCore/Private/DumpGPU.cpp
--- a/Engine/Source/Runtime/RenderCore/Private/DumpGPU.cpp
+++ b/Engine/Source/Runtime/RenderCore/Private/DumpGPU.cpp
@@ -2453,6 +2453,71 @@ static const TCHAR* GetPassEventNameWithGPUMask(const FRDGPass* Pass, FString& O
 	}
 }
 
+#if !(UE_BUILD_SHIPPING || UE_BUILD_TEST)
+void DumpReferencedStructs(const FShaderParametersMetadata* Metadata, const uint8* Content, uint32 ParentOffset)
+{
+	check(Metadata);
+	check(Content);
+
+	FRDGResourceDumpContext* ResourceDumpContext = GRDGResourceDumpContext_RenderThread;
+
+	for (const FShaderParametersMetadata::FMember& Member : Metadata->GetMembers())
+	{
+		const FShaderParametersMetadata* MemberMetadata = Member.GetStructMetadata();
+		if (!MemberMetadata)
+		{
+			// Basic types do not contain metadata and can be skipped
+			continue;
+		}
+
+		if (Member.GetBaseType() != EUniformBufferBaseType::UBMT_REFERENCED_STRUCT)
+		{
+			// Inlined structs can be skipped as they are contained in the pass content but since they can contain referenced structs we recursively call DumpReferencedStructs
+			DumpReferencedStructs(MemberMetadata, Content, ParentOffset + Member.GetOffset());
+			continue;
+		}
+
+		const FUniformBufferBinding* UniformBufferBinding = reinterpret_cast<const FUniformBufferBinding*>(Content + ParentOffset + Member.GetOffset());
+		if (!UniformBufferBinding)
+		{
+			// Skip referenced structs without a valid reference
+			continue;
+		}
+
+		FRHIUniformBuffer* UniformBuffer = UniformBufferBinding->GetUniformBuffer();
+		if (!UniformBuffer)
+		{
+			// Skip referenced structs without a valid buffer
+			continue;
+		}
+		
+		if (!ResourceDumpContext->IsDumped(UniformBuffer))
+		{
+			ResourceDumpContext->SetDumped(UniformBuffer);
+			ResourceDumpContext->Dump(MemberMetadata);
+
+			const uint8* ReferencedStructContent = (const uint8*)UE::RHI::UniformBufferContentMap::Get(UniformBuffer);
+			if (ReferencedStructContent)
+			{
+				DumpReferencedStructs(MemberMetadata, ReferencedStructContent, 0);
+
+				uint32 PassParametersByteSize = MemberMetadata->GetSize();
+				if (PassParametersByteSize > 0)
+				{
+					TArrayView<const uint8> ArrayView(ReferencedStructContent, PassParametersByteSize);
+					FString DumpFilePath = FRDGResourceDumpContext::kStructuresDir / ResourceDumpContext->PtrToString(UniformBuffer) + TEXT(".bin");
+					ResourceDumpContext->DumpBinaryToFile(ArrayView, DumpFilePath, FRDGResourceDumpContext::ETimingBucket::ParametersFileWrite);
+				}
+			}
+			else
+			{
+				UE_LOG(LogDumpGPU, Warning, TEXT("Unable to dump referenced struct %s as content could not be found. Please make sure to enable r.RHI.UniformBufferContentMap.Enable."), Member.GetShaderType());
+			}
+		}
+	}
+}
+#endif
+
 void FRDGBuilder::DumpNewGraphBuilder()
 {
 	FRDGResourceDumpContext* ResourceDumpContext = GRDGResourceDumpContext_RenderThread;
@@ -2716,6 +2781,12 @@ void FRDGBuilder::DumpResourcePassOutputs(const FRDGPass* Pass)
 			{
 				ResourceDumpContext->Dump(Metadata);
 				PassParametersByteSize = Metadata->GetSize();
+#if !(UE_BUILD_SHIPPING || UE_BUILD_TEST)
+				if (const uint8* Content = Pass->GetParameters().GetContents())
+				{
+					DumpReferencedStructs(Metadata, Content, 0);
+				}
+#endif
 			}
 		}
 
