Refactor SceneUniformBuffer to hoist out machinery for runtime extensible uniform buffers into a reusable component (TExtensibleUniformBuffer).
Also added r.PrintExtensibleUBs to print the layout of these extensible uniform buffers.

diff --git a/Engine/Source/Runtime/RenderCore/Private/ExtensibleUniformBuffer.cpp b/Engine/Source/Runtime/RenderCore/Private/ExtensibleUniformBuffer.cpp
new file mode 100644
--- /dev/null
+++ b/Engine/Source/Runtime/RenderCore/Private/ExtensibleUniformBuffer.cpp
@@ -0,0 +1,312 @@
+// Copyright Epic Games, Inc. All Rights Reserved.
+
+#include "ExtensibleUniformBuffer.h"
+#include "RenderGraphBuilder.h"
+
+/******************************/
+/* Struct layout registration */
+/******************************/
+
+struct FExtensibleUniformBufferTypeRegistry::FMemberInfo
+{
+	int32 Id;
+	int32 Offset;
+	FString Name;
+	const FShaderParametersMetadata* Metadata;
+	FShaderParameterStructConstructor DefaultValueFactory;
+};
+
+class FExtensibleUniformBufferTypeRegistry::FImpl
+{
+public:
+#if !UE_BUILD_SHIPPING
+	inline static TArray<FExtensibleUniformBufferTypeRegistry::FImpl*> Instances {};
+#endif
+
+	TArray<FExtensibleUniformBufferTypeRegistry::FMemberInfo> MemberInfos{};
+	FShaderParametersMetadata* StructMetaData = nullptr;
+
+	const TCHAR* LayoutName = nullptr;
+	const TCHAR* StructTypeName = nullptr;
+	const TCHAR* ShaderVariableName = nullptr;
+	const TCHAR* StaticSlotName = nullptr;
+	const ANSICHAR* FileName = nullptr;
+	int32 FileLine = 0;
+	uint64 MaxStructSize = 0;
+};
+
+FExtensibleUniformBufferTypeRegistry::FExtensibleUniformBufferTypeRegistry(
+	const TCHAR* InLayoutName,
+	const TCHAR* InStructTypeName,
+	const TCHAR* InShaderVariableName,
+	const TCHAR* InStaticSlotName,
+	const ANSICHAR* InFileName,
+	const int32 InFileLine,
+	const uint64 InMaxStructSize)
+: Impl(new FImpl)
+{
+	Impl->LayoutName = InLayoutName;
+	Impl->StructTypeName = InStructTypeName;
+	Impl->ShaderVariableName = InShaderVariableName;
+	Impl->StaticSlotName = InStaticSlotName;
+	Impl->FileName = InFileName;
+	Impl->FileLine = InFileLine;
+	Impl->MaxStructSize = InMaxStructSize;
+
+#if !UE_BUILD_SHIPPING
+	FImpl::Instances.Add(Impl);
+#endif
+}
+
+FExtensibleUniformBufferTypeRegistry::~FExtensibleUniformBufferTypeRegistry()
+{
+#if !UE_BUILD_SHIPPING
+	FImpl::Instances.Remove(Impl);
+#endif	
+	delete Impl;
+}
+
+#if !UE_BUILD_SHIPPING
+static FAutoConsoleCommand DumpExtensibleBufferInfoCmd(
+	TEXT("r.PrintExtensibleUBs"),
+	TEXT("Dumps layout of all extensible uniform buffers"),
+	FConsoleCommandWithOutputDeviceDelegate::CreateLambda(
+		[](FOutputDevice& Ar)
+		{
+			Ar.Log(TEXT("ExtensibleUniformBuffers:"));
+			for (const FExtensibleUniformBufferTypeRegistry::FImpl* Registry : FExtensibleUniformBufferTypeRegistry::FImpl::Instances)
+			{
+				uint32 BytesUsed = 0;
+				for (const FExtensibleUniformBufferTypeRegistry::FMemberInfo& Member : Registry->MemberInfos)
+				{
+					const FShaderParametersMetadata* Metadata = Member.Metadata;
+					BytesUsed = FMath::Max(BytesUsed, Member.Offset + Member.Metadata->GetSize());
+				}
+
+				Ar.Logf(TEXT("  - TypeName: '%s'"), Registry->StructTypeName);
+				Ar.Logf(TEXT("    LayoutName: '%s'"), Registry->LayoutName);
+				Ar.Logf(TEXT("    ShaderVariableName: '%s'"), Registry->ShaderVariableName);
+				Ar.Logf(TEXT("    StaticSlotName: '%s'"), Registry->StaticSlotName);
+				Ar.Logf(TEXT("    FileName: '%s'"), StringCast<TCHAR>(Registry->FileName).Get());
+				Ar.Logf(TEXT("    FileLine: %d"), Registry->FileLine);
+				Ar.Logf(TEXT("    MaxStructSize: %d"), Registry->MaxStructSize);
+				Ar.Logf(TEXT("    SizeUsed: %d"), BytesUsed);
+				Ar.Logf(TEXT("    Members: %s"), TEXT(""));
+
+				for (const FExtensibleUniformBufferTypeRegistry::FMemberInfo& Member : Registry->MemberInfos)
+				{
+					const FShaderParametersMetadata* Metadata = Member.Metadata;
+
+					Ar.Logf(TEXT("      - Name: '%s'"), *Member.Name);
+					Ar.Logf(TEXT("        TypeName: '%s'"), Metadata->GetStructTypeName());
+					Ar.Logf(TEXT("        Id: %d"), Member.Id);
+					Ar.Logf(TEXT("        Offset: %d"), Member.Offset);
+					Ar.Logf(TEXT("        Size: %d"), Metadata->GetSize());
+					Ar.Logf(TEXT("        FileName: '%s'"), StringCast<TCHAR>(Metadata->GetFileName()).Get());
+					Ar.Logf(TEXT("        FileLine: %d"), Metadata->GetFileLine());
+				}
+			}
+		})
+);
+#endif
+
+FExtensibleUniformBufferTypeRegistry::FMemberId FExtensibleUniformBufferTypeRegistry::RegisterMember(
+	const TCHAR* Name, const FShaderParametersMetadata& Metadata, FShaderParameterStructConstructor Factory)
+{
+	check(Impl->StructMetaData == nullptr);
+
+	int NewMemberInfoIndex = Impl->MemberInfos.Emplace();
+	FMemberInfo& NewMember = Impl->MemberInfos[NewMemberInfoIndex];
+	NewMember.Id = NewMemberInfoIndex;
+	NewMember.Name = Name;
+	NewMember.Offset = -1;
+	NewMember.Metadata = &Metadata;
+	NewMember.DefaultValueFactory = Factory;
+	return NewMemberInfoIndex;
+}
+
+void FExtensibleUniformBufferMemberRegistration::CommitAll()
+{
+	for (FExtensibleUniformBufferMemberRegistration* Instance : GetInstances())
+	{
+		Instance->Commit();
+	}
+	GetInstances().Empty();
+}
+
+static TArray<FExtensibleUniformBufferMemberRegistration*>* GExtensibleUniformBufferMemberRegistrationInstances = nullptr;
+TArray<FExtensibleUniformBufferMemberRegistration*>& FExtensibleUniformBufferMemberRegistration::GetInstances()
+{
+	// This is not a static or static local to ensure this works correctly during static initialization
+	if (GExtensibleUniformBufferMemberRegistrationInstances == nullptr)
+	{
+		GExtensibleUniformBufferMemberRegistrationInstances = new TArray<FExtensibleUniformBufferMemberRegistration*>();
+	}
+	return *GExtensibleUniformBufferMemberRegistrationInstances;
+}
+
+void FExtensibleUniformBufferMemberRegistration::Register(FExtensibleUniformBufferMemberRegistration& Entry)
+{
+	GetInstances().Add(&Entry);
+}
+
+FShaderParametersMetadata* FExtensibleUniformBufferTypeRegistry::GetStructMetadata()
+{
+	if (Impl->StructMetaData)
+	{
+		return Impl->StructMetaData;
+	}
+
+	FExtensibleUniformBufferMemberRegistration::CommitAll();
+
+	// TODO: Could be nice to have more intelligent packing here, rather than just appending with padding
+	TArray<FMemberInfo*> SortedMembers;
+	SortedMembers.Reserve(Impl->MemberInfos.Num());
+	for (FMemberInfo& Member : Impl->MemberInfos)
+	{
+		SortedMembers.Emplace(&Member);
+	}
+	Algo::SortBy(SortedMembers, &FMemberInfo::Name);
+
+	FShaderParametersMetadataBuilder Builder{};
+
+	for (FMemberInfo* Member : SortedMembers)
+	{
+		Member->Offset = Builder.AddNestedStruct(*Member->Name, Member->Metadata);
+	}
+
+	Impl->StructMetaData = Builder.Build(
+		FShaderParametersMetadata::EUseCase::UniformBuffer,
+		EUniformBufferBindingFlags::StaticAndShader,
+		Impl->LayoutName,
+		Impl->StructTypeName,
+		Impl->ShaderVariableName,
+		Impl->StaticSlotName,
+		Impl->FileName,
+		Impl->FileLine
+	);
+	checkf(Impl->StructMetaData->GetSize() <= Impl->MaxStructSize,
+		TEXT("The total size of the extensible uniform buffer members exceeds the storage capacity of the parameter struct. Increase the capacity. Use 'r.PrintExtensibleUBs' for details."));
+	return Impl->StructMetaData;
+}
+
+/**********/
+/* Buffer */
+/**********/
+
+FExtensibleUniformBufferBase::FExtensibleUniformBufferBase(FParameterTypeAdaptor& InParameterTypeAdaptor)
+	: ParameterTypeAdaptor(InParameterTypeAdaptor)
+	, MemberHasBeenSet(false, ParameterTypeAdaptor.Registry.Impl->MemberInfos.Num())
+{
+	CachedData.Init(0, ParameterTypeAdaptor.StructMetadata.GetSize());
+	check(Align(CachedData.GetData(), SHADER_PARAMETER_STRUCT_ALIGNMENT) == CachedData.GetData());
+}
+
+bool FExtensibleUniformBufferBase::Set(const FMemberId MemberId, const void* Value, const int32 ValueSize)
+{
+	const TArray<FRegistry::FMemberInfo>& MemberInfos = ParameterTypeAdaptor.Registry.Impl->MemberInfos;
+
+	check(MemberId >= 0 && MemberId < MemberInfos.Num());
+	const FRegistry::FMemberInfo& MemberInfo = MemberInfos[MemberId];
+	check(MemberInfo.Metadata->GetSize() == ValueSize);
+	check((MemberInfo.Offset + ValueSize) <= CachedData.Num());
+
+	MemberHasBeenSet[MemberId] = true;
+
+	uint8_t* MemberDataPtr = &CachedData[MemberInfo.Offset];
+	if (FPlatformMemory::Memcmp(MemberDataPtr, Value, ValueSize) != 0)
+	{
+		bAnyMemberDirty = true;
+		FPlatformMemory::Memcpy(MemberDataPtr, Value, ValueSize);
+		return true;
+	}
+	return false;
+}
+
+const void* FExtensibleUniformBufferBase::Get(const FMemberId MemberId, const int32 ValueSize) const
+{
+	const TArray<FRegistry::FMemberInfo>& MemberInfos = ParameterTypeAdaptor.Registry.Impl->MemberInfos;
+
+	check(MemberId >= 0 && MemberId < MemberInfos.Num());
+	const FRegistry::FMemberInfo& MemberInfo = MemberInfos[MemberId];
+	check(MemberInfo.Metadata->GetSize() == ValueSize);
+	check((MemberInfo.Offset + ValueSize) <= CachedData.Num());
+	
+	if (!MemberHasBeenSet[MemberId])
+	{
+		return nullptr;
+	}
+
+	return &CachedData[MemberInfo.Offset];
+}
+
+const void* FExtensibleUniformBufferBase::GetOrDefault(const FMemberId MemberId, const int32 ValueSize, FRDGBuilder& GraphBuilder)
+{
+	const TArray<FRegistry::FMemberInfo>& MemberInfos = ParameterTypeAdaptor.Registry.Impl->MemberInfos;
+
+	check(MemberId >= 0 && MemberId < MemberInfos.Num());
+	const FRegistry::FMemberInfo& MemberInfo = MemberInfos[MemberId];
+	check(MemberInfo.Metadata->GetSize() == ValueSize);
+	check((MemberInfo.Offset + ValueSize) <= CachedData.Num());
+
+	uint8* Offset = &CachedData[MemberInfo.Offset];
+	if (!MemberHasBeenSet[MemberId])
+	{
+		check((MemberInfo.Offset + MemberInfo.Metadata->GetSize()) <= static_cast<uint32>(CachedData.Num()));
+		MemberInfo.DefaultValueFactory(Offset, *MemberInfo.Metadata, GraphBuilder);
+		MemberHasBeenSet[MemberId] = true;
+	}
+
+	return Offset;
+}
+
+FRDGUniformBuffer* FExtensibleUniformBufferBase::GetBufferGeneric(FRDGBuilder& GraphBuilder)
+{
+	check(IsInRenderingThread());
+
+	const TArray<FRegistry::FMemberInfo>& MemberInfos = ParameterTypeAdaptor.Registry.Impl->MemberInfos;
+	if (Buffer == nullptr || bAnyMemberDirty)
+	{
+		for (int MemberIndex = 0; MemberIndex < MemberInfos.Num(); ++MemberIndex)
+		{
+			const FRegistry::FMemberInfo& MemberInfo = MemberInfos[MemberIndex];
+			if (!MemberHasBeenSet[MemberIndex])
+			{
+				// Never been populated, set up defaults, we defer this since we don't want to create redundant SRVs for the common case where it is not needed.
+				uint8* Offset = &CachedData[MemberInfo.Offset];
+
+				check((MemberInfo.Offset + MemberInfo.Metadata->GetSize()) <= static_cast<uint32>(CachedData.Num()));
+				MemberInfo.DefaultValueFactory(Offset, *MemberInfo.Metadata, GraphBuilder);
+				MemberHasBeenSet[MemberIndex] = true;
+			}
+		}
+		// Create and copy cached parameters into the RDG-lifetime struct
+		Buffer = ParameterTypeAdaptor.BuildRDGUniformBuffer(GraphBuilder, *this);
+
+		bAnyMemberDirty = false;
+	}
+	return Buffer;
+}
+
+FRHIUniformBuffer* FExtensibleUniformBufferBase::GetBufferRHI(FRDGBuilder& GraphBuilder)
+{
+	// Ensure the buffer is prepped.
+	GetBufferGeneric(GraphBuilder);
+	return GraphBuilder.ConvertToExternalUniformBuffer(Buffer);
+}
+
+#if !UE_BUILD_SHIPPING
+TArray<FExtensibleUniformBufferBase::FDebugMemberInfo> FExtensibleUniformBufferBase::GetDebugInfo() const
+{
+	TArray<FExtensibleUniformBufferBase::FDebugMemberInfo> DebugInfo;
+	for (const auto& Member : ParameterTypeAdaptor.StructMetadata.GetMembers())
+	{
+		FExtensibleUniformBufferBase::FDebugMemberInfo Info;
+		Info.ShaderType = Member.GetShaderType();
+		Info.Name = Member.GetName();
+		Info.Data = &CachedData[Member.GetOffset()];
+		DebugInfo.Add(Info);
+	}
+	return DebugInfo;
+}
+#endif
diff --git a/Engine/Source/Runtime/RenderCore/Public/ExtensibleUniformBuffer.h b/Engine/Source/Runtime/RenderCore/Public/ExtensibleUniformBuffer.h
new file mode 100644
--- /dev/null
+++ b/Engine/Source/Runtime/RenderCore/Public/ExtensibleUniformBuffer.h
@@ -0,0 +1,339 @@
+// Copyright Epic Games, Inc. All Rights Reserved.
+
+#pragma once
+
+#include "RenderGraphDefinitions.h"
+#include "RenderGraphBuilder.h"
+#include "ShaderParameterMacros.h"
+#include "ShaderParameterMetadata.h"
+#include "ShaderParameterMetadataBuilder.h"
+
+// Construct a shader parameter struct with the specified type at Obj.
+// The resulting struct must be valid for usage (i.e. all members initialized)
+using FShaderParameterStructConstructor = TFunction<void(void* Obj, const FShaderParametersMetadata& Metadata, FRDGBuilder& GraphBuilder)>;
+template<typename ParameterStructType>
+using TShaderParameterStructConstructor = TFunction<void(ParameterStructType& Obj, FRDGBuilder& GraphBuilder)>;
+
+/**
+ * Collects UB members during static initialization,
+ * and uses these to construct the buffer layout during Main.
+ */
+class FExtensibleUniformBufferTypeRegistry
+{
+public:
+	using FMemberId = int;
+	struct FMemberInfo;
+	class FImpl;
+	FImpl* Impl;
+
+	template<typename ParameterStructType>
+	static FExtensibleUniformBufferTypeRegistry New()
+	{
+		return FExtensibleUniformBufferTypeRegistry
+		{
+			ParameterStructType::FSourceDefinition::LayoutName,
+			ParameterStructType::FSourceDefinition::TypeName,
+			ParameterStructType::FSourceDefinition::ShaderVariableName,
+			ParameterStructType::FSourceDefinition::StaticSlotName,
+			ParameterStructType::FTypeInfo::FileName,
+			ParameterStructType::FTypeInfo::FileLine,
+			ParameterStructType::MaxSize 
+		};
+	}
+
+	template<typename ParameterStructType>
+	static FExtensibleUniformBufferTypeRegistry& Get();
+
+	RENDERCORE_API FExtensibleUniformBufferTypeRegistry(
+		const TCHAR* LayoutName,
+		const TCHAR* StructTypeName,
+		const TCHAR* ShaderVariableName,
+		const TCHAR* StaticSlotName,
+		const ANSICHAR* FileName,
+		const int32 FileLine,
+		const uint64 MaxStructSize);
+	RENDERCORE_API ~FExtensibleUniformBufferTypeRegistry();
+
+	RENDERCORE_API FMemberId RegisterMember(
+		const TCHAR* Name,
+		const FShaderParametersMetadata& StructMetadata,
+		// Used to set default values for this member if no value is set explicitly
+		FShaderParameterStructConstructor DefaultValueConstructor);
+
+	RENDERCORE_API FShaderParametersMetadata* GetStructMetadata();
+};
+
+/**
+ * Used to declare individual UB members.
+ */
+class FExtensibleUniformBufferMemberRegistration
+{
+public:
+	using FMemberId = FExtensibleUniformBufferTypeRegistry::FMemberId;
+
+	// Called at runtime
+	virtual void Commit() = 0;
+	static void CommitAll();
+
+	FString Name;
+	FMemberId MemberId = -1;
+	FShaderParameterStructConstructor DefaultValueConstructor;
+
+protected:
+	RENDERCORE_API static void Register(FExtensibleUniformBufferMemberRegistration& Entry);
+
+	// Called during static initialization
+	FExtensibleUniformBufferMemberRegistration(const TCHAR* Name)
+		: Name(Name)
+	{
+		Register(*this);
+	}
+	virtual ~FExtensibleUniformBufferMemberRegistration() = default;
+
+private:
+	static TArray<FExtensibleUniformBufferMemberRegistration*>& GetInstances();
+};
+
+template<typename TMember, typename ParameterStructType>
+class TExtensibleUniformBufferMemberRegistration : public FExtensibleUniformBufferMemberRegistration
+{
+	using FRegistry = FExtensibleUniformBufferTypeRegistry;
+
+public:
+	TExtensibleUniformBufferMemberRegistration(const TCHAR* Name, TShaderParameterStructConstructor<TMember> TemplatedDefaultValueConstructor)
+	: FExtensibleUniformBufferMemberRegistration(Name)
+	{
+		DefaultValueConstructor = [=](void* Obj, const FShaderParametersMetadata& Metadata, FRDGBuilder& GraphBuilder)
+		{
+			check(&Metadata == TMember::FTypeInfo::GetStructMetadata());
+			TMember& Member = *(new(Obj) TMember());
+			TemplatedDefaultValueConstructor(Member, GraphBuilder);
+		};
+	}
+
+	void Commit() override
+	{
+		MemberId = FRegistry::Get<ParameterStructType>().RegisterMember(
+			*Name,
+			*TShaderParameterStructTypeInfo<TMember>::GetStructMetadata(),
+			DefaultValueConstructor);
+	}
+};
+
+/**
+ * Shader parameter struct with a layout that is defined at runtime by FExtensibleUniformBufferTypeRegistry
+ */
+template<typename SelfType>
+class TExtensibleUniformParametersBase
+{
+public:
+	struct FTypeInfo
+	{
+		static constexpr int32 NumRows = 1;
+		static constexpr int32 NumColumns = 1;
+		static constexpr int32 NumElements = 0;
+		static constexpr int32 Alignment = SHADER_PARAMETER_STRUCT_ALIGNMENT;
+		static constexpr bool bIsStoredInConstantBuffer = true;
+		static constexpr const ANSICHAR* const FileName = SelfType::FSourceDefinition::FileName;
+		static constexpr int32 FileLine = SelfType::FSourceDefinition::FileLine;
+		using TAlignedType = SelfType;
+		static const FShaderParametersMetadata* GetStructMetadata()
+		{
+			FExtensibleUniformBufferTypeRegistry& Registry = FExtensibleUniformBufferTypeRegistry::Get<SelfType>();
+			return Registry.GetStructMetadata();
+		}
+	};
+
+	const void* GetContents() const
+	{
+		return Data;
+	}
+
+private:
+	template<typename, template<typename> typename>
+	friend class TExtensibleUniformBuffer;
+	friend class FExtensibleUniformBufferTypeRegistry;
+
+	// TODO: Decouple parameter struct and actual data, so the data can be allocated dynamically (variable size)
+	// This requires some work in buffer binding/parameter macros/RHI, as these all assume the data to be inline.
+	static constexpr size_t MaxSize = 2048;
+	uint8_t Data[MaxSize];
+};
+
+#define DECLARE_EXTENSIBLE_UNIFORM_BUFFER_PARAMETER_STRUCT(PrefixKeywords, ParameterStructTypeName, ShaderVariableAndSlotName)						\
+	DECLARE_UNIFORM_BUFFER_STRUCT(ParameterStructTypeName, PrefixKeywords)																			\
+	MS_ALIGN(SHADER_PARAMETER_STRUCT_ALIGNMENT) class ParameterStructTypeName final : public TExtensibleUniformParametersBase<ParameterStructTypeName>	\
+	{																																				\
+	public:																																			\
+		struct FSourceDefinition																													\
+		{																																			\
+			static constexpr const TCHAR* const LayoutName = TEXT(#ParameterStructTypeName);														\
+			static constexpr const TCHAR* const TypeName = TEXT(#ParameterStructTypeName);                                         					\
+			static constexpr const TCHAR* const ShaderVariableName = TEXT(#ShaderVariableAndSlotName);												\
+			static constexpr const TCHAR* const StaticSlotName = TEXT(#ShaderVariableAndSlotName);													\
+			static constexpr const ANSICHAR* const FileName = UE_LOG_SOURCE_FILE(__FILE__);															\
+			static constexpr int32 FileLine = __LINE__;																								\
+		};																																			\
+	} GCC_ALIGN(SHADER_PARAMETER_STRUCT_ALIGNMENT);
+
+/**
+ * Used to create instances of an extensible buffer, set data, and build an RDG object
+ */
+class FExtensibleUniformBufferBase
+{
+public:
+	// Get and re-create the UB if the cached parameters are modified.
+	RENDERCORE_API FRDGUniformBuffer* GetBufferGeneric(FRDGBuilder& GraphBuilder);
+	RENDERCORE_API FRHIUniformBuffer* GetBufferRHI(FRDGBuilder& GraphBuilder);
+
+#if !UE_BUILD_SHIPPING
+	struct FDebugMemberInfo
+	{
+		FString ShaderType;
+		FString Name;
+		const void* Data;
+	};
+	TArray<FDebugMemberInfo> GetDebugInfo() const;
+#endif
+
+protected:
+	using FMemberId = FExtensibleUniformBufferTypeRegistry::FMemberId;
+	using FRegistry = FExtensibleUniformBufferTypeRegistry;
+
+	class FParameterTypeAdaptor
+	{
+	public:
+		FParameterTypeAdaptor(const FShaderParametersMetadata& InStructMetadata, const FExtensibleUniformBufferTypeRegistry& InRegistry)
+		: StructMetadata(InStructMetadata)
+		, Registry(InRegistry)
+		{ }
+		virtual ~FParameterTypeAdaptor() = default;
+		// Create and copy cached parameters into the RDG-lifetime struct
+		virtual FRDGUniformBuffer* BuildRDGUniformBuffer(FRDGBuilder& GraphBuilder, FExtensibleUniformBufferBase& Buffer) const = 0;
+		const FShaderParametersMetadata& StructMetadata;
+		const FExtensibleUniformBufferTypeRegistry& Registry;
+	};
+
+	RENDERCORE_API FExtensibleUniformBufferBase(FParameterTypeAdaptor& InParameterTypeAdaptor);
+
+	/**
+	 * Set a field in the parameter struct.
+	 * The change will be reflected in any buffer that GetBuffer() returns after this call.
+	 * Returns true if anything actually changed.
+	 */
+	RENDERCORE_API bool Set(const FMemberId MemberId, const void* Value, const int32 ValueSize);
+
+	/**
+	 * Retrieve a field in the parameter struct.
+	 * If the field has not been set, this will return null;
+	 */
+	RENDERCORE_API const void* Get(const FMemberId MemberId, const int32 ExpectedValueSize) const UE_LIFETIMEBOUND;
+
+	/**
+	 * Retrieve a field in the parameter struct.
+	 * If the field has not been set, this will set a default value;
+	 */
+	RENDERCORE_API const void* GetOrDefault(const FMemberId MemberId, const int32 ExpectedValueSize, FRDGBuilder& GraphBuilder) UE_LIFETIMEBOUND;
+
+protected:
+	const FParameterTypeAdaptor& ParameterTypeAdaptor;
+
+	FRDGUniformBuffer* Buffer = nullptr;
+
+	bool bAnyMemberDirty = false;
+	TBitArray<> MemberHasBeenSet;
+	TArray<uint8_t, TSizedHeapAllocator<32>> CachedData;
+};
+
+
+template<typename ParameterStructType, template<typename> typename TMemberRegistration>
+class TExtensibleUniformBuffer : public FExtensibleUniformBufferBase
+{
+	template<typename SelfType>
+	friend class TExtensibleUniformParametersBase;
+
+	class TParameterTypeAdaptor : public FParameterTypeAdaptor
+	{
+		using FDerivedExtensibleUniformBuffer = TExtensibleUniformBuffer<ParameterStructType, TMemberRegistration>;
+	public:
+		static TParameterTypeAdaptor& Get()
+		{
+			static TParameterTypeAdaptor Instance{};
+			return Instance;
+		}
+
+		TParameterTypeAdaptor()
+		: FParameterTypeAdaptor(*ParameterStructType::FTypeInfo::GetStructMetadata(), FExtensibleUniformBufferTypeRegistry::Get<ParameterStructType>())
+		{ }
+
+		virtual ~TParameterTypeAdaptor() = default;
+
+		virtual FRDGUniformBuffer* BuildRDGUniformBuffer(FRDGBuilder& GraphBuilder, FExtensibleUniformBufferBase& Container) const override
+		{
+			ParameterStructType* ParameterStruct = GraphBuilder.AllocObject<ParameterStructType>();
+			FDerivedExtensibleUniformBuffer& DerivedContainer = static_cast<FDerivedExtensibleUniformBuffer&>(Container);
+			FPlatformMemory::Memcpy(ParameterStruct->Data, DerivedContainer.CachedData.GetData(), DerivedContainer.CachedData.Num());
+			return GraphBuilder.CreateUniformBuffer(ParameterStruct);
+		}
+	};
+
+public:
+	TExtensibleUniformBuffer()
+	: FExtensibleUniformBufferBase(TParameterTypeAdaptor::Get())
+	{ }
+
+	/**
+	 * Set a field in the parameter struct.
+	 * The change will be reflected in any buffer that GetBuffer() returns after this call.
+	 * Returns true if anything actually changed.
+	 */
+	template<typename TMember>
+	bool Set(const TMemberRegistration<TMember>& Registration, const TMember& Value)
+	{
+		return FExtensibleUniformBufferBase::Set(Registration.MemberId, &Value, TMember::FTypeInfo::GetStructMetadata()->GetSize());
+	}
+
+	/**
+	 * Retrieve a field in the parameter struct.
+	 * If the field has not been set, this will crash.
+	 */
+	template<typename TMember>
+	const TMember& Get(const TMemberRegistration<TMember>& Registration) const UE_LIFETIMEBOUND
+	{
+		const void* Ptr = FExtensibleUniformBufferBase::Get(Registration.MemberId, TMember::FTypeInfo::GetStructMetadata()->GetSize());
+		checkf(Ptr, TEXT("%s has not been set yet"), *Registration.Name);
+		return *reinterpret_cast<const TMember*>(Ptr);
+	}
+
+	/**
+	 * Retrieve a field in the parameter struct.
+	 * If the field has not been set, this will set a default value.
+	 */
+	template<typename TMember>
+	const TMember& GetOrDefault(const TMemberRegistration<TMember>& Registration, FRDGBuilder& GraphBuilder) const UE_LIFETIMEBOUND
+	{
+		return *reinterpret_cast<const TMember*>(FExtensibleUniformBufferBase::GetOrDefault(Registration.MemberId, TMember::FTypeInfo::GetStructMetadata()->GetSize(), GraphBuilder));
+	}
+
+	// Get and re-create the UB if the cached parameters are modified.
+	TRDGUniformBufferRef<ParameterStructType> GetBuffer(FRDGBuilder& GraphBuilder)
+	{
+		return static_cast<TRDGUniformBufferRef<ParameterStructType>>(FExtensibleUniformBufferBase::GetBufferGeneric(GraphBuilder));
+	}
+};
+
+#define IMPLEMENT_EXTENSIBLE_UNIFORM_BUFFER(PrefixKeywords, ParameterStructType)																	\
+	static FUniformBufferStaticSlotRegistrar UniformBufferStaticSlot_##ParameterStructType(ParameterStructType::FSourceDefinition::StaticSlotName);	\
+	template<> PrefixKeywords FExtensibleUniformBufferTypeRegistry& FExtensibleUniformBufferTypeRegistry::Get<ParameterStructType>()               \
+	{                                                                                                                                       	\
+		static FExtensibleUniformBufferTypeRegistry Instance = FExtensibleUniformBufferTypeRegistry::New<ParameterStructType>();               	\
+		return Instance;                                                                                                                    	\
+	}                                                                                                                                       	\
+																																				\
+	PrefixKeywords const FShaderParametersMetadata* GetForwardDeclaredShaderParametersStructMetadata(const ParameterStructType* DummyPtr)      	\
+	{                                                                                                                                       	\
+		return ParameterStructType::FTypeInfo::GetStructMetadata();                                                                            	\
+	}                                                                                                                                       	\
+																																				\
+	FShaderParametersMetadataRegistration ParametersMetadataRegistration                                                        				\
+	{ TFunctionRef<const ::FShaderParametersMetadata* ()>{ ParameterStructType::FTypeInfo::GetStructMetadata } };
diff --git a/Engine/Source/Runtime/Renderer/Private/SceneUniformBuffer.cpp b/Engine/Source/Runtime/Renderer/Private/SceneUniformBuffer.cpp
--- a/Engine/Source/Runtime/Renderer/Private/SceneUniformBuffer.cpp
+++ b/Engine/Source/Runtime/Renderer/Private/SceneUniformBuffer.cpp
@@ -1,239 +1,5 @@
 // Copyright Epic Games, Inc. All Rights Reserved.
 
 #include "SceneUniformBuffer.h"
-#include "RenderGraphBuilder.h"
 
-/******************************/
-/* Struct layout registration */
-/******************************/
-
-struct FSceneUniformBufferTypeRegistry::FMemberInfo
-{
-	int32 Id;
-	int32 Offset;
-	FString Name;
-	const FShaderParametersMetadata* Metadata;
-	FShaderParameterStructConstructor DefaultValueFactory;
-};
-
-class FSceneUniformBufferTypeRegistry::FImpl
-{
-public:
-	static inline TArray<FMemberInfo> MemberInfos{};
-	static inline bool bStructMetadataWasBuilt = false;
-
-	static FShaderParametersMetadata* BuildStructMetadata();
-};
-
-FSceneUniformBuffer::FMemberId FSceneUniformBufferTypeRegistry::RegisterMember(
-	const TCHAR* Name, const FShaderParametersMetadata& Metadata, FShaderParameterStructConstructor Factory)
-{
-	check(!FImpl::bStructMetadataWasBuilt);
-
-	int NewMemberInfoIndex = FImpl::MemberInfos.Emplace();
-	FMemberInfo& NewMember = FImpl::MemberInfos[NewMemberInfoIndex];
-	NewMember.Id = NewMemberInfoIndex;
-	NewMember.Name = Name;
-	NewMember.Offset = -1;
-	NewMember.Metadata = &Metadata;
-	NewMember.DefaultValueFactory = Factory;
-	return NewMemberInfoIndex;
-}
-
-void FSceneUniformBufferMemberRegistration::CommitAll()
-{
-	for (FSceneUniformBufferMemberRegistration* Instance : GetInstances())
-	{
-		Instance->Commit();
-	}
-	GetInstances().Empty();
-}
-
-static TArray<FSceneUniformBufferMemberRegistration*>* GSceneUniformBufferMemberRegistrationInstances = nullptr;
-TArray<FSceneUniformBufferMemberRegistration*>& FSceneUniformBufferMemberRegistration::GetInstances()
-{
-	// This is not a static or static local to ensure this works correctly during static initialization
-	if (GSceneUniformBufferMemberRegistrationInstances == nullptr)
-	{
-		GSceneUniformBufferMemberRegistrationInstances = new TArray<FSceneUniformBufferMemberRegistration*>();
-	}
-	return *GSceneUniformBufferMemberRegistrationInstances;
-}
-
-void FSceneUniformBufferMemberRegistration::Register(FSceneUniformBufferMemberRegistration& Entry)
-{
-	GetInstances().Add(&Entry);
-}
-
-FShaderParametersMetadata* FSceneUniformBufferTypeRegistry::FImpl::BuildStructMetadata()
-{
-	FSceneUniformBufferMemberRegistration::CommitAll();
-
-	// TODO: Could be nice to have more intelligent packing here, rather than just appending with padding
-	TArray<FMemberInfo*> SortedMembers;
-	SortedMembers.Reserve(MemberInfos.Num());
-	for (FMemberInfo& Member : MemberInfos)
-	{
-		SortedMembers.Emplace(&Member);
-	}
-	Algo::SortBy(SortedMembers, &FMemberInfo::Name);
-
-	FShaderParametersMetadataBuilder Builder{};
-
-	for (FMemberInfo* Member : SortedMembers)
-	{
-		Member->Offset = Builder.AddNestedStruct(*Member->Name, Member->Metadata);
-	}
-
-	bStructMetadataWasBuilt = true;
-	FShaderParametersMetadata* Metadata = Builder.Build(
-		FShaderParametersMetadata::EUseCase::UniformBuffer,
-		EUniformBufferBindingFlags::StaticAndShader,
-		TEXT("FSceneUniformParameters"),
-		TEXT("FSceneUniformParameters"),
-		TEXT("Scene"),
-		TEXT("Scene"),
-		FSceneUniformParameters::FTypeInfo::FileName,
-		FSceneUniformParameters::FTypeInfo::FileLine
-	);
-	checkf(Metadata->GetSize() <= FSceneUniformParameters::MaxSize,
-		TEXT("The total size of the scene uniform buffer members exceeds the storage capacity in FSceneUniformParameters::Data. Increase the capacity."));
-	return Metadata;
-}
-
-/*****************/
-/* Buffer values */
-/*****************/
-
-IMPLEMENT_STATIC_UNIFORM_BUFFER_SLOT(Scene);
-
-RENDERER_API const FShaderParametersMetadata* GetForwardDeclaredShaderParametersStructMetadata(const FSceneUniformParameters* DummyPtr)
-{
-	return FSceneUniformParameters::FTypeInfo::GetStructMetadata();
-}
-
-FShaderParametersMetadataRegistration SceneUniformParametersMetadataRegistration
-{ TFunctionRef<const ::FShaderParametersMetadata* ()>{ FSceneUniformParameters::FTypeInfo::GetStructMetadata } };
-
-RENDERER_API const FShaderParametersMetadata* FSceneUniformParameters::FTypeInfo::GetStructMetadata()
-{
-	static FShaderParametersMetadata* Metadata = FSceneUniformBuffer::FRegistry::FImpl::BuildStructMetadata();
-	return Metadata;
-}
-
-FSceneUniformParameters::FSceneUniformParameters()
-{ }
-
-
-
-FSceneUniformBuffer::FSceneUniformBuffer() :
-	Buffer(nullptr),
-	bAnyMemberDirty(false),
-	MemberHasBeenSet(false, FRegistry::FImpl::MemberInfos.Num())
-{
-	CachedData.Init(0, FSceneUniformParameters::FTypeInfo::GetStructMetadata()->GetSize());
-	check(Align(CachedData.GetData(), SHADER_PARAMETER_STRUCT_ALIGNMENT) == CachedData.GetData());
-}
-
-RENDERER_API bool FSceneUniformBuffer::Set(const FMemberId MemberId, const void* Value, const int32 ValueSize)
-{
-	check(MemberId >= 0 && MemberId < FRegistry::FImpl::MemberInfos.Num());
-	const FRegistry::FMemberInfo& MemberInfo = FRegistry::FImpl::MemberInfos[MemberId];
-	check(MemberInfo.Metadata->GetSize() == ValueSize);
-	check((MemberInfo.Offset + ValueSize) <= CachedData.Num());
-
-	MemberHasBeenSet[MemberId] = true;
-
-	uint8_t* MemberDataPtr = &CachedData[MemberInfo.Offset];
-	if (FPlatformMemory::Memcmp(MemberDataPtr, Value, ValueSize) != 0)
-	{
-		bAnyMemberDirty = true;
-		FPlatformMemory::Memcpy(MemberDataPtr, Value, ValueSize);
-		return true;
-	}
-	return false;
-}
-
-RENDERER_API const void* FSceneUniformBuffer::Get(const FMemberId MemberId, const int32 ValueSize) const
-{
-	check(MemberId >= 0 && MemberId < FRegistry::FImpl::MemberInfos.Num());
-	const FRegistry::FMemberInfo& MemberInfo = FRegistry::FImpl::MemberInfos[MemberId];
-	check(MemberInfo.Metadata->GetSize() == ValueSize);
-	check((MemberInfo.Offset + ValueSize) <= CachedData.Num());
-	
-	if (!MemberHasBeenSet[MemberId])
-	{
-		return nullptr;
-	}
-
-	return &CachedData[MemberInfo.Offset];
-}
-
-RENDERER_API const void* FSceneUniformBuffer::GetOrDefault(const FMemberId MemberId, const int32 ValueSize, FRDGBuilder& GraphBuilder)
-{
-	check(MemberId >= 0 && MemberId < FRegistry::FImpl::MemberInfos.Num());
-	const FRegistry::FMemberInfo& MemberInfo = FRegistry::FImpl::MemberInfos[MemberId];
-	check(MemberInfo.Metadata->GetSize() == ValueSize);
-	check((MemberInfo.Offset + ValueSize) <= CachedData.Num());
-
-	uint8* Offset = &CachedData[MemberInfo.Offset];
-	if (!MemberHasBeenSet[MemberId])
-	{
-		check((MemberInfo.Offset + MemberInfo.Metadata->GetSize()) <= static_cast<uint32>(CachedData.Num()));
-		MemberInfo.DefaultValueFactory(Offset, *MemberInfo.Metadata, GraphBuilder);
-		MemberHasBeenSet[MemberId] = true;
-	}
-
-	return Offset;
-}
-
-RENDERER_API TRDGUniformBufferRef<FSceneUniformParameters> FSceneUniformBuffer::GetBuffer(FRDGBuilder& GraphBuilder)
-{
-	check(IsInRenderingThread());
-	if (Buffer == nullptr || bAnyMemberDirty)
-	{
-		for (int MemberIndex = 0; MemberIndex < FRegistry::FImpl::MemberInfos.Num(); ++MemberIndex)
-		{
-			const FRegistry::FMemberInfo& MemberInfo = FRegistry::FImpl::MemberInfos[MemberIndex];
-			if (!MemberHasBeenSet[MemberIndex])
-			{
-				// Never been populated, set up defaults, we defer this since we don't want to create redundant SRVs for the common case where it is not needed.
-				uint8* Offset = &CachedData[MemberInfo.Offset];
-
-				check((MemberInfo.Offset + MemberInfo.Metadata->GetSize()) <= static_cast<uint32>(CachedData.Num()));
-				MemberInfo.DefaultValueFactory(Offset, *MemberInfo.Metadata, GraphBuilder);
-				MemberHasBeenSet[MemberIndex] = true;
-			}
-		}
-		// Create and copy cached parameters into the RDG-lifetime struct
-		FSceneUniformParameters* ParameterStruct = GraphBuilder.AllocObject<FSceneUniformParameters>();
-		FPlatformMemory::Memcpy(ParameterStruct->Data, CachedData.GetData(), CachedData.Num());
-		Buffer = GraphBuilder.CreateUniformBuffer(ParameterStruct);
-
-		bAnyMemberDirty = false;
-	}
-	return Buffer;
-}
-
-RENDERER_API FRHIUniformBuffer* FSceneUniformBuffer::GetBufferRHI(FRDGBuilder& GraphBuilder)
-{
-	// Ensure the buffer is prepped.
-	GetBuffer(GraphBuilder);
-	return GraphBuilder.ConvertToExternalUniformBuffer(Buffer);
-}
-
-#if !UE_BUILD_SHIPPING
-TArray<FSceneUniformBuffer::FDebugMemberInfo> FSceneUniformBuffer::GetDebugInfo() const
-{
-	TArray<FSceneUniformBuffer::FDebugMemberInfo> DebugInfo;
-	for (const auto& Member : FSceneUniformParameters::FTypeInfo::GetStructMetadata()->GetMembers())
-	{
-		FSceneUniformBuffer::FDebugMemberInfo Info;
-		Info.ShaderType = Member.GetShaderType();
-		Info.Name = Member.GetName();
-		Info.Data = &CachedData[Member.GetOffset()];
-		DebugInfo.Add(Info);
-	}
-	return DebugInfo;
-}
-#endif
+IMPLEMENT_EXTENSIBLE_UNIFORM_BUFFER(RENDERER_API, FSceneUniformParameters)
diff --git a/Engine/Source/Runtime/Renderer/Public/SceneUniformBuffer.h b/Engine/Source/Runtime/Renderer/Public/SceneUniformBuffer.h
--- a/Engine/Source/Runtime/Renderer/Public/SceneUniformBuffer.h
+++ b/Engine/Source/Runtime/Renderer/Public/SceneUniformBuffer.h
@@ -2,219 +2,29 @@
 
 #pragma once
 
-#include "RenderGraphDefinitions.h"
-#include "ShaderParameterMacros.h"
-#include "ShaderParameterMetadata.h"
-#include "ShaderParameterMetadataBuilder.h"
-
-// Construct a shader parameter struct with the specified type at Obj.
-// The resulting struct must be valid for usage (i.e. all members initialized)
-using FShaderParameterStructConstructor = TFunction<void(void* Obj, const FShaderParametersMetadata& Metadata, FRDGBuilder& GraphBuilder)>;
-template<typename TParameterStruct>
-using TShaderParameterStructConstructor = TFunction<void(TParameterStruct& Obj, FRDGBuilder& GraphBuilder)>;
+#include "ExtensibleUniformBuffer.h"
 
 /**
- * Collects scene UB members during static initialization,
- * and uses these to construct the buffer layout during Main.
+ * RDG shader parameter struct containing data for FSceneUniformBuffer.
  */
-class FSceneUniformBufferTypeRegistry
-{
-public:
-	using FMemberId = int;
-	struct FMemberInfo;
-	class FImpl;
-
-	// Must be called during static initialization
-	RENDERER_API static FMemberId RegisterMember(
-		const TCHAR* Name,
-		const FShaderParametersMetadata& StructMetadata,
-		// Used to set default values for this member if no value is set explicitly
-		FShaderParameterStructConstructor DefaultValueConstructor);
-};
-
-class FSceneUniformBufferMemberRegistration
-{
-public:
-	using FMemberId = FSceneUniformBufferTypeRegistry::FMemberId;
-
-	// Called at runtime
-	virtual void Commit() = 0;
-	static void CommitAll();
-
-	FString Name;
-	FMemberId MemberId = -1;
-	FShaderParameterStructConstructor DefaultValueConstructor;
-
-protected:
-	RENDERER_API static void Register(FSceneUniformBufferMemberRegistration& Entry);
-
-	// Called during static initialization
-	FSceneUniformBufferMemberRegistration(const TCHAR* Name)
-		: Name(Name)
-	{
-		Register(*this);
-	}
-	virtual ~FSceneUniformBufferMemberRegistration() = default;
-
-private:
-	static TArray<FSceneUniformBufferMemberRegistration*>& GetInstances();
-};
+DECLARE_EXTENSIBLE_UNIFORM_BUFFER_PARAMETER_STRUCT(RENDERER_API, FSceneUniformParameters, Scene);
 
+/**
+ * Used to declare individual UB members.
+ */
 template<typename TMember>
-class TSceneUniformBufferMemberRegistration : public FSceneUniformBufferMemberRegistration
-{
-public:
-	TSceneUniformBufferMemberRegistration(const TCHAR* Name, TShaderParameterStructConstructor<TMember> TemplatedDefaultValueConstructor)
-		: FSceneUniformBufferMemberRegistration(Name)
-	{
-		DefaultValueConstructor = [=](void* Obj, const FShaderParametersMetadata& Metadata, FRDGBuilder& GraphBuilder)
-		{
-			check(&Metadata == TMember::FTypeInfo::GetStructMetadata());
-			TMember& Member = *(new(Obj) TMember());
-			TemplatedDefaultValueConstructor(Member, GraphBuilder);
-		};
-	}
+using TSceneUniformBufferMemberRegistration = TExtensibleUniformBufferMemberRegistration<TMember, FSceneUniformParameters>;
 
-	void Commit() override
-	{
-		MemberId = FSceneUniformBufferTypeRegistry::RegisterMember(
-			*Name,
-			*TShaderParameterStructTypeInfo<TMember>::GetStructMetadata(),
-			DefaultValueConstructor);
-	}
-};
-
-#define DECLARE_SCENE_UB_STRUCT(StructType, FieldName, PrefixKeywords) \
-	namespace SceneUB { \
-		PrefixKeywords extern TSceneUniformBufferMemberRegistration<StructType> FieldName; \
+#define DECLARE_SCENE_UB_STRUCT(StructType, FieldName, PrefixKeywords)						\
+	namespace SceneUB {																		\
+		PrefixKeywords extern TSceneUniformBufferMemberRegistration<StructType> FieldName;	\
 	}
 
 #define IMPLEMENT_SCENE_UB_STRUCT(StructType, FieldName, DefaultValueFactoryType) \
 	TSceneUniformBufferMemberRegistration<StructType> SceneUB::FieldName { TEXT(#FieldName), DefaultValueFactoryType } 
 
-
-DECLARE_UNIFORM_BUFFER_STRUCT(FSceneUniformParameters, RENDERER_API)
-/**
- * RDG shader parameter struct containing data for FSceneUniformBuffer.
- */
-MS_ALIGN(SHADER_PARAMETER_STRUCT_ALIGNMENT) class FSceneUniformParameters final
-{
-public:
-	struct FTypeInfo {
-		static constexpr int32 NumRows = 1;
-		static constexpr int32 NumColumns = 1;
-		static constexpr int32 NumElements = 0;
-		static constexpr int32 Alignment = SHADER_PARAMETER_STRUCT_ALIGNMENT;
-		static constexpr bool bIsStoredInConstantBuffer = true;
-		static constexpr const ANSICHAR* const FileName = UE_LOG_SOURCE_FILE(__FILE__);
-		static constexpr int32 FileLine = __LINE__;
-		using TAlignedType = FSceneUniformParameters;
-		static RENDERER_API const FShaderParametersMetadata* GetStructMetadata();
-	};
-
-	const void* GetContents() const
-	{
-		return Data;
-	}
-
-	FSceneUniformParameters();
-
-private:
-	friend class FSceneUniformBuffer;
-	friend class FSceneUniformBufferTypeRegistry::FImpl;
-
-	// TODO: Decouple parameter struct and actual data, so the data can be allocated dynamically (variable size)
-	// This requires some work in buffer binding/parameter macros/RHI, as these all assume the data to be inline.
-	static constexpr size_t MaxSize = 4096;
-	uint8_t Data[MaxSize];
-} GCC_ALIGN(SHADER_PARAMETER_STRUCT_ALIGNMENT);
-
-
 /**
  * Holds scene-scoped parameters and stores these in uniform (constant) buffers for access on GPU.
  */
-class FSceneUniformBuffer final
-{
-public:
-	using FMemberId = FSceneUniformBufferTypeRegistry::FMemberId;
-	using FRegistry = FSceneUniformBufferTypeRegistry;
-
-	RENDERER_API FSceneUniformBuffer();
-
-	/**
-	 * Set a field in the parameter struct.
-	 * The change will be reflected in any buffer that GetBuffer() returns after this call.
-	 * Returns true if anything actually changed.
-	 */
-	template<typename TMember>
-	bool Set(const TSceneUniformBufferMemberRegistration<TMember>& Registration, const TMember& Value)
-	{
-		return Set(Registration.MemberId, &Value, TMember::FTypeInfo::GetStructMetadata()->GetSize());
-	}
-
-	/**
-	 * Retrieve a field in the parameter struct.
-	 * If the field has not been set, this will crash.
-	 */
-	template<typename TMember>
-	const TMember& Get(const TSceneUniformBufferMemberRegistration<TMember>& Registration) const UE_LIFETIMEBOUND
-	{
-		const void* Ptr = Get(Registration.MemberId, TMember::FTypeInfo::GetStructMetadata()->GetSize());
-		checkf(Ptr, TEXT("SceneUB::%s has not been set yet"), *Registration.Name);
-		return *reinterpret_cast<const TMember*>(Ptr);
-	}
-
-	/**
-	 * Retrieve a field in the parameter struct.
-	 * If the field has not been set, this will set a default value.
-	 */
-	template<typename TMember>
-	const TMember& GetOrDefault(const TSceneUniformBufferMemberRegistration<TMember>& Registration, FRDGBuilder& GraphBuilder) const UE_LIFETIMEBOUND
-	{
-		return *reinterpret_cast<const TMember*>(GetOrDefault(Registration.MemberId, TMember::FTypeInfo::GetStructMetadata()->GetSize(), GraphBuilder));
-	}
-
-	// Get and re-create the UB if the cached parameters are modified.
-	RENDERER_API TRDGUniformBufferRef<FSceneUniformParameters> GetBuffer(FRDGBuilder& GraphBuilder);
-	RENDERER_API FRHIUniformBuffer* GetBufferRHI(FRDGBuilder& GraphBuilder);
-
-#if !UE_BUILD_SHIPPING
-	struct FDebugMemberInfo
-	{
-		FString ShaderType;
-		FString Name;
-		const void* Data;
-	};
-	TArray<FDebugMemberInfo> GetDebugInfo() const;
-#endif
-
-private:
-	template<typename TMember>
-	friend class TSceneUniformBufferMemberRegistration;
-	friend class FSceneUniformParameters;
-
-	/**
-	 * Set a field in the parameter struct.
-	 * The change will be reflected in any buffer that GetBuffer() returns after this call.
-	 * Returns true if anything actually changed.
-	 */
-	RENDERER_API bool Set(const FMemberId MemberId, const void* Value, const int32 ValueSize);
-
-	/**
-	 * Retrieve a field in the parameter struct.
-	 * If the field has not been set, this will return null;
-	 */
-	RENDERER_API const void* Get(const FMemberId MemberId, const int32 ExpectedValueSize) const UE_LIFETIMEBOUND;
-
-	/**
-	 * Retrieve a field in the parameter struct.
-	 * If the field has not been set, this will set a default value;
-	 */
-	RENDERER_API const void* GetOrDefault(const FMemberId MemberId, const int32 ExpectedValueSize, FRDGBuilder& GraphBuilder) UE_LIFETIMEBOUND;
-
-	TRDGUniformBufferRef<FSceneUniformParameters> Buffer;
-
-	bool bAnyMemberDirty;
-	TBitArray<> MemberHasBeenSet;
-	TArray<uint8_t, TSizedHeapAllocator<32>> CachedData;
-};
+class FSceneUniformBuffer : public TExtensibleUniformBuffer<FSceneUniformParameters, TSceneUniformBufferMemberRegistration>
+{ };
\ No newline at end of file
