Shader Compiler:  fix for recompiling global shaders crashing when in mobile preview mode in the editor.  Need to include the shader platform in the shader compiler job key, to differentiate between versions of the same global shader on the preview platform versus default PC platform.  Without this fix, the shader compiler would generate a global shader for the preview platform, but skip kicking off a job on the default platform as it would find an existing job with the same key, eventually leading to a fatal missing shader error.

#rnx

diff --git a/Engine/Plugins/Compositing/OpenColorIO/Source/OpenColorIO/Private/OpenColorIOShaderMap.cpp b/Engine/Plugins/Compositing/OpenColorIO/Source/OpenColorIO/Private/OpenColorIOShaderMap.cpp
--- a/Engine/Plugins/Compositing/OpenColorIO/Source/OpenColorIO/Private/OpenColorIOShaderMap.cpp
+++ b/Engine/Plugins/Compositing/OpenColorIO/Source/OpenColorIO/Private/OpenColorIOShaderMap.cpp
@@ -160,7 +160,7 @@ void FOpenColorIOShaderType::BeginCompileShader(
 	FShaderTarget InTarget
 	)
 {
-	FShaderCompileJob* NewJob = GShaderCompilingManager->PrepareShaderCompileJob(InShaderMapId, FShaderCompileJobKey(this), EShaderCompileJobPriority::Normal);
+	FShaderCompileJob* NewJob = GShaderCompilingManager->PrepareShaderCompileJob(InShaderMapId, FShaderCompileJobKey(this, InPlatform), EShaderCompileJobPriority::Normal);
 	if (!NewJob)
 	{
 		return;
diff --git a/Engine/Plugins/Experimental/NNERuntimeIREE/Source/NNERuntimeIREEShader/Private/NNERuntimeIREEShaderMap.cpp b/Engine/Plugins/Experimental/NNERuntimeIREE/Source/NNERuntimeIREEShader/Private/NNERuntimeIREEShaderMap.cpp
--- a/Engine/Plugins/Experimental/NNERuntimeIREE/Source/NNERuntimeIREEShader/Private/NNERuntimeIREEShaderMap.cpp
+++ b/Engine/Plugins/Experimental/NNERuntimeIREE/Source/NNERuntimeIREEShader/Private/NNERuntimeIREEShaderMap.cpp
@@ -120,7 +120,7 @@ void FNNERuntimeIREEShaderType::BeginCompileShader(
 {
 	FShaderCompileJob* NewJob = GShaderCompilingManager->PrepareShaderCompileJob(
 		InShaderMapId, 
-		FShaderCompileJobKey(this, nullptr, PermutationId),
+		FShaderCompileJobKey(this, InPlatform, nullptr, PermutationId),
 		EShaderCompileJobPriority::Normal
 		);
 
diff --git a/Engine/Plugins/FX/Niagara/Source/NiagaraEditor/Private/NiagaraCompiler.cpp b/Engine/Plugins/FX/Niagara/Source/NiagaraEditor/Private/NiagaraCompiler.cpp
--- a/Engine/Plugins/FX/Niagara/Source/NiagaraEditor/Private/NiagaraCompiler.cpp
+++ b/Engine/Plugins/FX/Niagara/Source/NiagaraEditor/Private/NiagaraCompiler.cpp
@@ -1879,7 +1879,7 @@ int32 FHlslNiagaraCompiler::CompileScript(const FStringView GroupName, const FNi
 		}
 		if (NiagaraShaderType)
 		{
-			TRefCountPtr<FShaderCompileJob> Job = GShaderCompilingManager->PrepareShaderCompileJob(JobID, FShaderCompileJobKey(NiagaraShaderType), EShaderCompileJobPriority::Normal);
+			TRefCountPtr<FShaderCompileJob> Job = GShaderCompilingManager->PrepareShaderCompileJob(JobID, FShaderCompileJobKey(NiagaraShaderType, EShaderPlatform::SP_PCD3D_SM5), EShaderCompileJobPriority::Normal);
 			if (Job)
 			{
 				TArray<FShaderCommonCompileJobPtr> NewJobs;
@@ -1942,7 +1942,7 @@ uint32 FHlslNiagaraCompiler::CompileScriptVM(const FStringView GroupName, const
 	if (CompileResults.bVMSucceeded && NiagaraShaderType)
 	{
 		const uint32 JobID = FShaderCommonCompileJob::GetNextJobId();
-		TRefCountPtr<FShaderCompileJob> Job = GShaderCompilingManager->PrepareShaderCompileJob(JobID, FShaderCompileJobKey(NiagaraShaderType), EShaderCompileJobPriority::Normal);
+		TRefCountPtr<FShaderCompileJob> Job = GShaderCompilingManager->PrepareShaderCompileJob(JobID, FShaderCompileJobKey(NiagaraShaderType, EShaderPlatform::SP_PCD3D_SM5), EShaderCompileJobPriority::Normal);
 		if (Job)
 		{
 			Job->Input.Target = FShaderTarget(SF_Compute, SP_PCD3D_SM5);
diff --git a/Engine/Plugins/FX/Niagara/Source/NiagaraShader/Private/NiagaraShader.cpp b/Engine/Plugins/FX/Niagara/Source/NiagaraShader/Private/NiagaraShader.cpp
--- a/Engine/Plugins/FX/Niagara/Source/NiagaraShader/Private/NiagaraShader.cpp
+++ b/Engine/Plugins/FX/Niagara/Source/NiagaraShader/Private/NiagaraShader.cpp
@@ -332,7 +332,7 @@ void FNiagaraShaderType::BeginCompileShader(
 	FShaderTarget Target
 	)
 {
-	FShaderCompileJob* NewJob = GShaderCompilingManager->PrepareShaderCompileJob(ShaderMapId, FShaderCompileJobKey(this, nullptr, PermutationId), EShaderCompileJobPriority::Normal);
+	FShaderCompileJob* NewJob = GShaderCompilingManager->PrepareShaderCompileJob(ShaderMapId, FShaderCompileJobKey(this, Platform, nullptr, PermutationId), EShaderCompileJobPriority::Normal);
 	if (!NewJob)
 	{
 		return;
@@ -420,7 +420,7 @@ void FNiagaraShaderType::BeginCompileShaderFromSource(
 	TArray<TRefCountPtr<FShaderCommonCompileJob>>& NewJobs
 ) const
 {
-	FShaderCompileJob* NewJob = GShaderCompilingManager->PrepareShaderCompileJob(ShaderMapId, FShaderCompileJobKey(this, nullptr, PermutationId), EShaderCompileJobPriority::Normal);
+	FShaderCompileJob* NewJob = GShaderCompilingManager->PrepareShaderCompileJob(ShaderMapId, FShaderCompileJobKey(this, Platform, nullptr, PermutationId), EShaderCompileJobPriority::Normal);
 	if (!NewJob)
 	{
 		return;
diff --git a/Engine/Plugins/Runtime/ComputeFramework/Source/ComputeFramework/Private/ComputeKernelShaderMap.cpp b/Engine/Plugins/Runtime/ComputeFramework/Source/ComputeFramework/Private/ComputeKernelShaderMap.cpp
--- a/Engine/Plugins/Runtime/ComputeFramework/Source/ComputeFramework/Private/ComputeKernelShaderMap.cpp
+++ b/Engine/Plugins/Runtime/ComputeFramework/Source/ComputeFramework/Private/ComputeKernelShaderMap.cpp
@@ -142,7 +142,7 @@ void FComputeKernelShaderType::BeginCompileShader(
 {
 	FShaderCompileJob* NewJob = GShaderCompilingManager->PrepareShaderCompileJob(
 		InShaderMapId, 
-		FShaderCompileJobKey(this, nullptr, PermutationId),
+		FShaderCompileJobKey(this, InPlatform, nullptr, PermutationId),
 		EShaderCompileJobPriority::Normal
 		);
 
diff --git a/Engine/Source/Runtime/Engine/Private/Materials/MaterialShader.cpp b/Engine/Source/Runtime/Engine/Private/Materials/MaterialShader.cpp
--- a/Engine/Source/Runtime/Engine/Private/Materials/MaterialShader.cpp
+++ b/Engine/Source/Runtime/Engine/Private/Materials/MaterialShader.cpp
@@ -1897,7 +1897,7 @@ void FMaterialShaderType::BeginCompileShader(
 	const TCHAR* DebugExtension
 	) const
 {
-	FShaderCompileJob* NewJob = GShaderCompilingManager->PrepareShaderCompileJob(ShaderMapJobId, FShaderCompileJobKey(this, nullptr, PermutationId), Priority);
+	FShaderCompileJob* NewJob = GShaderCompilingManager->PrepareShaderCompileJob(ShaderMapJobId, FShaderCompileJobKey(this, Platform, nullptr, PermutationId), Priority);
 	if (NewJob)
 	{
 		PrepareMaterialShaderCompileJob(Platform, PermutationFlags, Material, ShaderMapId, MaterialEnvironment, nullptr, DebugGroupName, DebugDescription, DebugExtension, NewJob);
@@ -1923,7 +1923,7 @@ void FMaterialShaderType::BeginCompileShaderPipeline(
 	UE_LOG(LogShaders, Verbose, TEXT("	Pipeline: %s"), ShaderPipeline->GetName());
 
 	// Add all the jobs as individual first, then add the dependencies into a pipeline job
-	auto* NewPipelineJob = GShaderCompilingManager->PreparePipelineCompileJob(ShaderMapJobId, FShaderPipelineCompileJobKey(ShaderPipeline, nullptr, kUniqueShaderPermutationId), Priority);
+	auto* NewPipelineJob = GShaderCompilingManager->PreparePipelineCompileJob(ShaderMapJobId, FShaderPipelineCompileJobKey(ShaderPipeline, Platform, nullptr, kUniqueShaderPermutationId), Priority);
 	if (NewPipelineJob)
 	{
 		for (FShaderCompileJob* StageJob : NewPipelineJob->StageJobs)
diff --git a/Engine/Source/Runtime/Engine/Private/Materials/MeshMaterialShader.cpp b/Engine/Source/Runtime/Engine/Private/Materials/MeshMaterialShader.cpp
--- a/Engine/Source/Runtime/Engine/Private/Materials/MeshMaterialShader.cpp
+++ b/Engine/Source/Runtime/Engine/Private/Materials/MeshMaterialShader.cpp
@@ -94,7 +94,7 @@ void FMeshMaterialShaderType::BeginCompileShader(
 	const TCHAR* DebugDescription,
 	const TCHAR* DebugExtension) const
 {
-	FShaderCompileJob* NewJob = GShaderCompilingManager->PrepareShaderCompileJob(ShaderMapJobId, FShaderCompileJobKey(this, VertexFactoryType, PermutationId), Priority);
+	FShaderCompileJob* NewJob = GShaderCompilingManager->PrepareShaderCompileJob(ShaderMapJobId, FShaderCompileJobKey(this, Platform, VertexFactoryType, PermutationId), Priority);
 	if (NewJob)
 	{
 		PrepareMeshMaterialShaderCompileJob(Platform, PermutationFlags, Material, ShaderMapId, MaterialEnvironment, nullptr, DebugGroupName, DebugDescription, DebugExtension, NewJob);
@@ -122,7 +122,7 @@ void FMeshMaterialShaderType::BeginCompileShaderPipeline(
 	UE_LOG(LogShaders, Verbose, TEXT("	Pipeline: %s"), ShaderPipeline->GetName());
 
 	// Add all the jobs as individual first, then add the dependencies into a pipeline job
-	auto* NewPipelineJob = GShaderCompilingManager->PreparePipelineCompileJob(ShaderMapJobId, FShaderPipelineCompileJobKey(ShaderPipeline, VertexFactoryType, PermutationId), Priority);
+	auto* NewPipelineJob = GShaderCompilingManager->PreparePipelineCompileJob(ShaderMapJobId, FShaderPipelineCompileJobKey(ShaderPipeline, Platform, VertexFactoryType, PermutationId), Priority);
 	if (NewPipelineJob)
 	{
 		for (FShaderCompileJob* StageJob : NewPipelineJob->StageJobs)
diff --git a/Engine/Source/Runtime/Engine/Private/ShaderCompiler/ShaderCompilerEditor.cpp b/Engine/Source/Runtime/Engine/Private/ShaderCompiler/ShaderCompilerEditor.cpp
--- a/Engine/Source/Runtime/Engine/Private/ShaderCompiler/ShaderCompilerEditor.cpp
+++ b/Engine/Source/Runtime/Engine/Private/ShaderCompiler/ShaderCompilerEditor.cpp
@@ -116,7 +116,7 @@ static void PrepareGlobalShaderCompileJob(EShaderPlatform Platform,
 void FGlobalShaderTypeCompiler::BeginCompileShader(const FGlobalShaderType* ShaderType, int32 PermutationId, EShaderPlatform Platform, EShaderPermutationFlags PermutationFlags, TArray<FShaderCommonCompileJobPtr>& NewJobs)
 {
 	// Global shaders are always high priority (often need to block on completion)
-	FShaderCompileJob* NewJob = GShaderCompilingManager->PrepareShaderCompileJob(GlobalShaderMapId, FShaderCompileJobKey(ShaderType, nullptr, PermutationId), EShaderCompileJobPriority::High);
+	FShaderCompileJob* NewJob = GShaderCompilingManager->PrepareShaderCompileJob(GlobalShaderMapId, FShaderCompileJobKey(ShaderType, Platform, nullptr, PermutationId), EShaderCompileJobPriority::High);
 	if (NewJob)
 	{
 		PrepareGlobalShaderCompileJob(Platform, PermutationFlags, nullptr, NewJob);
@@ -130,7 +130,7 @@ void FGlobalShaderTypeCompiler::BeginCompileShaderPipeline(EShaderPlatform Platf
 	UE_LOG(LogShaders, Verbose, TEXT("	Pipeline: %s"), ShaderPipeline->GetName());
 
 	// Add all the jobs as individual first, then add the dependencies into a pipeline job
-	FShaderPipelineCompileJob* NewPipelineJob = GShaderCompilingManager->PreparePipelineCompileJob(GlobalShaderMapId, FShaderPipelineCompileJobKey(ShaderPipeline, nullptr, kUniqueShaderPermutationId), EShaderCompileJobPriority::High);
+	FShaderPipelineCompileJob* NewPipelineJob = GShaderCompilingManager->PreparePipelineCompileJob(GlobalShaderMapId, FShaderPipelineCompileJobKey(ShaderPipeline, Platform, nullptr, kUniqueShaderPermutationId), EShaderCompileJobPriority::High);
 	if (NewPipelineJob)
 	{
 		for (FShaderCompileJob* StageJob : NewPipelineJob->StageJobs)
diff --git a/Engine/Source/Runtime/RenderCore/Private/ShaderCore.cpp b/Engine/Source/Runtime/RenderCore/Private/ShaderCore.cpp
--- a/Engine/Source/Runtime/RenderCore/Private/ShaderCore.cpp
+++ b/Engine/Source/Runtime/RenderCore/Private/ShaderCore.cpp
@@ -4734,7 +4734,7 @@ FShaderPipelineCompileJob::FShaderPipelineCompileJob(uint32 InHash, uint32 InId,
 	StageJobs.Empty(Stages.Num());
 	for (const FShaderType* ShaderType : Stages)
 	{
-		const FShaderCompileJobKey StageKey(ShaderType, InKey.VFType, InKey.PermutationId);
+		const FShaderCompileJobKey StageKey(ShaderType, InKey.Platform, InKey.VFType, InKey.PermutationId);
 		StageJobs.Add(new FShaderCompileJob(StageKey.MakeHash(InId), InId, InPriroity, StageKey));
 	}
 	
diff --git a/Engine/Source/Runtime/RenderCore/Public/ShaderCompilerJobTypes.h b/Engine/Source/Runtime/RenderCore/Public/ShaderCompilerJobTypes.h
--- a/Engine/Source/Runtime/RenderCore/Public/ShaderCompilerJobTypes.h
+++ b/Engine/Source/Runtime/RenderCore/Public/ShaderCompilerJobTypes.h
@@ -314,19 +314,20 @@ using FShaderCommonCompileJobPtr = TRefCountPtr<FShaderCommonCompileJob>;
 
 struct FShaderCompileJobKey
 {
-	explicit FShaderCompileJobKey(const FShaderType* InType = nullptr, const FVertexFactoryType* InVFType = nullptr, int32 InPermutationId = 0)
-		: ShaderType(InType), VFType(InVFType), PermutationId(InPermutationId)
+	explicit FShaderCompileJobKey(const FShaderType* InType = nullptr, EShaderPlatform InPlatform = EShaderPlatform::SP_PCD3D_SM6, const FVertexFactoryType* InVFType = nullptr, int32 InPermutationId = 0)
+		: ShaderType(InType), VFType(InVFType), PermutationId(InPermutationId), Platform(InPlatform)
 	{}
 
-	uint32 MakeHash(uint32 Id) const { return HashCombine(HashCombine(HashCombine(GetTypeHash(Id), GetTypeHash(VFType)), GetTypeHash(ShaderType)), GetTypeHash(PermutationId)); }
+	uint32 MakeHash(uint32 Id) const { return HashCombine(HashCombine(HashCombine(HashCombine(GetTypeHash(Id), GetTypeHash(VFType)), GetTypeHash(ShaderType)), GetTypeHash(PermutationId)), GetTypeHash(Platform)); }
 	RENDERCORE_API FString ToString() const;
 	const FShaderType* ShaderType;
 	const FVertexFactoryType* VFType;
 	int32 PermutationId;
+	EShaderPlatform Platform;
 
 	friend inline bool operator==(const FShaderCompileJobKey& Lhs, const FShaderCompileJobKey& Rhs)
 	{
-		return Lhs.VFType == Rhs.VFType && Lhs.ShaderType == Rhs.ShaderType && Lhs.PermutationId == Rhs.PermutationId;
+		return Lhs.VFType == Rhs.VFType && Lhs.ShaderType == Rhs.ShaderType && Lhs.PermutationId == Rhs.PermutationId && Lhs.Platform == Rhs.Platform;
 	}
 	friend inline bool operator!=(const FShaderCompileJobKey& Lhs, const FShaderCompileJobKey& Rhs)
 	{
@@ -403,19 +404,20 @@ private:
 
 struct FShaderPipelineCompileJobKey
 {
-	explicit FShaderPipelineCompileJobKey(const FShaderPipelineType* InType = nullptr, const FVertexFactoryType* InVFType = nullptr, int32 InPermutationId = 0)
-		: ShaderPipeline(InType), VFType(InVFType), PermutationId(InPermutationId)
+	explicit FShaderPipelineCompileJobKey(const FShaderPipelineType* InType = nullptr, EShaderPlatform InPlatform = EShaderPlatform::SP_PCD3D_SM6, const FVertexFactoryType* InVFType = nullptr, int32 InPermutationId = 0)
+		: ShaderPipeline(InType), VFType(InVFType), PermutationId(InPermutationId), Platform(InPlatform)
 	{}
 
-	uint32 MakeHash(uint32 Id) const { return HashCombine(HashCombine(HashCombine(GetTypeHash(Id), GetTypeHash(ShaderPipeline)), GetTypeHash(VFType)), GetTypeHash(PermutationId)); }
+	uint32 MakeHash(uint32 Id) const { return HashCombine(HashCombine(HashCombine(HashCombine(GetTypeHash(Id), GetTypeHash(ShaderPipeline)), GetTypeHash(VFType)), GetTypeHash(PermutationId)), GetTypeHash(Platform)); }
 
 	const FShaderPipelineType* ShaderPipeline;
 	const FVertexFactoryType* VFType;
 	int32 PermutationId;
+	EShaderPlatform Platform;
 
 	friend inline bool operator==(const FShaderPipelineCompileJobKey& Lhs, const FShaderPipelineCompileJobKey& Rhs)
 	{
-		return Lhs.ShaderPipeline == Rhs.ShaderPipeline && Lhs.VFType == Rhs.VFType && Lhs.PermutationId == Rhs.PermutationId;
+		return Lhs.Platform == Rhs.Platform && Lhs.ShaderPipeline == Rhs.ShaderPipeline && Lhs.VFType == Rhs.VFType && Lhs.PermutationId == Rhs.PermutationId;
 	}
 	friend inline bool operator!=(const FShaderPipelineCompileJobKey& Lhs, const FShaderPipelineCompileJobKey& Rhs)
 	{
