Restore ERHIAccess::SRVMask for FRenderTargetPool::CreateRenderTarget to fix regressions

The previous attempt to reduce transitions was a breaking change.

diff --git a/Engine/Source/Runtime/RenderCore/Private/RenderTargetPool.cpp b/Engine/Source/Runtime/RenderCore/Private/RenderTargetPool.cpp
--- a/Engine/Source/Runtime/RenderCore/Private/RenderTargetPool.cpp
+++ b/Engine/Source/Runtime/RenderCore/Private/RenderTargetPool.cpp
@@ -100,9 +100,7 @@ FPooledRenderTarget* FRenderTargetPool::CreateRenderTarget(FRHICommandListBase&
 #endif
 	}
 
-	const ETextureCreateFlags TexCreateFlags = Desc.Flags;
-	const ERHIAccess AccessInitial = RHIGetDefaultResourceState(TexCreateFlags, false);
-	FRHITextureCreateDesc CreateDesc(Desc, AccessInitial, Name);
+	FRHITextureCreateDesc CreateDesc(Desc, ERHIAccess::SRVMask, Name);
 	const static FLazyName ClassName(TEXT("FPooledRenderTarget"));
 	CreateDesc.SetClassName(ClassName);
 
@@ -114,7 +112,7 @@ FPooledRenderTarget* FRenderTargetPool::CreateRenderTarget(FRHICommandListBase&
 	PooledRenderTargets.Add(Result);
 	PooledRenderTargetHashes.Add(DescHash);
 
-	if (EnumHasAnyFlags(TexCreateFlags, TexCreate_UAV))
+	if (EnumHasAnyFlags(Desc.Flags, TexCreate_UAV))
 	{
 		if (GRHISupportsUAVFormatAliasing)
 		{
