[Rendering] Updating the Create Render Target size logging to only log when we hit a specific threshold (number of logs or size of created targets) per frame.

#Synced-CL 48344811

diff --git a/Engine/Source/Runtime/RenderCore/Private/RenderTargetPool.cpp b/Engine/Source/Runtime/RenderCore/Private/RenderTargetPool.cpp
--- a/Engine/Source/Runtime/RenderCore/Private/RenderTargetPool.cpp
+++ b/Engine/Source/Runtime/RenderCore/Private/RenderTargetPool.cpp
@@ -25,12 +25,71 @@ UE_TRACE_EVENT_BEGIN(Cpu, FRenderTargetPool_CreateTexture, NoSync)
 	UE_TRACE_EVENT_FIELD(uint32, Height)
 UE_TRACE_EVENT_END()
 
+
+///////////////////////////////////////////////////////////////////////////////
+
 static int32 GRenderTargetPoolLogCreationSizes = 0;
 static FAutoConsoleVariableRef CVarRenderTargetPoolLogCreationSizes(
 	TEXT("r.RenderTargetPool.LogCreationSizes"),
 	GRenderTargetPoolLogCreationSizes,
 	TEXT("Enable/disable warning log for render target pool creation sizes."));
 
+static int32 GRenderTargetPoolLogCreationSizesSpamInterval = 4; // 4 per frames
+static FAutoConsoleVariableRef CVarRenderTargetPoolLogCreationSpamInterval(
+	TEXT("r.RenderTargetPool.LogCreationSizesSpamInterval"),
+	GRenderTargetPoolLogCreationSizesSpamInterval,
+	TEXT("The number of logs in a given frame that are generated before we force a print."));
+
+static int32 GRenderTargetPoolLogCreationSizesSizeInterval = 1024 * 1024; // 1 MB
+static FAutoConsoleVariableRef CVarRenderTargetPoolLogCreationSizeInterval(
+	TEXT("r.RenderTargetPool.LogCreationSizesSizeInterval"),
+	GRenderTargetPoolLogCreationSizesSizeInterval,
+	TEXT("The total size of accumulated logs during a frame before we force a print."));
+
+static void LogRenderTargetCreationSizeAndFrequencies(const FRHITextureCreateInfo& Desc, const TCHAR* Name)
+{
+#if !UE_BUILD_SHIPPING && !UE_SERVER
+	static FString AccumulatedLogs;
+	static uint64 AccumulatedCreatedTargets = 0;
+	static uint64 AccumulatedCreatedTargetSizes = 0;
+	static uint64 TrackingCurrentFrame = 0;
+	
+	if (GRenderTargetPoolLogCreationSizes <= 0)
+	{
+		AccumulatedLogs.Empty();
+		AccumulatedCreatedTargets = 0;
+		AccumulatedCreatedTargetSizes = 0;
+		TrackingCurrentFrame = 0;
+		return;
+	}
+
+	// Log last frame's results IFF we are over one or more of our target
+	// thresholds. Either way, clear the accumulators for the current frame.
+	const uint64 CurrentRTFrame = GFrameCounterRenderThread;
+	if (TrackingCurrentFrame != CurrentRTFrame)
+	{
+		if (AccumulatedCreatedTargets >= GRenderTargetPoolLogCreationSizesSpamInterval ||
+			AccumulatedCreatedTargetSizes >= GRenderTargetPoolLogCreationSizesSizeInterval)
+		{
+			UE_LOG(LogRenderTargetPool, Warning, TEXT("FRenderTargetPool::CreateRenderTarget(s) since last frame:\n%s"), *AccumulatedLogs);
+		}
+
+		AccumulatedLogs.Empty();
+		AccumulatedCreatedTargets = 0;
+		AccumulatedCreatedTargetSizes = 0;
+		TrackingCurrentFrame = CurrentRTFrame;
+	}
+
+	const uint64 CreationSize = RHICalcTexturePlatformSize(Desc).Size;
+	AccumulatedCreatedTargets += 1;
+	AccumulatedCreatedTargetSizes += CreationSize;
+	AccumulatedLogs += FString::Printf(
+		TEXT("\tRenderTargetCreation { \"%s\", Size(%.2f KB) }\n"), Name, CreationSize / 1024.0);
+#endif
+}
+
+///////////////////////////////////////////////////////////////////////////////
+
 TRefCountPtr<IPooledRenderTarget> CreateRenderTarget(FRHITexture* Texture, const TCHAR* Name)
 {
 	check(Texture);
@@ -82,23 +141,7 @@ FPooledRenderTarget* FRenderTargetPool::CreateRenderTarget(FRHICommandListBase&
 		<< FRenderTargetPool_CreateTexture.Height(Desc.Extent.Y);
 #endif
 
-	if (GRenderTargetPoolLogCreationSizes > 0)
-	{
-#if CSV_PROFILER_STATS
-		int32 CurrentCaptureFrame = 0;
-		if (!FCsvProfiler::Get())
-		{
-			CurrentCaptureFrame = -1;
-		}
-		else if (FCsvProfiler::Get()->IsCapturing_Renderthread())
-		{
-			CurrentCaptureFrame = FCsvProfiler::Get()->GetCaptureFrameNumberRT();
-		}
-		UE_LOG(LogRenderTargetPool, Warning, TEXT("FRenderTargetPool::CreateRenderTarget{Frame %d}: %s create with expected size - %ld"), CurrentCaptureFrame, Name, RHICalcTexturePlatformSize(Desc).Size);
-#else
-		UE_LOG(LogRenderTargetPool, Warning, TEXT("FRenderTargetPool::CreateRenderTarget: %s create with expected size - %ld"), Name, RHICalcTexturePlatformSize(Desc).Size);
-#endif
-	}
+	LogRenderTargetCreationSizeAndFrequencies(Desc, Name);
 
 	FRHITextureCreateDesc CreateDesc(Desc, ERHIAccess::SRVMask, Name);
 	const static FLazyName ClassName(TEXT("FPooledRenderTarget"));
