Bindless Uniform Buffer support
- Adding way to enable bindless constant buffer access for uniform buffers.
- Shaders can enable this behavior by setting "SUPPORTS_BINDLESS_UNIFORM_BUFFER"
- Once enabled, shaders must resolve/load the uniform buffer global before accessing it.
- This relies on a new bindless mode that adds getter calls to uniform buffer member accesses.

#Synced-CL 48479927

diff --git a/Engine/Shaders/Public/BindlessResources.ush b/Engine/Shaders/Public/BindlessResources.ush
--- a/Engine/Shaders/Public/BindlessResources.ush
+++ b/Engine/Shaders/Public/BindlessResources.ush
@@ -4,6 +4,17 @@
 
 #include "Platform.ush"
 
+#if !defined(SUPPORTS_BINDLESS_UNIFORM_BUFFER)
+	#define SUPPORTS_BINDLESS_UNIFORM_BUFFER 0
+
+#elif !PLATFORM_SUPPORTS_CONSTANTBUFFER_OBJECT
+	// Bindless Uniform Buffers require ConstantBufer<> support
+	// TODO: christopher.waters - should this be required at a higher level?
+
+	#undef SUPPORTS_BINDLESS_UNIFORM_BUFFER
+	#define SUPPORTS_BINDLESS_UNIFORM_BUFFER 0
+#endif
+
 #if UE_BINDLESS_ENABLED
 	#if !defined(UE_BINDLESS_SRV_HEAP_DECL)
 		#define UE_BINDLESS_SRV_HEAP_DECL(Type)
@@ -13,34 +24,85 @@
 		#define UE_BINDLESS_UAV_HEAP_DECL(Type)
 	#endif
 
+	#if !defined(UE_BINDLESS_CBV_HEAP_DECL)
+		#define UE_BINDLESS_CBV_HEAP_DECL(Type)
+	#endif
+
 	#if !defined(UE_BINDLESS_SAMPLER_HEAP_DECL)
 		#define UE_BINDLESS_SAMPLER_HEAP_DECL(Type)
 	#endif
 
-	// Using platform defined bindless macros
+	#if SUPPORTS_BINDLESS_UNIFORM_BUFFER
+		#define UB_SAFE_TYPE(UniformBufferName) SafeType##UniformBufferName##Constants
+
+		#define UB_BINDLESS_DECL(UniformBufferType, UniformBufferName) \
+			typedef ConstantBuffer<UniformBufferType> UB_SAFE_TYPE(UniformBufferName); \
+			UE_BINDLESS_CBV_HEAP_DECL(UB_SAFE_TYPE(UniformBufferName)); \
+			UB_SAFE_TYPE(UniformBufferName) Get##UniformBufferName##Constants(uint Index) { return GetCBVFromHeap(UB_SAFE_TYPE(UniformBufferName), Index); }
+
+	#endif // SUPPORTS_BINDLESS_UNIFORM_BUFFER
+
+	#define UB_RES_SAFE_TYPE(UniformBufferName, UniformBufferMember) SafeType##UniformBufferName##_##UniformBufferMember
+
+	// Declarations for UB member types and heaps.
+	#define UB_RES_DECL_SRV(Type, UniformBufferName, UniformBufferMember) \
+		typedef Type UB_RES_SAFE_TYPE(UniformBufferName, UniformBufferMember); \
+		UE_BINDLESS_SRV_HEAP_DECL(UB_RES_SAFE_TYPE(UniformBufferName, UniformBufferMember))
+
+	#define UB_RES_DECL_UAV(Type, UniformBufferName, UniformBufferMember) \
+		typedef Type UB_RES_SAFE_TYPE(UniformBufferName, UniformBufferMember); \
+		UE_BINDLESS_UAV_HEAP_DECL(UB_RES_SAFE_TYPE(UniformBufferName, UniformBufferMember))
+
+	#define UB_RES_DECL_SAMPLER(Type, UniformBufferName, UniformBufferMember) \
+		typedef Type UB_RES_SAFE_TYPE(UniformBufferName, UniformBufferMember); \
+		UE_BINDLESS_SAMPLER_HEAP_DECL(UB_RES_SAFE_TYPE(UniformBufferName, UniformBufferMember))
+
+	// Getters for cbuffer members
+	#define UB_CB_LOAD_SRV(UniformBufferName, UniformBufferMember) \
+		GetSRVFromHeap(UB_RES_SAFE_TYPE(UniformBufferName, UniformBufferMember), UB_CB_PREFIXED_MEMBER_ACCESS(UniformBufferName, BindlessSRV_, UniformBufferMember))
+
+	#define UB_CB_LOAD_UAV(UniformBufferName, UniformBufferMember) \
+		GetUAVFromHeap(UB_RES_SAFE_TYPE(UniformBufferName, UniformBufferMember), UB_CB_PREFIXED_MEMBER_ACCESS(UniformBufferName, BindlessUAV_, UniformBufferMember))
+
+	#define UB_CB_LOAD_SAMPLER(UniformBufferName, UniformBufferMember) \
+		GetSamplerFromHeap(UB_RES_SAFE_TYPE(UniformBufferName, UniformBufferMember), UB_CB_PREFIXED_MEMBER_ACCESS(UniformBufferName, BindlessSampler_, UniformBufferMember))
+
+	// Getters for ConstantBuffer<> members. The constant buffer variable name needs to be an argument because it could be a local variable.
+	#define UB_CBO_LOAD_SRV(UniformBufferName, VarName, UniformBufferMember) \
+		GetSRVFromHeap(UB_RES_SAFE_TYPE(UniformBufferName, UniformBufferMember), VarName##.BindlessSRV_##UniformBufferMember)
+
+	#define UB_CBO_LOAD_UAV(UniformBufferName, VarName, UniformBufferMember) \
+		GetUAVFromHeap(UB_RES_SAFE_TYPE(UniformBufferName, UniformBufferMember), VarName##.BindlessUAV_##UniformBufferMember)
+
+	#define UB_CBO_LOAD_SAMPLER(UniformBufferName, VarName, UniformBufferMember) \
+		GetSamplerFromHeap(UB_RES_SAFE_TYPE(UniformBufferName, UniformBufferMember), VarName##.BindlessSampler_##UniformBufferMember)
+
+	// Used in standard resource definitions, note that they are all cbuffer members
 	#define UB_RESOURCE_MEMBER_SRV(Type, UBName, Name) \
-		typedef Type SafeType##UBName##_##Name; \
-		UE_BINDLESS_SRV_HEAP_DECL(SafeType##UBName##_##Name); \
-		SafeType##UBName##_##Name GetBindlessResource##UBName##_##Name() { return GetSRVFromHeap(SafeType##UBName##_##Name, UB_CB_PREFIXED_MEMBER_ACCESS(UBName, BindlessSRV_, Name)); } \
-		static const SafeType##UBName##_##Name UBName##_##Name = GetBindlessResource##UBName##_##Name();
+		UB_RES_SAFE_TYPE(UBName, Name) GetBindlessResource##UBName##_##Name() { return UB_CB_LOAD_SRV(UBName, Name); } \
+		static const UB_RES_SAFE_TYPE(UBName, Name) UBName##_##Name = GetBindlessResource##UBName##_##Name()
 
 	#define UB_RESOURCE_MEMBER_UAV(Type, UBName, Name) \
-		typedef Type SafeType##UBName##_##Name; \
-		UE_BINDLESS_UAV_HEAP_DECL(SafeType##UBName##_##Name); \
-		SafeType##UBName##_##Name GetBindlessResource##UBName##_##Name() { return GetUAVFromHeap(SafeType##UBName##_##Name, UB_CB_PREFIXED_MEMBER_ACCESS(UBName, BindlessUAV_, Name)); } \
-		static const SafeType##UBName##_##Name UBName##_##Name = GetBindlessResource##UBName##_##Name();
+		UB_RES_SAFE_TYPE(UBName, Name) GetBindlessResource##UBName##_##Name() { return UB_CB_LOAD_UAV(UBName, Name); } \
+		static const UB_RES_SAFE_TYPE(UBName, Name) UBName##_##Name = GetBindlessResource##UBName##_##Name()
 
 	#define UB_RESOURCE_MEMBER_SAMPLER(Type, UBName, Name) \
-		typedef Type SafeType##UBName##_##Name; \
-		UE_BINDLESS_SAMPLER_HEAP_DECL(SafeType##UBName##_##Name); \
-		SafeType##UBName##_##Name GetBindlessSampler##UBName##_##Name() { return GetSamplerFromHeap(SafeType##UBName##_##Name, UB_CB_PREFIXED_MEMBER_ACCESS(UBName, BindlessSampler_, Name)); } \
-		static const SafeType##UBName##_##Name UBName##_##Name = GetBindlessSampler##UBName##_##Name();
+		UB_RES_SAFE_TYPE(UBName, Name) GetBindlessSampler##UBName##_##Name() { return UB_CB_LOAD_SAMPLER(UBName, Name); } \
+		static const UB_RES_SAFE_TYPE(UBName, Name) UBName##_##Name = GetBindlessSampler##UBName##_##Name()
+
 #else
+	#define UB_RES_DECL_SRV(Type, UniformBufferName, UniformBufferMember)
+	#define UB_RES_DECL_UAV(Type, UniformBufferName, UniformBufferMember)
+	#define UB_RES_DECL_SAMPLER(Type, UniformBufferName, UniformBufferMember)
+
 	#define UB_RESOURCE_MEMBER_SRV(Type, UBName, Name) Type UBName##_##Name
 	#define UB_RESOURCE_MEMBER_UAV(Type, UBName, Name) Type UBName##_##Name
 	#define UB_RESOURCE_MEMBER_SAMPLER(Type, UBName, Name) Type UBName##_##Name
 #endif
 
+//
+// Resource Collection support
+
 #if UE_BINDLESS_ENABLED
 
 typedef ByteAddressBuffer FResourceCollection;
diff --git a/Engine/Shaders/Public/Platform.ush b/Engine/Shaders/Public/Platform.ush
--- a/Engine/Shaders/Public/Platform.ush
+++ b/Engine/Shaders/Public/Platform.ush
@@ -516,7 +516,8 @@
 #define UB_CB_MEMBER_NAME(           UBName,         MemberName) UBName##_##MemberName
 #define UB_CB_PREFIXED_MEMBER_NAME(  UBName, Prefix, MemberName) Prefix##UBName##_##MemberName
 #define UB_CB_DEFINITION_END(        UBName                    ) }
-#define UB_CB_DEFINITION_END_REG(    UBName, Register, Space   ) UB_CB_DEFINITION_END(UBName)
+#define UB_CB_FINALIZE(              UBName                    )
+#define UB_CB_FINALIZE_REG(          UBName, Register, Space   )
 #define UB_CB_MEMBER_ACCESS(         UBName,         MemberName) UBName##_##MemberName
 #define UB_CB_PREFIXED_MEMBER_ACCESS(UBName, Prefix, MemberName) Prefix##UBName##_##MemberName
 
@@ -532,8 +533,9 @@
 	#define UB_CBO_DEFINITION_START_REG(  UBName, Register, Space   ) UB_CBO_DEFINITION_START(UBName)
 	#define UB_CBO_MEMBER_NAME(           UBName,         MemberName) MemberName
 	#define UB_CBO_PREFIXED_MEMBER_NAME(  UBName, Prefix, MemberName) Prefix##MemberName
-	#define UB_CBO_DEFINITION_END(        UBName                    ) }; ConstantBuffer<UB_CBO_TYPE(UBName)> UBName;
-	#define UB_CBO_DEFINITION_END_REG(    UBName, Register, Space   ) }; ConstantBuffer<UB_CBO_TYPE(UBName)> UBName : register(b##Register, space##Space);
+	#define UB_CBO_DEFINITION_END(        UBName                    ) };
+	#define UB_CBO_FINALIZE(              UBName                    ) ConstantBuffer<UB_CBO_TYPE(UBName)> UBName;
+	#define UB_CBO_FINALIZE_REG(          UBName, Register, Space   ) ConstantBuffer<UB_CBO_TYPE(UBName)> UBName : register(b##Register, space##Space);
 	#define UB_CBO_MEMBER_ACCESS(         UBName,         MemberName) UBName##.##MemberName
 	#define UB_CBO_PREFIXED_MEMBER_ACCESS(UBName, Prefix, MemberName) UBName##.##Prefix##MemberName
 
@@ -548,7 +550,8 @@
 	#define UB_CBO_MEMBER_NAME(           UBName,         MemberName) UB_CB_MEMBER_NAME(UBName,MemberName)
 	#define UB_CBO_PREFIXED_MEMBER_NAME(  UBName, Prefix, MemberName) UB_CB_PREFIXED_MEMBER_NAME(UBName,Prefix,MemberName)
 	#define UB_CBO_DEFINITION_END(        UBName                    ) UB_CB_DEFINITION_END(UBName)
-	#define UB_CBO_DEFINITION_END_REG(    UBName, Register, Space   ) UB_CB_DEFINITION_END_REG(UBName,Register,Space)
+	#define UB_CBO_FINALIZE(              UBName                    ) UB_CB_FINALIZE(UBName)
+	#define UB_CBO_FINALIZE_REG(          UBName, Register, Space   ) UB_CB_FINALIZE_REG(UBName,Register,Space)
 	#define UB_CBO_MEMBER_ACCESS(         UBName,         MemberName) UB_CB_MEMBER_ACCESS(UBName,MemberName)
 	#define UB_CBO_PREFIXED_MEMBER_ACCESS(UBName, Prefix, MemberName) UB_CB_PREFIXED_MEMBER_ACCESS(UBName,Prefix,MemberName)
 
diff --git a/Engine/Source/Runtime/RenderCore/Private/ShaderParameterMetadata.cpp b/Engine/Source/Runtime/RenderCore/Private/ShaderParameterMetadata.cpp
--- a/Engine/Source/Runtime/RenderCore/Private/ShaderParameterMetadata.cpp
+++ b/Engine/Source/Runtime/RenderCore/Private/ShaderParameterMetadata.cpp
@@ -589,6 +589,11 @@ void FShaderParametersMetadata::InitializeLayout(FRHIUniformBufferLayoutInitiali
 	{
 		EnumAddFlags(LayoutInitializer.Flags, ERHIUniformBufferFlags::NeedsReflectedMembers);
 	}
+
+	if (EnumHasAnyFlags(UsageFlags, EUsageFlags::BindlessAccessible))
+	{
+		EnumAddFlags(LayoutInitializer.Flags, ERHIUniformBufferFlags::BindlessAccessible);
+	}
 	
 	if (StaticSlotName)
 	{
diff --git a/Engine/Source/Runtime/RenderCore/Private/ShaderParameters.cpp b/Engine/Source/Runtime/RenderCore/Private/ShaderParameters.cpp
--- a/Engine/Source/Runtime/RenderCore/Private/ShaderParameters.cpp
+++ b/Engine/Source/Runtime/RenderCore/Private/ShaderParameters.cpp
@@ -133,11 +133,20 @@ struct FUniformBufferDecl
 	/** Members to place in the constant buffer. */
 	FStringBuilderBase ConstantBufferMembers;
 
+	/** Declarations of resource types. Used for all bindless scenarios. */
+	FStringBuilderBase ResourceDecls;
+
 	/** Members to place in the resource table. */
 	FStringBuilderBase ResourceMembers;
 
 	/** Remappings to apply in the UniformBuffer { } block */
 	FStringBuilderBase Remappings;
+
+	/** Getter functions for bindless member resource access. */
+	FStringBuilderBase BindlessResourceGetters;
+
+	/** Remappings to apply in the UniformBuffer { } block for when bindless UB access is enabled. */
+	FStringBuilderBase BindlessRemappings;
 };
 
 inline EShaderParameterType GetShaderParameterTypeFromMemberResourceType(EUniformBufferBaseType InMemberType)
@@ -208,6 +217,7 @@ static void CreateHLSLUniformBufferStructMembersDeclaration(
 	uint32& HLSLBaseOffset)
 {
 	const TArray<FShaderParametersMetadata::FMember>& StructMembers = UniformBufferStruct.GetMembers();
+	const bool bBindlessAccessible = EnumHasAnyFlags(UniformBufferStruct.GetUsageFlags(), FShaderParametersMetadata::EUsageFlags::BindlessAccessible);
 
 	auto GenerateGlobalMemberName = [&GlobalPrefix](FTempParameterStringBuilder& Builder, const FShaderParametersMetadata::FMember& Member)
 	{
@@ -245,6 +255,42 @@ static void CreateHLSLUniformBufferStructMembersDeclaration(
 		HLSLBaseOffset = AbsoluteMemberOffset + HLSLMemberSize;
 	};
 
+	auto AddStructResource = [&Decl, bBindlessAccessible, &UniformBufferName, &GenerateStructMemberName, &GenerateGlobalMemberName](const FShaderParametersMetadata::FMember& Member, const TCHAR* ShaderType, FStringView GlobalName, EShaderParameterType ParameterType)
+	{
+		if (bBindlessAccessible)
+		{
+			// Converting UniformBuffer.Member into a function call without string reallocations:
+			//   The replacement code goes in *after* preprocessing, so we have to generate an actual function() call. At a minimum we have to cut 2 characters from "UniformBuffer.Member".
+			//   Removing the '.' gives us one character, the second one can be trimmed off of either the uniform buffer name or member name. We have to do this without causing collisions with other uniform buffers or members. 
+			//   If we want to keep things readable and not cut off _# array expansions, we should cut off the last letter of the uniform buffer name.
+
+			const FString TrimmedUBName(FStringView(UniformBufferName).LeftChop(1));
+
+			FTempParameterStringBuilder FullStructMemberName;
+			GenerateStructMemberName(FullStructMemberName, Member);
+
+			FTempParameterStringBuilder GlobalMemberName;
+			GenerateGlobalMemberName(GlobalMemberName, Member);
+
+			const TCHAR* LoadSuffix = GetShaderParameterTypeSuffix(ParameterType);
+
+			// Global getter functions with the shortened name.
+			Decl.BindlessResourceGetters.Appendf(
+				TEXT("UB_RES_SAFE_TYPE(%s,%s) %s%s() { return UB_CBO_LOAD_%s(%s,%s,%s); }\n"),
+				*UniformBufferName, *GlobalMemberName,
+				*TrimmedUBName, *GlobalMemberName,
+				LoadSuffix, *UniformBufferName, *UniformBufferName, *GlobalMemberName
+			);
+
+			// "UniformBuffer.Member" -> "UniformBuffeMember()"
+			Decl.BindlessRemappings.Appendf(
+				TEXT("%s.%s = %s%s();\n"),
+				*UniformBufferName, *FullStructMemberName,
+				*TrimmedUBName, *GlobalMemberName
+			);
+		}
+	};
+
 	auto AddRemapping = [&](const FShaderParametersMetadata::FMember& Member, const FTempParameterStringBuilder& GlobalName, EShaderParameterType ParameterType)
 	{
 		const bool bResource = (ParameterType != EShaderParameterType::LooseData);
@@ -391,9 +437,11 @@ static void CreateHLSLUniformBufferStructMembersDeclaration(
 			const EShaderParameterType ShaderParameterType = GetShaderParameterTypeFromMemberResourceType(Member.GetBaseType());
 			const TCHAR* MemberSuffix = GetShaderParameterTypeSuffix(ShaderParameterType);
 
+			Decl.ResourceDecls.Appendf(TEXT("UB_RES_DECL_%s(%s, %s, %s);\n"), MemberSuffix, Member.GetShaderType(), *UniformBufferName, *GlobalParameterName);
 			Decl.ResourceMembers.Appendf(TEXT("UB_RESOURCE_MEMBER_%s(%s, %s, %s);\n"), MemberSuffix, Member.GetShaderType(), *UniformBufferName, *GlobalParameterName);
 
 			AddRemapping(Member, GlobalParameterName, ShaderParameterType);
+			AddStructResource(Member, Member.GetShaderType(), GlobalParameterName, ShaderParameterType);
 		}
 	}
 }
@@ -408,6 +456,7 @@ static FString CreateHLSLUniformBufferDeclaration(const TCHAR* UniformBufferName
 		uint32 HLSLBaseOffset = 0;
 		CreateHLSLUniformBufferStructMembersDeclaration(UniformBufferStruct, UniformBufferName, TEXT(""), TEXT(""), 0, Decl, HLSLBaseOffset);
 
+		const bool bBindlessAccessible = EnumHasAnyFlags(UniformBufferStruct.GetUsageFlags(), FShaderParametersMetadata::EUsageFlags::BindlessAccessible);
 		const bool bUseRegister = (UniformBufferSBLayout && UniformBufferSBLayout->RegisterSpace > 0);
 
 		// UB_CB_DEFINITION_START and UB_CB_DEFINITION_END both have variants that take the register + space specifiers, so use those when appropriate
@@ -415,25 +464,82 @@ static FString CreateHLSLUniformBufferDeclaration(const TCHAR* UniformBufferName
 			? FString::Printf(TEXT("_REG(%s,%d,%d)"), UniformBufferName, UniformBufferSBLayout->CBVResourceIndex, UniformBufferSBLayout->RegisterSpace)
 			: FString::Printf(TEXT("(%s)"), UniformBufferName);
 
-		FString Result = FString::Printf(
-			TEXT("#pragma once\n")
+		FStringBuilderBase Builder;
+		Builder.Append(TEXT("#pragma once\n"));
+
+		// Type declarations are common between all UB generated code.
+		Builder.Append(Decl.ResourceDecls);
+
+		if (bBindlessAccessible)
+		{
+			Builder.Append(TEXT("#if !(UE_BINDLESS_ENABLED && SUPPORTS_BINDLESS_UNIFORM_BUFFER)\n"));
+		}
+
+		Builder.Appendf(
 			TEXT("UB_CB_DEFINITION_START%s\n")
 			TEXT("%s")
-			TEXT("UB_CB_DEFINITION_END%s\n")
-			TEXT("%s")
+			TEXT("UB_CB_DEFINITION_END(%s)\n")
+			TEXT("UB_CB_FINALIZE%s\n"),
+			*StartEndSuffix,
+			*Decl.ConstantBufferMembers,
+			UniformBufferName,
+			*StartEndSuffix
+		);
+
+		Builder.Append(Decl.ResourceMembers);
+		Builder.Appendf(
 			TEXT("UniformBuffer %s\n")
 			TEXT("{\n")
 			TEXT("%s")
 			TEXT("};\n"),
-			*StartEndSuffix,
-			*Decl.ConstantBufferMembers,
-			*StartEndSuffix,
-			*Decl.ResourceMembers,
 			UniformBufferName,
 			*Decl.Remappings
 		);
 
-		if (Decl.ResourceMembers.Len() == 0)
+		if (bBindlessAccessible)
+		{
+			Builder.Append(TEXT("#else\n"));
+
+			// Force all ConstantBuffer members to CBO syntax.
+			const FString CBOMembers = FString(*Decl.ConstantBufferMembers).Replace(TEXT("UB_CB_"), TEXT("UB_CBO_"), ESearchCase::CaseSensitive);
+
+			Builder.Appendf(
+				TEXT("UB_CBO_DEFINITION_START(%s)\n")
+				TEXT("%s")
+				TEXT("UB_CBO_DEFINITION_END(%s)\n"),
+				UniformBufferName,
+				*CBOMembers,
+				UniformBufferName
+			);
+
+			Builder.Appendf(
+				TEXT("UB_BINDLESS_DECL(UB_CBO_TYPE(%s), %s);\n"),
+				UniformBufferName,
+				UniformBufferName
+			);
+
+			Builder.Appendf(TEXT("static UB_CBO_TYPE(%s) %s;\n"),
+				UniformBufferName, UniformBufferName
+			);
+
+			Builder.Append(Decl.BindlessResourceGetters);
+
+			Builder.Appendf(
+				TEXT("UniformBuffer %s\n")
+				TEXT("{\n")
+				TEXT("%s")
+				TEXT("};\n"),
+				UniformBufferName,
+				*Decl.BindlessRemappings
+			);
+			Builder.Append(TEXT("#endif\n"));
+		}
+
+		Builder.Append(TEXT("\n"));
+
+		FString Result = *Builder;
+
+		if (Decl.ResourceMembers.Len() == 0 && !bBindlessAccessible)
 		{
 			return Result.Replace(TEXT("UB_CB_"), TEXT("UB_CBO_"), ESearchCase::CaseSensitive);
 		}
diff --git a/Engine/Source/Runtime/RenderCore/Public/ShaderParameterMetadata.h b/Engine/Source/Runtime/RenderCore/Public/ShaderParameterMetadata.h
--- a/Engine/Source/Runtime/RenderCore/Public/ShaderParameterMetadata.h
+++ b/Engine/Source/Runtime/RenderCore/Public/ShaderParameterMetadata.h
@@ -164,6 +164,9 @@ public:
 
 		/** Signals that the uniform buffer is manually bound by the pass and should be ignored by the mesh pass processor. */
 		ManuallyBoundByPass = 1 << 3,
+
+		/** When bindless is enabled, allow this uniform buffer to be accessed via a bindless constant buffer. */
+		BindlessAccessible = 1 << 4,
 	};
 
 	/** Shader binding name of the uniform buffer that contains the root shader parameters. */
