[shaders] remove separate shader code library loading path in LoadExistingShaderCodeLibrary, instead using the more optimal (and complete) version in AddShaderCodeLibraryByName. Fixes linked JIRA; this occurred due to this separate loading path not handling the load of shader type info.

#Synced-CL 48460480

diff --git a/Engine/Source/Runtime/RenderCore/Private/ShaderCodeLibrary.cpp b/Engine/Source/Runtime/RenderCore/Private/ShaderCodeLibrary.cpp
--- a/Engine/Source/Runtime/RenderCore/Private/ShaderCodeLibrary.cpp
+++ b/Engine/Source/Runtime/RenderCore/Private/ShaderCodeLibrary.cpp
@@ -2065,7 +2065,15 @@ struct FEditorShaderCodeArchive
 		return ShaderMapIndex;
 	}
 
-	void ReadShaderCodeBlock(FArchive& Ar, int64 StartOffset, const TArrayView<FShaderCodeReadInfo>& BlockReadInfo)
+	enum class EShaderCodeReadStatus : uint8
+	{
+		OK,
+		FileNotFound,
+		MalformedArchive,
+		VersionMismatch,
+	};
+
+	EShaderCodeReadStatus ReadShaderCodeBlock(FArchive& Ar, int64 StartOffset, const TArrayView<FShaderCodeReadInfo>& BlockReadInfo)
 	{
 		if (BlockReadInfo.Num() > 0)
 		{
@@ -2086,11 +2094,18 @@ struct FEditorShaderCodeArchive
 				FUniqueBuffer Code = FUniqueBuffer::Clone(static_cast<const uint8*>(BlockOfShaderCode.GetData()) + Block.Offset - First.Offset, Block.Size);
 				ShaderCode[Block.ShaderCodeIndex] = Code.MoveToShared();
 			}
+
+			if (Ar.GetError())
+			{
+				return EShaderCodeReadStatus::MalformedArchive;
+			}
 		}
+		return EShaderCodeReadStatus::OK;
 	}
 
-	void ReadShaderCode(FArchive& Ar, TArray<FShaderCodeReadInfo>& ReadInfo)
+	EShaderCodeReadStatus ReadShaderCode(FArchive& Ar, TArray<FShaderCodeReadInfo>& ReadInfo)
 	{
+		EShaderCodeReadStatus Status = EShaderCodeReadStatus::OK;
 		if (ReadInfo.Num() > 0)
 		{
 			// Sort by offset
@@ -2112,9 +2127,13 @@ struct FEditorShaderCodeArchive
 				int64 PotentialSize = Next.Offset - ReadInfo[BlockStartIndex].Offset + Next.Size;
 				if (PotentialSize > DesiredBlockSize)
 				{
-					ReadShaderCodeBlock(Ar, PrevCookedShadersCodeStart, MakeArrayView(&ReadInfo[BlockStartIndex], IndexCount));
+					Status = ReadShaderCodeBlock(Ar, PrevCookedShadersCodeStart, MakeArrayView(&ReadInfo[BlockStartIndex], IndexCount));
 					BlockStartIndex = i;
 					IndexCount = 1;
+					if (Status != EShaderCodeReadStatus::OK)
+					{
+						return Status;
+					}
 				}
 				else
 				{
@@ -2123,67 +2142,18 @@ struct FEditorShaderCodeArchive
 			}
 
 			// Read last block
-			ReadShaderCodeBlock(Ar, PrevCookedShadersCodeStart, MakeArrayView(&ReadInfo[BlockStartIndex], IndexCount));
+			Status = ReadShaderCodeBlock(Ar, PrevCookedShadersCodeStart, MakeArrayView(&ReadInfo[BlockStartIndex], IndexCount));
 		}
+		return Status;
 	}
 
-	bool LoadExistingShaderCodeLibrary(FString const& MetaDataDir)
-	{
-		FString IntermediateFormatPath = GetCodeArchiveFilename(MetaDataDir / TEXT("ShaderLibrarySource"), LibraryName, FormatName);
-		FArchive* PrevCookedAr = CreateShaderFileReader(*IntermediateFormatPath);
-		bool bOK = true;
-		if (PrevCookedAr)
-		{
-			uint32 ArchiveVersion = 0;
-			*PrevCookedAr << ArchiveVersion;
-			if (ArchiveVersion == GShaderCodeArchiveVersion)
-			{
-				// Read shader library
-				*PrevCookedAr << SerializedShaders;
-				MarkAllShaderMapsDirty();
-
-				ShaderCode.AddDefaulted(SerializedShaders.ShaderEntries.Num());
-				for (int32 Index = 0; Index < ShaderCode.Num(); ++Index)
-				{
-					const FShaderCodeEntry& Entry = SerializedShaders.ShaderEntries[Index];
-					check(Entry.Size > 0);
-					FUniqueBuffer Code = FUniqueBuffer::Alloc(Entry.Size);
-					PrevCookedAr->Serialize(Code.GetData(), Entry.Size);
 
-					bOK = !PrevCookedAr->GetError();
-					if (!bOK)
-					{
-						UE_LOG(LogShaderLibrary, Error, TEXT("Failed to deserialize shader code for %s from %s"), *SerializedShaders.ShaderHashes[Index].ToString(), *IntermediateFormatPath);
-						break;
-					}
-
-					ShaderCode[Index] = Code.MoveToShared();
-				}
-			}
-			else
-			{
-				bOK = false;
-				UE_LOG(LogShaderLibrary, Warning, TEXT("Failed to deserialize shader code from %s because the archive format %u is incompatible with the current version %u"), *IntermediateFormatPath, ArchiveVersion, GShaderCodeArchiveVersion);
-			}
-			
-			PrevCookedAr->Close();
-			delete PrevCookedAr;
-		}
-		else
-		{
-			bOK = false;
-			UE_LOG(LogShaderLibrary, Error, TEXT("Failed to open shader code library from %s"), *IntermediateFormatPath);
-		}
-		
-		return bOK;
-	}
-
-	void AddShaderCodeLibraryByName(const FString& BaseDir, const FString& InLibraryName, FShaderCodeStats& CodeStats)
+	EShaderCodeReadStatus AddShaderCodeLibraryByName(const FString& BaseDir, const FString& InLibraryName, FShaderCodeStats& CodeStats)
 	{
-		EShaderPlatform Platform = ShaderFormatToLegacyShaderPlatform(FormatName);
 		FString CodeArchiveFileName = GetCodeArchiveFilename(BaseDir, InLibraryName, FormatName);
 		if (FArchive* PrevCookedAr = CreateShaderFileReader(*CodeArchiveFileName))
 		{
+			EShaderCodeReadStatus Status;
 			uint32 Version = 0;
 			*PrevCookedAr << Version;
 
@@ -2239,11 +2209,45 @@ struct FEditorShaderCodeArchive
 				{
 					PrepareAddShaderCode(CodeStats, PrevShaderMapIndex, PrevCookedShaders, ReadInfo);
 				}
-				ReadShaderCode(*PrevCookedAr, ReadInfo);
+				Status = ReadShaderCode(*PrevCookedAr, ReadInfo);
+			}
+			else
+			{
+				Status = EShaderCodeReadStatus::VersionMismatch;
 			}
 
 			PrevCookedAr->Close();
 			delete PrevCookedAr;
+			return Status;
+		}
+		return EShaderCodeReadStatus::FileNotFound;
+	}
+
+	bool LoadExistingShaderCodeLibrary(FString const& MetaDataDir)
+	{
+		FString BaseDir = MetaDataDir / TEXT("ShaderLibrarySource");
+		FShaderCodeStats CodeStats;
+		EShaderCodeReadStatus Status = AddShaderCodeLibraryByName(BaseDir, LibraryName, CodeStats);
+		switch (Status)
+		{
+		case EShaderCodeReadStatus::FileNotFound:
+			{
+				UE_LOG(LogShaderLibrary, Error, TEXT("Failed to load shader code library from %s; file not found"), *GetCodeArchiveFilename(BaseDir, LibraryName, FormatName));
+				return false;
+			}
+		case EShaderCodeReadStatus::MalformedArchive:
+			{
+				UE_LOG(LogShaderLibrary, Error, TEXT("Failed to deserialized code from library %s"), *GetCodeArchiveFilename(BaseDir, LibraryName, FormatName));
+				return false;
+			}
+		case EShaderCodeReadStatus::VersionMismatch:
+			{
+				UE_LOG(LogShaderLibrary, Error, TEXT("Failed to deserialize shader code from %s because the archive format is incompatible with the current version %u"), *GetCodeArchiveFilename(BaseDir, LibraryName, FormatName), GShaderCodeArchiveVersion);
+				return false;
+			}
+		case EShaderCodeReadStatus::OK:
+		default:
+			return true;
 		}
 	}
 
