Removing deprecated code from RenderCore.

#Synced-CL 48482604

diff --git a/Engine/Source/Runtime/RenderCore/Private/OneColorShader.cpp b/Engine/Source/Runtime/RenderCore/Private/OneColorShader.cpp
--- a/Engine/Source/Runtime/RenderCore/Private/OneColorShader.cpp
+++ b/Engine/Source/Runtime/RenderCore/Private/OneColorShader.cpp
@@ -30,13 +30,6 @@ void FOneColorPS::FillParameters(FParameters& Parameters, const FLinearColor* Co
 	}
 }
 
-void FOneColorPS::SetColors(FRHICommandList& RHICmdList, const TShaderMapRef<FOneColorPS>& Shader, const FLinearColor* Colors, int32 NumColors)
-{
-	FParameters Parameters;
-	Shader->FillParameters(Parameters, Colors, NumColors);
-	SetShaderParameters(RHICmdList, Shader, Shader.GetPixelShader(), Parameters);
-}
-
 IMPLEMENT_GLOBAL_SHADER(FOneColorPS,"/Engine/Private/OneColorShader.usf","MainPixelShader",SF_Pixel);
 
 TOneColorPixelShaderMRT::TOneColorPixelShaderMRT() = default;
@@ -83,13 +76,6 @@ void TOneColorPixelShaderMRT::ModifyCompilationEnvironment(const FGlobalShaderPe
 	}
 }
 
-void TOneColorPixelShaderMRT::SetColors(FRHICommandList& RHICmdList, const TShaderMapRef<TOneColorPixelShaderMRT>& Shader, const FLinearColor* Colors, int32 NumColors)
-{
-	FParameters Parameters;
-	Shader->FillParameters(Parameters, Colors, NumColors);
-	SetShaderParameters(RHICmdList, Shader, Shader.GetPixelShader(), Parameters);
-}
-
 // Compiling a version for every number of MRT's
 // On AMD PC hardware, outputting to a color index in the shader without a matching render target set has a significant performance hit
 IMPLEMENT_GLOBAL_SHADER(TOneColorPixelShaderMRT,"/Engine/Private/OneColorShader.usf","MainPixelShaderMRT",SF_Pixel);
diff --git a/Engine/Source/Runtime/RenderCore/Private/RenderingThread.cpp b/Engine/Source/Runtime/RenderCore/Private/RenderingThread.cpp
--- a/Engine/Source/Runtime/RenderCore/Private/RenderingThread.cpp
+++ b/Engine/Source/Runtime/RenderCore/Private/RenderingThread.cpp
@@ -1315,20 +1315,6 @@ void FlushRenderingCommands()
 	FCoreRenderDelegates::OnFlushRenderingCommandsEnd.Broadcast();
 }
 
-void FlushPendingDeleteRHIResources_GameThread()
-{
-	ENQUEUE_RENDER_COMMAND(FlushPendingDeleteRHIResources)(
-		[](FRHICommandListImmediate& RHICmdList)
-	{
-		RHICmdList.ImmediateFlush(EImmediateFlushType::FlushRHIThreadFlushResources);
-	});
-}
-
-void FlushPendingDeleteRHIResources_RenderThread()
-{
-	FRHICommandListImmediate::Get().ImmediateFlush(EImmediateFlushType::FlushRHIThreadFlushResources);
-}
-
 FRHICommandListImmediate& GetImmediateCommandList_ForRenderCommand()
 {
 	return FRHICommandListExecutor::GetImmediateCommandList();
diff --git a/Engine/Source/Runtime/RenderCore/Private/ShaderCodeLibrary.cpp b/Engine/Source/Runtime/RenderCore/Private/ShaderCodeLibrary.cpp
--- a/Engine/Source/Runtime/RenderCore/Private/ShaderCodeLibrary.cpp
+++ b/Engine/Source/Runtime/RenderCore/Private/ShaderCodeLibrary.cpp
@@ -4513,11 +4513,6 @@ void FShaderLibraryCooker::DumpShaderTypeStats(const FString& DebugInfoDir, cons
 	}
 }
 
-void FShaderLibraryCooker::DumpShaderCodeStats()
-{
-	// stub - deprecated since UE 5.6
-}
-
 bool FShaderLibraryCooker::CreatePatchLibrary(TArray<FString> const& OldMetaDataDirs, FString const& NewMetaDataDir, FString const& OutDir, bool bNativeFormat, bool bNeedsDeterministicOrder)
 {
 	TMap<FName, TSet<FString>> FormatLibraryMap;
diff --git a/Engine/Source/Runtime/RenderCore/Public/DynamicBufferAllocator.h b/Engine/Source/Runtime/RenderCore/Public/DynamicBufferAllocator.h
--- a/Engine/Source/Runtime/RenderCore/Public/DynamicBufferAllocator.h
+++ b/Engine/Source/Runtime/RenderCore/Public/DynamicBufferAllocator.h
@@ -22,12 +22,6 @@ struct FDynamicAllocReadBuffer : public FDynamicReadBuffer
 	int32 NumFramesUnused = 0;
 
 	TArray<FShaderResourceViewRHIRef> SubAllocations;
-	
-	UE_DEPRECATED(5.3, "Lock now requires a command list.")
-	void Lock() { Lock(FRHICommandListImmediate::Get()); }
-	
-	UE_DEPRECATED(5.3, "Lock now requires a command list.")
-	void Unlock() { Unlock(FRHICommandListImmediate::Get()); }
 
 	void Lock(FRHICommandListBase& RHICmdList)
 	{
@@ -96,9 +90,6 @@ public:
 	*/
 	RENDERCORE_API void Commit(FRHICommandListImmediate& RHICmdList);
 
-	UE_DEPRECATED(5.3, "Commit now requires a command list.")
-	void Commit() { Commit(FRHICommandListImmediate::Get()); }
-
 	/** Returns true if log statements should be made because we exceeded GMaxVertexBytesAllocatedPerFrame */
 	RENDERCORE_API bool IsRenderAlarmLoggingEnabled() const;
 
diff --git a/Engine/Source/Runtime/RenderCore/Public/GBufferInfo.h b/Engine/Source/Runtime/RenderCore/Public/GBufferInfo.h
--- a/Engine/Source/Runtime/RenderCore/Public/GBufferInfo.h
+++ b/Engine/Source/Runtime/RenderCore/Public/GBufferInfo.h
@@ -323,12 +323,6 @@ int32 RENDERCORE_API FindGBufferTargetByName(const FGBufferInfo& GBufferInfo, co
 
 FGBufferBinding RENDERCORE_API FindGBufferBindingByName(const FGBufferInfo& GBufferInfo, const FString& Name, EShaderPlatform ShaderPlatform);
 
-UE_DEPRECATED(5.4, "Please use the overload which takes a shader platform parameter")
-inline FGBufferBinding FindGBufferBindingByName(const FGBufferInfo& GBufferInfo, const FString& Name)
-{
-	return FindGBufferBindingByName(GBufferInfo, Name, GMaxRHIShaderPlatform);
-}
-
 FGBufferInfo RENDERCORE_API FetchFullGBufferInfo(const FGBufferParams& Params);
 
 bool RENDERCORE_API IsGBufferInfoEqual(const FGBufferInfo& Lhs, const FGBufferInfo& Rhs);
diff --git a/Engine/Source/Runtime/RenderCore/Public/GlobalRenderResources.h b/Engine/Source/Runtime/RenderCore/Public/GlobalRenderResources.h
--- a/Engine/Source/Runtime/RenderCore/Public/GlobalRenderResources.h
+++ b/Engine/Source/Runtime/RenderCore/Public/GlobalRenderResources.h
@@ -248,9 +248,6 @@ public:
 	 */
 	RENDERCORE_API void Commit();
 
-	UE_DEPRECATED(5.4, "Use GlobalDynamicBuffer::GarbageCollect instead.")
-	void GarbageCollect() {}
-
 	/** Returns true if log statements should be made because we exceeded GMaxVertexBytesAllocatedPerFrame */
 	RENDERCORE_API bool IsRenderAlarmLoggingEnabled() const;
 
diff --git a/Engine/Source/Runtime/RenderCore/Public/MediaShaders.h b/Engine/Source/Runtime/RenderCore/Public/MediaShaders.h
--- a/Engine/Source/Runtime/RenderCore/Public/MediaShaders.h
+++ b/Engine/Source/Runtime/RenderCore/Public/MediaShaders.h
@@ -126,12 +126,6 @@ inline FBufferRHIRef CreateTempMediaVertexBuffer(FRHICommandListBase& RHICmdList
 	return Vertices.Finalize();
 }
 
-UE_DEPRECATED(5.6, "CreateTempMediaVertexBuffer requires a Command List argument")
-inline FBufferRHIRef CreateTempMediaVertexBuffer(float ULeft = 0.0f, float URight = 1.0f, float VTop = 0.0f, float VBottom = 1.0f)
-{
-	return CreateTempMediaVertexBuffer(FRHICommandListImmediate::Get(), ULeft, URight, VTop, VBottom);
-}
-
 /**
  * The simple element vertex declaration resource type.
  */
diff --git a/Engine/Source/Runtime/RenderCore/Public/OneColorShader.h b/Engine/Source/Runtime/RenderCore/Public/OneColorShader.h
--- a/Engine/Source/Runtime/RenderCore/Public/OneColorShader.h
+++ b/Engine/Source/Runtime/RenderCore/Public/OneColorShader.h
@@ -87,9 +87,6 @@ public:
 	END_SHADER_PARAMETER_STRUCT()
 
 	RENDERCORE_API void FillParameters(FParameters& Parameters, const FLinearColor* Colors, int32 NumColors);
-
-	UE_DEPRECATED(5.3, "FParameters and FillParameters should be used instead of this helper.")
-	RENDERCORE_API static void SetColors(FRHICommandList& RHICmdList, const TShaderMapRef<FOneColorPS>& Shader, const FLinearColor* Colors, int32 NumColors);
 };
 
 /**
@@ -109,9 +106,6 @@ public:
 
 	static bool ShouldCompilePermutation(const FGlobalShaderPermutationParameters& Parameters);
 	static void ModifyCompilationEnvironment(const FGlobalShaderPermutationParameters& Parameters, FShaderCompilerEnvironment& OutEnvironment);
-
-	UE_DEPRECATED(5.3, "FParameters and FillParameters should be used instead of this helper.")
-	RENDERCORE_API static void SetColors(FRHICommandList& RHICmdList, const TShaderMapRef<TOneColorPixelShaderMRT>& Shader, const FLinearColor* Colors, int32 NumColors);
 };
 
 /**
diff --git a/Engine/Source/Runtime/RenderCore/Public/PixelShaderUtils.h b/Engine/Source/Runtime/RenderCore/Public/PixelShaderUtils.h
--- a/Engine/Source/Runtime/RenderCore/Public/PixelShaderUtils.h
+++ b/Engine/Source/Runtime/RenderCore/Public/PixelShaderUtils.h
@@ -214,32 +214,6 @@ struct FPixelShaderUtils
 		});
 	}
 
-	/** Rect based pixel shader pass. */
-	
-	template<typename TPixelShaderClass, typename TPassParameters>
-	UE_DEPRECATED(5.5, "Use the other prototype of AddRasterizeToRectsPass that takes a ERDGPassFlags in parameter and use ERDGPassFlags::SkipRenderPass in lieu of bSkipRenderPass == true")
-	static inline void AddRasterizeToRectsPass(
-		FRDGBuilder& GraphBuilder,
-		const FGlobalShaderMap* GlobalShaderMap,
-		FRDGEventName&& PassName,
-		const TShaderRef<TPixelShaderClass>& PixelShader,
-		TPassParameters* Parameters,
-		FIntPoint ViewportSize,
-		FRDGBufferSRVRef RectCoordBufferSRV,
-		uint32 NumRects,
-		FRHIBlendState* BlendState,
-		FRHIRasterizerState* RasterizerState,
-		FRHIDepthStencilState* DepthStencilState,
-		uint32 StencilRef,
-		FIntPoint TextureSize,
-		FRDGBufferSRVRef RectUVBufferSRV,
-		uint32 DownsampleFactor,
-		const bool bSkipRenderPass)
-	{
-		AddRasterizeToRectsPass(GraphBuilder, GlobalShaderMap, PassName, PixelShader, Parameters, ViewportSize, RectCoordBufferSRV, NumRects, BlendState, RasterizerState, DepthStencilState, 
-			StencilRef, TextureSize, RectUVBufferSRV, DownsampleFactor, bSkipRenderPass ? ERDGPassFlags::SkipRenderPass : ERDGPassFlags::None);
-	}
-
 	/** Rect based pixel shader pass. */
 	template<typename TPixelShaderClass, typename TPassParameters>
 	static inline void AddRasterizeToRectsPass(
diff --git a/Engine/Source/Runtime/RenderCore/Public/RenderGraphBuilder.h b/Engine/Source/Runtime/RenderCore/Public/RenderGraphBuilder.h
--- a/Engine/Source/Runtime/RenderCore/Public/RenderGraphBuilder.h
+++ b/Engine/Source/Runtime/RenderCore/Public/RenderGraphBuilder.h
@@ -245,10 +245,6 @@ public:
 	/** Adds a user-defined dependency between two passes. This can be used to fine-tune async compute overlap by forcing a sync point. */
 	RENDERCORE_API void AddPassDependency(FRDGPass* Producer, FRDGPass* Consumer);
 
-	/** Sets the current command list stat for all subsequent passes. */
-	UE_DEPRECATED(5.5, "SetCommandListStat is deprecated. The underlying stats have been removed. Consider marking up rendering code with RDG event scopes.")
-	inline void SetCommandListStat(TStatId StatId) {}
-
 	/** A hint to the builder to flush work to the RHI thread after the last queued pass on the execution timeline. */
 	void AddDispatchHint();
 
@@ -448,18 +444,6 @@ public:
 	}
 #endif
 
-	UE_DEPRECATED(5.5, "This path is no longer supported.")
-	static void DumpDraw(const FRDGEventName& DrawEventName) {}
-
-	UE_DEPRECATED(5.5, "This path is no longer supported.")
-	static bool IsDumpingDraws() { return false; }
-
-	UE_DEPRECATED(5.6, "RemoveUnusedTextureWarning is no longer necessary.")
-	void RemoveUnusedTextureWarning(FRDGTextureRef Texture) {}
-
-	UE_DEPRECATED(5.6, "RemoveUnusedBufferWarning is no longer necessary.")
-	void RemoveUnusedBufferWarning(FRDGBufferRef Buffer) {}
-
 private:
 	static const char* const kDefaultUnaccountedCSVStat;
 
diff --git a/Engine/Source/Runtime/RenderCore/Public/RenderGraphDefinitions.h b/Engine/Source/Runtime/RenderCore/Public/RenderGraphDefinitions.h
--- a/Engine/Source/Runtime/RenderCore/Public/RenderGraphDefinitions.h
+++ b/Engine/Source/Runtime/RenderCore/Public/RenderGraphDefinitions.h
@@ -118,8 +118,6 @@ enum class ERDGBuilderFlags
 	ParallelExecute = 1 << 2,
 
 	Parallel = ParallelSetup | ParallelCompile | ParallelExecute,
-
-	AllowParallelExecute UE_DEPRECATED(5.5, "Use ERDDGBuilderFlags::Parallel instead.") = Parallel,
 };
 ENUM_CLASS_FLAGS(ERDGBuilderFlags);
 
diff --git a/Engine/Source/Runtime/RenderCore/Public/RenderGraphResources.h b/Engine/Source/Runtime/RenderCore/Public/RenderGraphResources.h
--- a/Engine/Source/Runtime/RenderCore/Public/RenderGraphResources.h
+++ b/Engine/Source/Runtime/RenderCore/Public/RenderGraphResources.h
@@ -548,12 +548,6 @@ public:
 	/** Finds a SRV matching the descriptor in the cache or creates a new one and updates the cache. */
 	inline FRHIShaderResourceView* GetOrCreateSRV(FRHICommandListBase& RHICmdList, const FRHITextureSRVCreateInfo& SRVDesc) { return ViewCache.GetOrCreateSRV(RHICmdList, Texture, SRVDesc); }
 
-	UE_DEPRECATED(5.3, "GetOrCreateUAV now requires a command list.")
-	inline FRHIUnorderedAccessView* GetOrCreateUAV(const FRHITextureUAVCreateInfo& UAVDesc) { return ViewCache.GetOrCreateUAV(FRHICommandListImmediate::Get(), Texture, UAVDesc); }
-
-	UE_DEPRECATED(5.3, "GetOrCreateSRV now requires a command list.")
-	inline FRHIShaderResourceView* GetOrCreateSRV(const FRHITextureSRVCreateInfo& SRVDesc) { return ViewCache.GetOrCreateSRV(FRHICommandListImmediate::Get(), Texture, SRVDesc); }
-
 	inline FRHITexture* GetRHI() const { return Texture; }
 
 private:
@@ -1219,12 +1213,6 @@ public:
 	/** Finds a SRV matching the descriptor in the cache or creates a new one and updates the cache. */
 	inline FRHIShaderResourceView* GetOrCreateSRV(FRHICommandListBase& RHICmdList, const FRHIBufferSRVCreateInfo& SRVDesc) { return ViewCache.GetOrCreateSRV(RHICmdList, Buffer, SRVDesc); }
 
-	UE_DEPRECATED(5.3, "GetOrCreateUAV now requires a command list.")
-	inline FRHIUnorderedAccessView* GetOrCreateUAV(const FRHIBufferUAVCreateInfo& UAVDesc) { return ViewCache.GetOrCreateUAV(FRHICommandListImmediate::Get(), Buffer, UAVDesc); }
-
-	UE_DEPRECATED(5.3, "GetOrCreateSRV now requires a command list.")
-	inline FRHIShaderResourceView* GetOrCreateSRV(const FRHIBufferSRVCreateInfo& SRVDesc) { return ViewCache.GetOrCreateSRV(FRHICommandListImmediate::Get(), Buffer, SRVDesc); }
-
 	/** Returns the RHI buffer. */
 	inline FRHIBuffer* GetRHI() const { return Buffer; }
 
@@ -1235,12 +1223,6 @@ public:
 		return CachedSRV;
 	}
 
-	UE_DEPRECATED(5.3, "GetSRV now requires a command list.")
-	inline FRHIShaderResourceView* GetSRV(const FRHIBufferSRVCreateInfo& SRVDesc)
-	{
-		return GetOrCreateSRV(FRHICommandListImmediate::Get(), SRVDesc);
-	}
-
 	inline FRHIShaderResourceView* GetSRV(FRHICommandListBase& RHICmdList, const FRHIBufferSRVCreateInfo& SRVDesc)
 	{
 		return GetOrCreateSRV(RHICmdList, SRVDesc);
diff --git a/Engine/Source/Runtime/RenderCore/Public/RenderGraphUtils.h b/Engine/Source/Runtime/RenderCore/Public/RenderGraphUtils.h
--- a/Engine/Source/Runtime/RenderCore/Public/RenderGraphUtils.h
+++ b/Engine/Source/Runtime/RenderCore/Public/RenderGraphUtils.h
@@ -328,16 +328,6 @@ RENDERCORE_API FRDGTextureMSAA CreateTextureMSAA(
 	const TCHAR* NameMultisampled, const TCHAR* NameResolved,
 	ETextureCreateFlags ResolveFlagsToAdd = TexCreate_None);
 
-UE_DEPRECATED(5.3, "CreateTextureMSAA with one name is deprecated, prease provide separate name for the multisampled texture")
-inline FRDGTextureMSAA CreateTextureMSAA(
-	FRDGBuilder& GraphBuilder,
-	FRDGTextureDesc Desc,
-	const TCHAR* Name,
-	ETextureCreateFlags ResolveFlagsToAdd = TexCreate_None)
-{
-	return CreateTextureMSAA(GraphBuilder, Desc, Name, Name, ResolveFlagsToAdd);
-}
-
 /** All utils for compute shaders.
  */
 namespace FComputeShaderUtils
diff --git a/Engine/Source/Runtime/RenderCore/Public/RenderingThread.h b/Engine/Source/Runtime/RenderCore/Public/RenderingThread.h
--- a/Engine/Source/Runtime/RenderCore/Public/RenderingThread.h
+++ b/Engine/Source/Runtime/RenderCore/Public/RenderingThread.h
@@ -109,11 +109,6 @@ extern RENDERCORE_API void AdvanceRenderingThreadStatsGT( bool bDiscardCallstack
  */
 extern RENDERCORE_API void FlushRenderingCommands();
 
-UE_DEPRECATED(5.6, "FlushPendingDeleteRHIResources_GameThread is deprecated. Enqueue a render command that calls ImmediateFlush(EImmediateFlushType::FlushRHIThreadFlushResources) on the immediate RHI command list instead.")
-extern RENDERCORE_API void FlushPendingDeleteRHIResources_GameThread();
-UE_DEPRECATED(5.6, "FlushPendingDeleteRHIResources_RenderThread is deprecated. Call ImmediateFlush(EImmediateFlushType::FlushRHIThreadFlushResources) on the immediate RHI command list instead.")
-extern RENDERCORE_API void FlushPendingDeleteRHIResources_RenderThread();
-
 extern RENDERCORE_API void StartRenderCommandFenceBundler();
 extern RENDERCORE_API void StopRenderCommandFenceBundler();
 
@@ -1343,78 +1338,6 @@ extern RENDERCORE_API void UnregisterStopRenderingThreadDelegate(FDelegateHandle
 ///////////////////////////////////////////////////////////////////////////////
 // Deprecated Types
 
-class UE_DEPRECATED(5.6, "FRenderThreadScope is no longer used.") FRenderThreadScope
-{
-	typedef TFunction<void(FRHICommandListImmediate&)> RenderCommandFunction;
-	typedef TArray<RenderCommandFunction> RenderCommandFunctionArray;
-public:
-	FRenderThreadScope()
-	{
-		RenderCommands = new RenderCommandFunctionArray;
-	}
-
-	~FRenderThreadScope()
-	{
-		RenderCommandFunctionArray* RenderCommandArray = RenderCommands;
-
-		ENQUEUE_RENDER_COMMAND(DispatchScopeCommands)(
-			[RenderCommandArray](FRHICommandListImmediate& RHICmdList)
-		{
-			for(int32 Index = 0; Index < RenderCommandArray->Num(); Index++)
-			{
-				(*RenderCommandArray)[Index](RHICmdList);
-			}
-
-			delete RenderCommandArray;
-		});
-	}
-
-	void EnqueueRenderCommand(RenderCommandFunction&& Lambda)
-	{
-		RenderCommands->Add(MoveTemp(Lambda));
-	}
-
-private:
-	RenderCommandFunctionArray* RenderCommands;
-};
-
-class UE_DEPRECATED(5.6, "FRenderCommand is no longer used") FRenderCommand
-{
-public:
-	static ENamedThreads::Type GetDesiredThread()
-	{
-		check(!GIsThreadedRendering || ENamedThreads::GetRenderThread() != ENamedThreads::GameThread);
-		return ENamedThreads::GetRenderThread();
-	}
-
-	static ESubsequentsMode::Type GetSubsequentsMode()
-	{
-		return ESubsequentsMode::FireAndForget;
-	}
-};
-
-template <typename TagType, typename LambdaType>
-class UE_DEPRECATED(5.6, "TEnqueueUniqueRenderCommandType is no longer used.") TEnqueueUniqueRenderCommandType : public FRenderCommand
-{
-public:
-	TEnqueueUniqueRenderCommandType(LambdaType&& InLambda) : Lambda(Forward<LambdaType>(InLambda)) {}
-
-	void DoTask(ENamedThreads::Type CurrentThread, const FGraphEventRef& MyCompletionGraphEvent)
-	{
-		TRACE_CPUPROFILER_EVENT_SCOPE_ON_CHANNEL_STR(TagType::GetName(), RenderCommandsChannel);
-		FRHICommandListImmediate& RHICmdList = GetImmediateCommandList_ForRenderCommand();
-		Lambda(RHICmdList);
-	}
-
-	inline TStatId GetStatId() const
-	{
-		return TagType::GetStatId();
-	}
-
-private:
-	LambdaType Lambda;
-};
-
 ///////////////////////////////////////////////////////////////////////////////
 
 #undef UE_API
\ No newline at end of file
diff --git a/Engine/Source/Runtime/RenderCore/Public/Shader.h b/Engine/Source/Runtime/RenderCore/Public/Shader.h
--- a/Engine/Source/Runtime/RenderCore/Public/Shader.h
+++ b/Engine/Source/Runtime/RenderCore/Public/Shader.h
@@ -422,23 +422,6 @@ public:
 	RENDERCORE_API void BeginCreateAllShaders();
 
 #if RHI_RAYTRACING
-
-	UE_DEPRECATED(5.6, "GetRayTracingHitGroupLibrary needs current EShaderPlatform")
-	static void GetRayTracingHitGroupLibrary(TArray<FRHIRayTracingShader*>& RayTracingHitGroupShaders, FRHIRayTracingShader* DefaultShader)
-	{
-		GetRayTracingHitGroupLibrary(GMaxRHIShaderPlatform, RayTracingHitGroupShaders, DefaultShader);
-	}
-	UE_DEPRECATED(5.6, "GetRayTracingCallableShaderLibrary needs current EShaderPlatform")
-	static void GetRayTracingCallableShaderLibrary(TArray<FRHIRayTracingShader*>& RayTracingCallableShaders, FRHIRayTracingShader* DefaultShader)
-	{
-		GetRayTracingCallableShaderLibrary(GMaxRHIShaderPlatform, RayTracingCallableShaders, DefaultShader);
-	}
-	UE_DEPRECATED(5.6, "GetRayTracingMissShaderLibrary needs current EShaderPlatform")
-	static void GetRayTracingMissShaderLibrary(TArray<FRHIRayTracingShader*>& RayTracingMissShaders, FRHIRayTracingShader* DefaultShader)
-	{
-		GetRayTracingMissShaderLibrary(GMaxRHIShaderPlatform, RayTracingMissShaders, DefaultShader);
-	}
-
 	static RENDERCORE_API void GetRayTracingHitGroupLibrary(EShaderPlatform ShaderPlatform, TArray<FRHIRayTracingShader*>& RayTracingHitGroupShaders, FRHIRayTracingShader* DefaultShader);
 	static RENDERCORE_API void GetRayTracingCallableShaderLibrary(EShaderPlatform ShaderPlatform, TArray<FRHIRayTracingShader*>& RayTracingCallableShaders, FRHIRayTracingShader* DefaultShader);
 	static RENDERCORE_API void GetRayTracingMissShaderLibrary(EShaderPlatform ShaderPlatform, TArray<FRHIRayTracingShader*>& RayTracingMissShaders, FRHIRayTracingShader* DefaultShader);
@@ -523,9 +506,6 @@ public:
 #if WITH_EDITORONLY_DATA
 	struct FShaderEditorOnlyDataEntry
 	{
-		UE_DEPRECATED(5.6, "Use SymbolBuffer")
-		TArray<uint8> PlatformDebugData;
-
 		// Sets the symbols for this entry to the given buffer if either unset or the hash of the new buffer is "less" than the hash of the stored buffer.
 		// This is done purely for determinism's sake (we always want to store the same copy of symbols in the editoronly data for any given shader bytecode).
 		void ConditionalSetSymbolBuffer(FCompressedBuffer InSymbols);
@@ -581,9 +561,6 @@ public:
 
 	RENDERCORE_API uint32 GetSizeBytes() const;
 
-	UE_DEPRECATED(5.6, "AddShaderCompilerOutput now accepts an FShaderCompileJobKey instead of DebugName")
-	inline void AddShaderCompilerOutput(const FShaderCompilerOutput& Output, const FString& DebugName = FString(), FString DebugInfo = FString()) {}
-
 	RENDERCORE_API void AddShaderCompilerOutput(const FShaderCompilerOutput& Output, const FShaderCompileJobKey& Key, FString DebugInfo = FString());
 
 	int32 FindShaderIndex(const FSHAHash& InHash) const;
@@ -1159,12 +1136,6 @@ public:
 		return RHIShader;
 	}
 
-	UE_DEPRECATED(5.1, "GetRayTracingMaterialLibraryIndex is deprecated. Use GetRayTracingHitGroupLibraryIndex instead.")
-	inline uint32 GetRayTracingMaterialLibraryIndex() const
-	{
-		return GetRayTracingHitGroupLibraryIndex();
-	}
-
 	inline uint32 GetRayTracingHitGroupLibraryIndex() const
 	{
 		checkSlow(ShaderContent);
diff --git a/Engine/Source/Runtime/RenderCore/Public/ShaderCodeLibrary.h b/Engine/Source/Runtime/RenderCore/Public/ShaderCodeLibrary.h
--- a/Engine/Source/Runtime/RenderCore/Public/ShaderCodeLibrary.h
+++ b/Engine/Source/Runtime/RenderCore/Public/ShaderCodeLibrary.h
@@ -401,10 +401,6 @@ struct FShaderLibraryCooker
 	static RENDERCORE_API void DeleteShaderLibraryChunks(const ITargetPlatform* TargetPlatform,
 		const FString& SandboxDestinationPath, const FString& SandboxMetadataPath);
 
-	// Dump collected stats for each shader platform
-	UE_DEPRECATED(5.6, "DumpShaderCodeStats is deprecated due to being incorrect w.r.t shader duplication due to chunking.")
-	static RENDERCORE_API void DumpShaderCodeStats();
-
 	static RENDERCORE_API void DumpShaderTypeStats(const FString& DebugInfoPath, const FString& MetaDataPath);
 
 	// Create a smaller 'patch' library that only contains data from 'NewMetaDataDir' not contained in any of 'OldMetaDataDirs'
diff --git a/Engine/Source/Runtime/RenderCore/Public/ShaderCompilerCore.h b/Engine/Source/Runtime/RenderCore/Public/ShaderCompilerCore.h
--- a/Engine/Source/Runtime/RenderCore/Public/ShaderCompilerCore.h
+++ b/Engine/Source/Runtime/RenderCore/Public/ShaderCompilerCore.h
@@ -512,20 +512,9 @@ struct FShaderCompilerOutput
 	double PreprocessTime;
 	bool bSucceeded;
 	bool bSupportsQueryingUsedAttributes;
-	UE_DEPRECATED(5.6, "bSerializeModifiedSource is no longer used")
-	bool bSerializeModifiedSource;
 	bool bSerializingForCache;
 	TArray<FString> UsedAttributes;
 
-	PRAGMA_DISABLE_DEPRECATION_WARNINGS
-	// Explicitly-defaulted copy/move ctors & assignment operators are needed temporarily due to deprecation
-	// of bSerializeModifiedSource field. These can be removed once the deprecation window for this field ends.
-	FShaderCompilerOutput(FShaderCompilerOutput&&) = default;
-	FShaderCompilerOutput(const FShaderCompilerOutput&) = default;
-	FShaderCompilerOutput& operator=(FShaderCompilerOutput&&) = default;
-	FShaderCompilerOutput& operator=(const FShaderCompilerOutput&) = default;
-	PRAGMA_ENABLE_DEPRECATION_WARNINGS
-
 	TArray<FShaderCodeValidationStride> ParametersStrideToValidate;
 	TArray<FShaderCodeValidationType> ParametersSRVTypeToValidate;
 	TArray<FShaderCodeValidationType> ParametersUAVTypeToValidate;
@@ -545,9 +534,6 @@ struct FShaderCompilerOutput
 	 * a different entry point than was set on the FShaderCompilerInput struct. */
 	FString ModifiedEntryPointName;
 
-	UE_DEPRECATED(5.6, "Use symbols accessors (GetSymbolWriteAccess, GetSymbolReadView) on FShaderCode object instead (ShaderCode member)")
-	TArray<uint8> PlatformDebugData;
-
 	TArray<FGenericShaderStat> ShaderStatistics;
 
 	/** Generates OutputHash from the compiler output. */
diff --git a/Engine/Source/Runtime/RenderCore/Public/ShaderCore.h b/Engine/Source/Runtime/RenderCore/Public/ShaderCore.h
--- a/Engine/Source/Runtime/RenderCore/Public/ShaderCore.h
+++ b/Engine/Source/Runtime/RenderCore/Public/ShaderCore.h
@@ -835,7 +835,7 @@ enum class EShaderCodeFeatures : uint16
 	SixteenBitTypes         = 1 << 1,
 	TypedUAVLoadsExtended   = 1 << 2,
 	Atomic64                = 1 << 3,
-	DiagnosticBuffer UE_DEPRECATED(5.5, "EShaderCodeFeatures::DiagnosticBuffer is superseded by EShaderResourceUsageFlags::DiagnosticBuffer") = 1 << 4,
+	UNUSED_BIT_4            = 1 << 4,
 	BindlessResources       = 1 << 5,
 	BindlessSamplers        = 1 << 6,
 	StencilRef              = 1 << 7,
diff --git a/Engine/Source/Runtime/RenderCore/Public/ShaderParameterMacros.h b/Engine/Source/Runtime/RenderCore/Public/ShaderParameterMacros.h
--- a/Engine/Source/Runtime/RenderCore/Public/ShaderParameterMacros.h
+++ b/Engine/Source/Runtime/RenderCore/Public/ShaderParameterMacros.h
@@ -106,14 +106,6 @@ public:
 		return TUniformBufferRef<TBufferStruct>(RHICreateUniformBuffer(nullptr, TUniformBufferMetadataHelper<TBufferStruct>::GetStructMetadata()->GetLayoutPtr(), Usage, EUniformBufferValidation::ValidateResources));
 	}
 
-	UE_DEPRECATED(5.3, "UpdateUniformBufferImmediate requires a command list.")
-	void UpdateUniformBufferImmediate(const TBufferStruct& Value)
-	{
-	PRAGMA_DISABLE_DEPRECATION_WARNINGS
-		RHIUpdateUniformBuffer(GetReference(), &Value);
-	PRAGMA_ENABLE_DEPRECATION_WARNINGS
-	}
-
 	void UpdateUniformBufferImmediate(FRHICommandListBase& RHICmdList, const TBufferStruct& Value)
 	{
 		RHICmdList.UpdateUniformBuffer(GetReference(), &Value);
@@ -550,9 +542,6 @@ struct FRenderTargetBinding
 		check(Validate());
 	}
 
-	UE_DEPRECATED(5.1, "GetMipIndex has been renamed to SetMipIndex.")
-	void GetMipIndex(uint8 InMipIndex) { SetMipIndex(InMipIndex); }
-
 private:
 	/** Intentionally kept private to force setting the load action in the constructor. */
 	TAlignedShaderParameterPtr<FRDGTexture*> Texture		= nullptr;
diff --git a/Engine/Source/Runtime/RenderCore/Public/ShaderSource.h b/Engine/Source/Runtime/RenderCore/Public/ShaderSource.h
--- a/Engine/Source/Runtime/RenderCore/Public/ShaderSource.h
+++ b/Engine/Source/Runtime/RenderCore/Public/ShaderSource.h
@@ -36,16 +36,6 @@ public:
 	template <int NumChars>
 	using TStringBuilder = TAnsiStringBuilder<NumChars>;
 
-	/** Constexpr predicate indicating whether wide or ansi chars are used */
-	UE_DEPRECATED(5.6, "Shader code is always assumed to be ANSI")
-	static constexpr bool IsWide() { return false; }
-	/** Constexpr function returning the number of characters read in a single SIMD compare op */
-	UE_DEPRECATED(5.6, "Shader code is always assumed to be ANSI, so GetSimdCharCount() is always 16")
-	static constexpr int32 GetSimdCharCount() { return 16; }
-	/** Constexpr function returning a mask value for a single character */
-	UE_DEPRECATED(5.6, "Shader code is always assumed to be ANSI, so GetSingleCharMask() is always 1")
-	static constexpr int32 GetSingleCharMask() { return 1; }
-
 	/* Construct an empty shader source object; will still contain padding */
 	FShaderSource() { SetLen(0); };
 
diff --git a/Engine/Source/Runtime/RenderCore/Public/UniformBuffer.h b/Engine/Source/Runtime/RenderCore/Public/UniformBuffer.h
--- a/Engine/Source/Runtime/RenderCore/Public/UniformBuffer.h
+++ b/Engine/Source/Runtime/RenderCore/Public/UniformBuffer.h
@@ -52,12 +52,6 @@ public:
 		UpdateRHI(RHICmdList);
 	}
 
-	UE_DEPRECATED(5.4, "SetContents requires a command list.")
-	void SetContents(const TBufferStruct& NewContents)
-	{
-		SetContents(FRenderResource::GetImmediateCommandList(), NewContents);
-	}
-
 	/** Sets the contents of the uniform buffer to all zeros. */
 	void SetContentsToZero(FRHICommandListBase& RHICmdList)
 	{
@@ -69,12 +63,6 @@ public:
 		UpdateRHI(RHICmdList);
 	}
 
-	UE_DEPRECATED(5.4, "SetContentsToZero requires a command list.")
-	void SetContentsToZero()
-	{
-		SetContentsToZero(FRenderResource::GetImmediateCommandList());
-	}
-
 	const uint8* GetContents() const 
 	{
 		return Contents;
diff --git a/Engine/Source/Runtime/RenderCore/Public/VertexFactory.h b/Engine/Source/Runtime/RenderCore/Public/VertexFactory.h
--- a/Engine/Source/Runtime/RenderCore/Public/VertexFactory.h
+++ b/Engine/Source/Runtime/RenderCore/Public/VertexFactory.h
@@ -482,13 +482,6 @@ public:
 	/** Adds include statements for uniform buffers that this shader type references, and builds a prefix for the shader file with the include statements. */
 	RENDERCORE_API void AddUniformBufferIncludesToEnvironment(FShaderCompilerEnvironment& OutEnvironment, EShaderPlatform Platform) const;
 
-	UE_DEPRECATED(5.5, "GetReferencedUniformBufferNames is deprecated; call GetReferencedUniformBuffers instead.")
-	inline const TSet<const TCHAR*, TStringPointerSetKeyFuncs_DEPRECATED<const TCHAR*>>& GetReferencedUniformBufferNames() const 
-	{
-		static TSet<const TCHAR*, TStringPointerSetKeyFuncs_DEPRECATED<const TCHAR*>> EmptySet;
-		return EmptySet; 
-	};
-
 	const TSet<const FShaderParametersMetadata*>& GetReferencedUniformBuffers() const
 	{
 		return ReferencedUniformBuffers;
