[shaders] improved checks around definition fname string buffer sizes

#Synced-CL 48460480

diff --git a/Engine/Source/Runtime/RenderCore/Internal/ShaderCompilerDefinitions.h b/Engine/Source/Runtime/RenderCore/Internal/ShaderCompilerDefinitions.h
--- a/Engine/Source/Runtime/RenderCore/Internal/ShaderCompilerDefinitions.h
+++ b/Engine/Source/Runtime/RenderCore/Internal/ShaderCompilerDefinitions.h
@@ -444,13 +444,14 @@ private:
 			if (Ar.IsLoading())
 			{
 				Ar << StringLength;
-				check(StringLength < FName::StringBufferSize);
+				checkf(StringLength < FName::StringBufferSize, TEXT("Trying to deserialize FName with length %d; buffer size is %d"), StringLength, FName::StringBufferSize);
 				Ar.Serialize(KeyBufferAnsi, StringLength + 1);
 				Pair.Key = FName(KeyBufferAnsi);
 			}
 			else
 			{
 				StringLength = Pair.Key.ToString(KeyBuffer);
+				checkf(StringLength < FName::StringBufferSize, TEXT("FName %s length (%d) exceeds buffer size (%d)"), *Pair.Key.ToString(), StringLength, FName::StringBufferSize);
 				for (uint32 CharIndex = 0; CharIndex < StringLength + 1; CharIndex++)
 				{
 					KeyBufferAnsi[CharIndex] = (ANSICHAR)KeyBuffer[CharIndex];
