Fixed FRDGScatterUploadBuilder to support utility passes tacked on in order to share the command list task. Fixed Execute to reset state so that multiple Execute calls are possible.

#Synced-CL 48489722

diff --git a/Engine/Source/Runtime/RenderCore/Private/UnifiedBuffer.cpp b/Engine/Source/Runtime/RenderCore/Private/UnifiedBuffer.cpp
--- a/Engine/Source/Runtime/RenderCore/Private/UnifiedBuffer.cpp
+++ b/Engine/Source/Runtime/RenderCore/Private/UnifiedBuffer.cpp
@@ -511,6 +511,8 @@ FRDGBuffer* ResizeBufferIfNeeded(FRDGBuilder& GraphBuilder, TRefCountPtr<FRDGPoo
 	}
 	else
 	{
+		RDG_EVENT_SCOPE(GraphBuilder, "ResizeIfNeeded::MemcpyResource %d->%d", BufferSizeOld, BufferSizeNew);
+
 		InternalBufferNew = GraphBuilder.CreateBuffer(BufferDesc, Name);
 		ExternalBuffer = GraphBuilder.ConvertToExternalBuffer(InternalBufferNew);
 
@@ -1054,14 +1056,14 @@ FRDGScatterUploadBuilder* FRDGScatterUploadBuilder::Create(FRDGBuilder& GraphBui
 
 UE::Tasks::FTask FRDGScatterUploadBuilder::Execute(FRDGBuilder& GraphBuilder)
 {
-	if (Passes.IsEmpty())
+	if (Passes.IsEmpty() && CommandListFunctions.IsEmpty())
 	{
 		return {};
 	}
 
-	const bool bUseAsyncTask = MaxBytes > 32 * 1024;
+	const bool bUseAsyncTask = MaxBytes > 32 * 1024 || !CommandListFunctions.IsEmpty();
 
-	UE::Tasks::FTask Task = GraphBuilder.AddCommandListSetupTask([this] (FRHICommandList& RHICmdList)
+	UE::Tasks::FTask Task = GraphBuilder.AddCommandListSetupTask([Passes = Passes, CommandListFunctions = MoveTemp(CommandListFunctions)] (FRHICommandList& RHICmdList)
 	{
 		TRACE_CPUPROFILER_EVENT_SCOPE(FRDGScatterUploadBuilder);
 		for (auto& Pass : Passes)
@@ -1071,6 +1073,11 @@ UE::Tasks::FTask FRDGScatterUploadBuilder::Execute(FRDGBuilder& GraphBuilder)
 			Pass.Uploader->Unlock(RHICmdList);
 		}
 
+		for (auto& CommandListFunction : CommandListFunctions)
+		{
+			CommandListFunction(RHICmdList);
+		}
+
 	}, bUseAsyncTask);
 
 	for (auto& Pass : Passes)
@@ -1078,6 +1085,8 @@ UE::Tasks::FTask FRDGScatterUploadBuilder::Execute(FRDGBuilder& GraphBuilder)
 		Pass.UploadBuffer.End(GraphBuilder, Pass.Uploader);
 	}
 
+	Passes.Reset();
+	CommandListFunctions.Reset();
 	return Task;
 }
 
diff --git a/Engine/Source/Runtime/RenderCore/Public/UnifiedBuffer.h b/Engine/Source/Runtime/RenderCore/Public/UnifiedBuffer.h
--- a/Engine/Source/Runtime/RenderCore/Public/UnifiedBuffer.h
+++ b/Engine/Source/Runtime/RenderCore/Public/UnifiedBuffer.h
@@ -398,6 +398,7 @@ class FRDGScatterUploadBuilder
 {
 public:
 	using FPassFunction = TFunction<void(FRDGScatterUploader&)>;
+	using FCommandListFunction = TFunction<void(FRHICommandList&)>;
 
 	static UE::Tasks::FTask Process(
 		FRDGBuilder& GraphBuilder,
@@ -457,6 +458,11 @@ public:
 		MaxBytes += NumElements * NumBytesPerElement;
 	}
 
+	void AddFunction(FCommandListFunction&& InFunction)
+	{
+		CommandListFunctions.Emplace(MoveTemp(InFunction));
+	}
+
 	RENDERCORE_API UE::Tasks::FTask Execute(FRDGBuilder& GraphBuilder);
 
 private:
@@ -476,6 +482,7 @@ private:
 	};
 
 	TArray<FPass, SceneRenderingAllocator> Passes;
+	TArray<FCommandListFunction, SceneRenderingAllocator> CommandListFunctions;
 	uint32 MaxBytes = 0;
 
 	RDG_FRIEND_ALLOCATOR_FRIEND(FRDGScatterUploadBuilder);
