[IoStore] Refactor function `CreateIoStoreShaderCodeArchiveHeader()` into class `FShaderCodeArchiveProcessor` to improve readability.

This serves as readability improvement for upcoming changes to break up the 1:1 mapping between shader library chunks and pakchunks during packaging.
- Adds `FShaderIndexToShaderMapIndicesPair` struct to repalce tuple and give its members a name intead of generic `Get<0>()` etc.
- Move entire code logic of `CreateIoStoreShaderCodeArchiveHeader()` function into new class `FShaderCodeArchiveProcessor` and make all its lambdas member functions;
This improves readability by reducing code indentations, reduces boilerplate code of the large lambda captures, and makes future abstractions easier to implement once multiple input archives are processed into a single monolithic code archive.

#rnx

#Synced-CL 48546765

diff --git a/Engine/Source/Runtime/RenderCore/Private/ShaderCodeArchive.cpp b/Engine/Source/Runtime/RenderCore/Private/ShaderCodeArchive.cpp
--- a/Engine/Source/Runtime/RenderCore/Private/ShaderCodeArchive.cpp
+++ b/Engine/Source/Runtime/RenderCore/Private/ShaderCodeArchive.cpp
@@ -1505,397 +1505,196 @@ FIoChunkId FIoStoreShaderCodeArchive::GetShaderCodeChunkId(const FSHAHash& Shade
 	return ChunkId;
 }
 
-void FIoStoreShaderCodeArchive::CreateIoStoreShaderCodeArchiveHeader(const FName& Format, const FSerializedShaderArchive& SerializedShaders, FIoStoreShaderCodeArchiveHeader& OutHeader)
+struct FShaderIndexToShaderMapIndicesPair
 {
-	OutHeader.ShaderMapHashes = SerializedShaders.GetShaderMapHashes();
-	OutHeader.ShaderHashes = SerializedShaders.GetShaderHashes();
-	// shader group hashes will be populated later
+	uint32 ShaderIndex;
+	TArray<int32> ShaderMapIndices;
 
-	OutHeader.ShaderMapEntries.Empty(SerializedShaders.GetShaderMapEntries().Num());
-	for (const FShaderMapEntry& ShaderMapEntry : SerializedShaders.GetShaderMapEntries())
+	static bool Sort(const FShaderIndexToShaderMapIndicesPair& EntryA, const FShaderIndexToShaderMapIndicesPair& EntryB)
 	{
-		FIoStoreShaderMapEntry& IoStoreShaderMapEntry = OutHeader.ShaderMapEntries.AddDefaulted_GetRef();
-		IoStoreShaderMapEntry.ShaderIndicesOffset = ShaderMapEntry.ShaderIndicesOffset;
-		IoStoreShaderMapEntry.NumShaders = ShaderMapEntry.NumShaders;
-	}
-
-	// indices should be copied before grouping as the groups will append to the array
-	OutHeader.ShaderIndices = SerializedShaders.GetShaderIndices();
-
-	// shader entries are copied, the remainder of the field will be assigned when splitting into groups
-	OutHeader.ShaderEntries.Empty(SerializedShaders.GetShaderEntries().Num());
-	for (const FShaderCodeEntry& ShaderEntry : SerializedShaders.GetShaderEntries())
-	{
-		FIoStoreShaderCodeEntry& IoStoreShaderEntry = OutHeader.ShaderEntries.AddDefaulted_GetRef();
-		IoStoreShaderEntry.Frequency = ShaderEntry.Frequency;
-	}
-
-	// Higher level description of the group splitting algo that follows:
-	// We compress together shaders that are loaded together (all other strategies, like grouping similar shaders, were found to compress better but not reduce RAM usage).
-	// For that, we find for each shader which shadermaps are referencing it (often times it will be just one, but for some simple and shared shaders it can be thousands).
-	// We group the shaders by those sets of shadermaps - all shaders referenced by the same shadermap(s) are a candidate for being a single group. Then we potentially split this candidate group
-	// into raytracing and non-raytracing groups (so we can avoid preloading RTX shaders run-time if RTX is disabled), and then each of those is potentially split further by size
-	// (to avoid too large groups that will take too much time to decompress - this is regulated by r.ShaderCodeLibrary.MaxShaderGroupSize). The results of that process (note, it can still be
-	// a single group) is added to the header.
-	// Each group's indices, like in case of shadermaps, are stored in ShaderIndices array. Before we append a new group's indices however, we look if we can find an existing range that we can reuse.
-	const bool bSeparateRaytracingShaders = (Format == FName("PCD3D_SM5"));
-
-	TArray<TPair<uint32, TArray<int32>>> ShaderToShadermapsArray;
-	{
-		ShaderToShadermapsArray.AddDefaulted(OutHeader.ShaderEntries.Num());
-
+		const TArray<int32>& A = EntryA.ShaderMapIndices;
+		const TArray<int32>& B = EntryB.ShaderMapIndices;
+		// if the number of shadermaps is the same, we need to sort "alphabetically"
+		if (A.Num() == B.Num())
 		{
-			FCriticalSection ShaderLocks[1024];
-			// for each shader, find all the shadermaps it belongs to
-			ParallelFor(OutHeader.ShaderMapEntries.Num(),
-				[&ShaderToShadermapsArray, &OutHeader, &ShaderLocks](int32 ShaderMapIndex)
+			for (int32 Idx = 0, Num = A.Num(); Idx < Num; ++Idx)
+			{
+				if (A[Idx] != B[Idx])
 				{
-					FIoStoreShaderMapEntry& ShaderMapEntry = OutHeader.ShaderMapEntries[ShaderMapIndex];
-					for (int32 ShaderIdxIdx = ShaderMapEntry.ShaderIndicesOffset, StopBeforeIdxIdx = ShaderMapEntry.ShaderIndicesOffset + ShaderMapEntry.NumShaders; ShaderIdxIdx < StopBeforeIdxIdx; ++ShaderIdxIdx)
-					{
-						int32 ShaderIndex = OutHeader.ShaderIndices[ShaderIdxIdx];
-						// add this shadermap as a dependency.
-						int32 ShaderLockNumber = ShaderIndex % UE_ARRAY_COUNT(ShaderLocks);
-						FScopeLock Locker(&ShaderLocks[ShaderLockNumber]);
-						ShaderToShadermapsArray[ShaderIndex].Value.Add(ShaderMapIndex);
-					}
+					return A[Idx] < B[Idx];
 				}
-			);
-		}
+			}
 
-		// sort shadermaps entries in shaders
-		{
-			const int32 kShaderSortedPerThread = 1024;
-			const int32 NumThreads = (ShaderToShadermapsArray.Num() / kShaderSortedPerThread) + 1;
-			ParallelFor(NumThreads,
-				[&ShaderToShadermapsArray, kShaderSortedPerThread](int ThreadIndex)
-				{
-					int32 StartingShader = ThreadIndex * kShaderSortedPerThread;
-					int32 EndingShader = FMath::Min(StartingShader + kShaderSortedPerThread, ShaderToShadermapsArray.Num());
-					for (int32 Idx = StartingShader; Idx < EndingShader; ++Idx)
-					{
-						ShaderToShadermapsArray[Idx].Value.Sort();
-					}
-				},
-				EParallelForFlags::Unbalanced
-			);
+			return EntryA.ShaderIndex < EntryB.ShaderIndex;
 		}
 
-		// now assigning the indices in the array so we can sort it
-		for (int32 Idx = 0, Num = ShaderToShadermapsArray.Num(); Idx < Num; ++Idx)
-		{
-			// check that no shader is unreferenced
-			checkf(!ShaderToShadermapsArray[Idx].Value.IsEmpty(), TEXT("Error converting to IoStore archive: shader (index=%d) is not referenced by any of the shadermaps!"), Idx);
-			ShaderToShadermapsArray[Idx].Key = Idx;
-		}
+		return A.Num() < B.Num();
 	}
+};
 
-	// sort the mapping so the first are shaders that are referenced by a smaller number of shadermaps, then by index for determinism
-	Algo::Sort(ShaderToShadermapsArray,
-		[](const TPair<uint32, TArray<int32>>& EntryA, const TPair<uint32, TArray<int32>>& EntryB)
-		{
-			const TArray<int32>& A = EntryA.Value;
-			const TArray<int32>& B = EntryB.Value;
-			// if the number of shadermaps is the same, we need to sort "alphabetically"
-			if (A.Num() == B.Num())
-			{
-				for (int32 Idx = 0, Num = A.Num(); Idx < Num; ++Idx)
-				{
-					if (A[Idx] != B[Idx])
-					{
-						return A[Idx] < B[Idx];
-					}
-				}
-
-				return EntryA.Key < EntryB.Key;
-			}
+class FShaderCodeArchiveProcessor
+{
+	const bool bSeparateRaytracingShaders;
+	const uint32 MaxUncompressedShaderGroupSize;
+	const TStringConversion<TStringConvert<TCHAR, UTF8CHAR>> Utf8FormatName;
 
-			return A.Num() < B.Num();
-		}
-	);
+	FIoStoreShaderCodeArchiveHeader& OutHeader;
+	const FSerializedShaderArchive* SerializedShaders;
 
-	// get the effective maximum uncompressed group size
-	uint32 MaxUncompressedShaderGroupSize = FMath::Max(static_cast<uint32>(GShaderCodeLibraryMaxShaderGroupSize), 1U);	// cannot be lower than 1
+	TArray<FShaderIndexToShaderMapIndicesPair> ShaderToShadermapsArray;
 
 	// for statistics
-	uint64 Stats_TotalUncompressedMemory = 0;
-	uint32 Stats_MinGroupSize = MAX_uint32;
-	uint32 Stats_MaxGroupSize = 0;
+	struct FProcessingStats
+	{
+		uint64 TotalUncompressedMemory = 0;
+		uint32 MinGroupSize = MAX_uint32;
+		uint32 MaxGroupSize = 0;
+		int32 GroupsThatAppendedToShaderIndices = 0;
+	};
+	FProcessingStats Stats;
 
 	// We want to avoid adding group indices to ShaderIndices, however looking them up sequentially is to slow. Store them here for a future ParallelFor lookup.
 	TArray<TArray<uint32>> StoredGroupShaderIndices;
 
-	/** Third and last stage of processing the shader group. We actually add it here, and do the book-keeping. */
-	auto ProcessShaderGroup_AddNewGroup = [&Format, &OutHeader, &SerializedShaders, &StoredGroupShaderIndices, &Stats_TotalUncompressedMemory, &Stats_MinGroupSize, &Stats_MaxGroupSize](TArray<uint32>& ShaderIndicesInGroup)
-	{
-		// first, sort the shaders by uncompressed size, as this was found to compress better
-		ShaderIndicesInGroup.Sort(
-			[&SerializedShaders](const int32 ShaderIndexA, const int32 ShaderIndexB)
-			{
-				const FShaderCodeEntry& ShaderEntryA = SerializedShaders.GetShaderEntries()[ShaderIndexA];
-				const FShaderCodeEntry& ShaderEntryB = SerializedShaders.GetShaderEntries()[ShaderIndexB];
-				if (ShaderEntryA.UncompressedSize != ShaderEntryB.UncompressedSize)
-				{
-					return ShaderEntryA.UncompressedSize < ShaderEntryB.UncompressedSize;
-				}
-				if (ShaderEntryA.Size != ShaderEntryB.Size)
-				{
-					return ShaderEntryA.Size < ShaderEntryB.Size;
-				}
-				if (ShaderEntryA.Frequency != ShaderEntryB.Frequency)
-				{
-					return ShaderEntryA.Frequency < ShaderEntryB.Frequency;
-				}
-				return ShaderEntryA.Offset < ShaderEntryB.Offset;
-			}
-		);
+public:
+	FShaderCodeArchiveProcessor(const FName& Format, FIoStoreShaderCodeArchiveHeader& OutHeader);
 
-		// add a new group entry
-		const int32 CurrentGroupIdx = OutHeader.ShaderGroupEntries.Num();
-		FIoStoreShaderGroupEntry& GroupEntry = OutHeader.ShaderGroupEntries.AddDefaulted_GetRef();
-		StoredGroupShaderIndices.Add(ShaderIndicesInGroup);
-		GroupEntry.NumShaders = ShaderIndicesInGroup.Num();
-		// ShaderIndicesOffset will be filled later, once we know all the groups (see comment about StoredGroupShaderIndices above).
-
-		const TArrayView<const FSHAHash> ShaderHashes = SerializedShaders.GetShaderHashes();
-		const TArrayView<const FShaderCodeEntry> ShaderEntries = SerializedShaders.GetShaderEntries();
-
-		// update shader entries both with the group number and their uncompressed offset in the group
-		FSHA1 GroupHasher;
-		uint32 CurrentGroupSize = 0;
-		for (int32 ShaderIdxIdx = 0, NumIdxIdx = ShaderIndicesInGroup.Num(); ShaderIdxIdx < NumIdxIdx; ++ShaderIdxIdx)
-		{
-			int32 ShaderIndex = ShaderIndicesInGroup[ShaderIdxIdx];
-			FIoStoreShaderCodeEntry& IoStoreShaderEntry = OutHeader.ShaderEntries[ShaderIndex];
-			IoStoreShaderEntry.ShaderGroupIndex = CurrentGroupIdx;
-			IoStoreShaderEntry.UncompressedOffsetInGroup = CurrentGroupSize;
-
-			// group hash is constructed from hashing the shaders in the group.
-			GroupHasher.Update(ShaderHashes[ShaderIndex].Hash, sizeof(FSHAHash));
-			// shader hash as of now excludes optional data, so we cannot rely on it, especially across the shader formats. Make the group hash a bit more robust by including the shader size in it.
-			GroupHasher.Update(reinterpret_cast<const uint8*>(&ShaderEntries[ShaderIndex].UncompressedSize), sizeof(FShaderCodeEntry::UncompressedSize));
-
-			CurrentGroupSize += ShaderEntries[ShaderIndex].UncompressedSize;
-		}
-		// Shader hashes cannot be used to uniquely identify across shader formats due to aforementioned exclusion of optional data from it.
-		// Include the shader format (in a platform-agnostic way) into the group hash to lower the risk of collision of shaders of different formats.
-		const TStringConversion<TStringConvert<TCHAR, UTF8CHAR>> Utf8Name(Format.ToString());
-		GroupHasher.Update(reinterpret_cast<const uint8*>(Utf8Name.Get()), Utf8Name.Length());
-		static_assert(sizeof(uint8) == sizeof(UTF8CHAR), "Unexpected UTF-8 char size.");
-
-		GroupEntry.UncompressedSize = CurrentGroupSize;
-		OutHeader.ShaderGroupIoHashes.Add(FIoStoreShaderCodeArchive::GetShaderCodeChunkId(GroupHasher.Finalize()));
-		
-		Stats_TotalUncompressedMemory += CurrentGroupSize;
-		Stats_MinGroupSize = FMath::Min(Stats_MinGroupSize, CurrentGroupSize);
-		Stats_MaxGroupSize = FMath::Max(Stats_MaxGroupSize, CurrentGroupSize);
-	};
+	void ProcessSerializedShaders(const FSerializedShaderArchive& SerializedShaders);
 
-	/** Second stage of processing shader group. Here we potentially split the group into smaller ones (as equally as possible), striving to meet limit imposed by r.ShaderCodeLibrary.MaxShaderGroupSize */
-	auto ProcessShaderGroup_SplitBySize = [&OutHeader, &SerializedShaders, &ProcessShaderGroup_AddNewGroup, &MaxUncompressedShaderGroupSize](TArray<uint32>& CurrentShaderGroup)
-	{
-		// calculate current group size
-		uint32 GroupSize = 0;
-		const TArrayView<const FShaderCodeEntry> ShaderEntries = SerializedShaders.GetShaderEntries();
-		for (uint32 ShaderIdx : CurrentShaderGroup)
-		{
-			GroupSize += ShaderEntries[ShaderIdx].UncompressedSize;
-		}
-		
-		if (LIKELY(GroupSize <= MaxUncompressedShaderGroupSize || CurrentShaderGroup.Num() == 1))
-		{
-			ProcessShaderGroup_AddNewGroup(CurrentShaderGroup);
-		}
-		else
-		{
-			// split the shaders evenly into N new groups (don't allow more new groups than there are shaders)
-			int32 NumNewGroups = FMath::Min(static_cast<int32>(GroupSize / MaxUncompressedShaderGroupSize + 1), CurrentShaderGroup.Num());
-			checkf(NumNewGroups > 1, TEXT("Off by one error in group count calculation? NumNewGroups=%d, GroupSize=%u, MaxUncompressedShaderGroupSize=%u, CurrentShaderGroup.Num()=%d"), NumNewGroups, GroupSize, MaxUncompressedShaderGroupSize, CurrentShaderGroup.Num());
+private:
+	// Populates the ShaderToShadermapsArray member via information from the archive header.
+	void BuildShaderToShadermapsArray();
 
-			TArray<TArray<uint32>> NewGroups;
-			TArray<uint32> NewGroupSizes;
-			NewGroups.AddDefaulted(NumNewGroups);
-			NewGroupSizes.AddZeroed(NumNewGroups);
+	// First stage of processing a streak of shaders all referenced by the same set of shadermaps.
+	// We begin with separating raytracing and non-raytracing shaders, so we can avoid preloading RTX in non-RT runs.
+	void SplitShaderGroupForRaytracing(TArray<uint32>& CurrentShaderGroup);
 
-			// sort the shaders descending, as this is easier to split (greedy algorithm)
-			CurrentShaderGroup.Sort(
-				[&SerializedShaders](const int32 ShaderIndexA, const int32 ShaderIndexB)
-				{
-					const TArrayView<const FShaderCodeEntry> ShaderEntries = SerializedShaders.GetShaderEntries();
-					const FShaderCodeEntry& ShaderEntryA = ShaderEntries[ShaderIndexA];
-					const FShaderCodeEntry& ShaderEntryB = ShaderEntries[ShaderIndexB];
-					if (ShaderEntryA.UncompressedSize != ShaderEntryB.UncompressedSize)
-					{
-						return ShaderEntryA.UncompressedSize > ShaderEntryB.UncompressedSize;
-					}
-					if (ShaderEntryA.Size != ShaderEntryB.Size)
-					{
-						return ShaderEntryA.Size > ShaderEntryB.Size;
-					}
-					if (ShaderEntryA.Frequency != ShaderEntryB.Frequency)
-					{
-						return ShaderEntryA.Frequency > ShaderEntryB.Frequency;
-					}
-					return ShaderEntryA.Offset > ShaderEntryB.Offset;
-				}
-			);
+	// Second stage of processing shader group. Here we potentially split the group into smaller ones (as equally as possible),
+	// striving to meet limit imposed by r.ShaderCodeLibrary.MaxShaderGroupSize.
+	void SplitShaderGroupBySize(TArray<uint32>& CurrentShaderGroup);
 
-			for (int32 ShaderIdx : CurrentShaderGroup)
-			{
-				// add the shader to the group of smallest size
-				int32 SmallestNewGroupIdx = 0;
-				for (int32 IdxNewGroup = 1; IdxNewGroup < NumNewGroups; ++IdxNewGroup)
-				{
-					if (NewGroupSizes[IdxNewGroup] < NewGroupSizes[SmallestNewGroupIdx])
-					{
-						SmallestNewGroupIdx = IdxNewGroup;
-					}
-				}
+	// Third and last stage of processing the shader group. We actually add it here, and do the book-keeping.
+	void AddNewShaderGroup(TArray<uint32>& ShaderIndicesInGroup);
 
-				NewGroups[SmallestNewGroupIdx].Add(ShaderIdx);
-				NewGroupSizes[SmallestNewGroupIdx] += ShaderEntries[ShaderIdx].UncompressedSize;
-			}
-			
-#if DO_CHECK // sanity checks
-			uint32 NewGroupTotalSize = 0;
-			for (uint32 NewGroupSize : NewGroupSizes)
-			{
-				NewGroupTotalSize += NewGroupSize;
-			}
-			checkf(NewGroupTotalSize == GroupSize, TEXT("Original shader group was of size %u bytes, which was larger than limit %u, but it was split into %d group of total size %u, which is not %u - sizes must agree"),
-				GroupSize,
-				MaxUncompressedShaderGroupSize,
-				NumNewGroups,
-				NewGroupTotalSize, GroupSize
-				);
-#endif
+};
 
-			for (TArray<uint32>& NewGroup : NewGroups)
-			{
-				// note there can be empty groups (take a very edge case of MaxUncompressedShaderGroupSize = 2 bytes and a shader group of 1 shader)
-				if (!NewGroup.IsEmpty())
-				{
-					ProcessShaderGroup_AddNewGroup(NewGroup);
-				}
-			}
-		}
-	};
+FShaderCodeArchiveProcessor::FShaderCodeArchiveProcessor(const FName& Format, FIoStoreShaderCodeArchiveHeader& OutHeader)
+	: bSeparateRaytracingShaders(Format == FName("PCD3D_SM5"))
+	// get the effective maximum uncompressed group size (cannot be lower than 1)
+	, MaxUncompressedShaderGroupSize(FMath::Max(static_cast<uint32>(GShaderCodeLibraryMaxShaderGroupSize), 1U))
+	, Utf8FormatName(Format.ToString())
+	, OutHeader(OutHeader)
+{
+}
+
+/** Tries to find whether SequenceToFind exists as a subsequence in ExistingIndices. Returns INDEX_NONE if not found. */
+static int32 FindStartOfSequenceInArray(const TArrayView<const uint32> SuperSet, const TArrayView<const uint32> SequenceToFind)
+{
+	check(SequenceToFind.Num() > 0);
 
-	/** First stage of processing a streak of shaders all referenced by the same set of shadermaps. We begin with separating raytracing and non-raytracing shaders, so we can avoid preloading RTX in non-RT runs. */
-	auto ProcessShaderGroup_SplitRaytracing = [&bSeparateRaytracingShaders, &OutHeader, &ProcessShaderGroup_SplitBySize](TArray<uint32>& CurrentShaderGroup)
+	const uint32 FirstSequenceEntry = SequenceToFind[0];
+	const int32 SequenceLength = SequenceToFind.Num();
+	for (int32 SuperSetIndex = 0, SuperSetLength = SuperSet.Num(); SuperSetIndex < SuperSetLength - SequenceLength + 1; ++SuperSetIndex)
 	{
-		if (!bSeparateRaytracingShaders)
+		if (LIKELY(SuperSet[SuperSetIndex] != FirstSequenceEntry))
 		{
-			ProcessShaderGroup_SplitBySize(CurrentShaderGroup);
+			continue;
 		}
-		else
+
+		// check the rest
+		bool bFoundSequence = true;
+		for (int32 SequenceIndex = 1; SequenceIndex < SequenceLength; ++SequenceIndex)
 		{
-			// The streak changed. Create the group, but first, determine if the group needs to be split in two because of the raytracing shaders.
-			// We want to isolate them into separate groups so their preload can be skipped if raytracing is off.
-			TArray<uint32> RaytracingShaders;
-			TArray<uint32> NonraytracingShaders;
-			for (int32 ShaderIndex : CurrentShaderGroup)
+			checkSlow(SuperSetIndex + SequenceIndex < SuperSetLength);
+			if (LIKELY(SuperSet[SuperSetIndex + SequenceIndex] != SequenceToFind[SequenceIndex]))
 			{
-				if (LIKELY(!IsRayTracingShaderFrequency(static_cast<EShaderFrequency>(OutHeader.ShaderEntries[ShaderIndex].Frequency))))
-				{
-					NonraytracingShaders.Add(ShaderIndex);
-				}
-				else
-				{
-					RaytracingShaders.Add(ShaderIndex);
-				}
+				bFoundSequence = false;
+				break;
 			}
-			check(CurrentShaderGroup.Num() == NonraytracingShaders.Num() + RaytracingShaders.Num());
+		}
 
-			if (LIKELY(!NonraytracingShaders.IsEmpty()))
-			{
-				ProcessShaderGroup_SplitBySize(NonraytracingShaders);
-			}
-			if (UNLIKELY(!RaytracingShaders.IsEmpty()))
-			{
-				ProcessShaderGroup_SplitBySize(RaytracingShaders);
-			}
+		if (UNLIKELY(bFoundSequence))
+		{
+			return SuperSetIndex;
 		}
-	};
+	}
+
+	return INDEX_NONE;
+}
+
+void FShaderCodeArchiveProcessor::ProcessSerializedShaders(const FSerializedShaderArchive& InSerializedShaders)
+{
+	// Load up serialized shaders and populate archive header
+	SerializedShaders = &InSerializedShaders;
+
+	OutHeader.ShaderMapHashes = SerializedShaders->GetShaderMapHashes();
+	OutHeader.ShaderHashes = SerializedShaders->GetShaderHashes();
+	// shader group hashes will be populated later
+
+	OutHeader.ShaderMapEntries.Empty(SerializedShaders->GetShaderMapEntries().Num());
+	for (const FShaderMapEntry& ShaderMapEntry : SerializedShaders->GetShaderMapEntries())
+	{
+		FIoStoreShaderMapEntry& IoStoreShaderMapEntry = OutHeader.ShaderMapEntries.AddDefaulted_GetRef();
+		IoStoreShaderMapEntry.ShaderIndicesOffset = ShaderMapEntry.ShaderIndicesOffset;
+		IoStoreShaderMapEntry.NumShaders = ShaderMapEntry.NumShaders;
+	}
+
+	// indices should be copied before grouping as the groups will append to the array
+	OutHeader.ShaderIndices = SerializedShaders->GetShaderIndices();
+
+	// shader entries are copied, the remainder of the field will be assigned when splitting into groups
+	OutHeader.ShaderEntries.Empty(SerializedShaders->GetShaderEntries().Num());
+	for (const FShaderCodeEntry& ShaderEntry : SerializedShaders->GetShaderEntries())
+	{
+		FIoStoreShaderCodeEntry& IoStoreShaderEntry = OutHeader.ShaderEntries.AddDefaulted_GetRef();
+		IoStoreShaderEntry.Frequency = ShaderEntry.Frequency;
+	}
+
+	BuildShaderToShadermapsArray();
 
 	// now split this into streaks of shaders that are referenced by the same set of shadermaps and compress
 	TArray<uint32> CurrentShaderGroup;
 	TArray<int32> LastShadermapSetSeen;
-	for (TPair<uint32, TArray<int32>>& Iter : ShaderToShadermapsArray)
+	for (const FShaderIndexToShaderMapIndicesPair& ShaderIndexInfo : ShaderToShadermapsArray)
 	{
 		// if we have have just started the group, we don't check the last seen
 		if (UNLIKELY(CurrentShaderGroup.IsEmpty()))
 		{
-			CurrentShaderGroup.Add(Iter.Key);
-			LastShadermapSetSeen = Iter.Value;
+			CurrentShaderGroup.Add(ShaderIndexInfo.ShaderIndex);
+			LastShadermapSetSeen = ShaderIndexInfo.ShaderMapIndices;
 		}
-		else if (UNLIKELY(LastShadermapSetSeen != Iter.Value))
+		else if (UNLIKELY(LastShadermapSetSeen != ShaderIndexInfo.ShaderMapIndices))
 		{
-			ProcessShaderGroup_SplitRaytracing(CurrentShaderGroup);
+			SplitShaderGroupForRaytracing(CurrentShaderGroup);
 
 			// reset the collection, but don't forget to add to it the current element
 			CurrentShaderGroup.SetNum(1);
-			CurrentShaderGroup[0] = Iter.Key;
-			LastShadermapSetSeen = Iter.Value;
+			CurrentShaderGroup[0] = ShaderIndexInfo.ShaderIndex;
+			LastShadermapSetSeen = ShaderIndexInfo.ShaderMapIndices;
 		}
 		else
 		{
 			// keep adding to the same group
-			CurrentShaderGroup.Add(Iter.Key);
+			CurrentShaderGroup.Add(ShaderIndexInfo.ShaderIndex);
 		}
 	}
 	// add the last group
 	if (!CurrentShaderGroup.IsEmpty())
 	{
-		ProcessShaderGroup_SplitRaytracing(CurrentShaderGroup);
+		SplitShaderGroupForRaytracing(CurrentShaderGroup);
 	}
 
-	/** Tries to find whether NewIndices exist as a subsequence in ExistingIndices.Returns - 1 if not found. */
-	auto FindSequenceInArray = [](const TArrayView<const uint32> ExistingIndices, const TArrayView<const uint32> NewIndices) -> int32
-	{
-		check(NewIndices.Num() > 0);
-
-		uint32 FirstNew = NewIndices[0];
-		int32 NumNew = NewIndices.Num();
-		for (int32 IdxExisting = 0, NumExisting = ExistingIndices.Num(); IdxExisting < NumExisting - NumNew + 1; ++IdxExisting)
-		{
-			if (LIKELY(ExistingIndices[IdxExisting] != FirstNew))
-			{
-				continue;
-			}
-
-			// check the rest
-			bool bFoundSequence = true;
-			for (int32 IdxNew = 1; IdxNew < NumNew; ++IdxNew)
-			{
-				checkSlow(IdxExisting + IdxNew < NumExisting);
-				if (LIKELY(ExistingIndices[IdxExisting + IdxNew] != NewIndices[IdxNew]))
-				{
-					bFoundSequence = false;
-					break;
-				}
-			}
-
-			if (UNLIKELY(bFoundSequence))
-			{
-				return IdxExisting;
-			}
-		}
-
-		return INDEX_NONE;
-	};
-
 	// now, try to see if we can look up group's indices in the existing ShaderIndicesArray to avoid storing them there.
 	checkf(StoredGroupShaderIndices.Num() == OutHeader.ShaderGroupEntries.Num(), TEXT("We should have stored shader indices for all groups."));
 	ParallelFor(StoredGroupShaderIndices.Num(),
-		[&OutHeader, &SerializedShaders, &StoredGroupShaderIndices, &FindSequenceInArray](int ShaderGroupIndex)
+		[this](int32 ShaderGroupIndex)
 		{
 			const TArray<uint32>& ShaderIndicesInGroup = StoredGroupShaderIndices[ShaderGroupIndex];
 			FIoStoreShaderGroupEntry& GroupEntry = OutHeader.ShaderGroupEntries[ShaderGroupIndex];
 			// See if we can find indices in that order somewhere in the ShaderIndices array already, to avoid adding new indices.
 			// We are looking in the read-only original array, because there's no sense to look in OutHeader.ShaderIndices - groups don't overlap,
 			// so we know that newly added (by some previous group) indices aren't useful for us.
-			int32 ExistingOffset = FindSequenceInArray(SerializedShaders.GetShaderIndices(), ShaderIndicesInGroup);
+			const int32 ExistingOffset = FindStartOfSequenceInArray(SerializedShaders->GetShaderIndices(), ShaderIndicesInGroup);
 			if (ExistingOffset != INDEX_NONE)
 			{
 				GroupEntry.ShaderIndicesOffset = static_cast<uint32>(ExistingOffset);
@@ -1908,7 +1707,6 @@ void FIoStoreShaderCodeArchive::CreateIoStoreShaderCodeArchiveHeader(const FName
 	);
 
 	// Now append all the groups that weren't found to the end of ShaderIndices, slow (we could have done that in above ParallelFor, with a lock), but good for determinism
-	int32 Stats_GroupsThatAppendedToShaderIndices = 0;
 	for (int32 ShaderGroupIndex = 0, NumGroups = OutHeader.ShaderGroupEntries.Num(); ShaderGroupIndex < NumGroups; ++ShaderGroupIndex)
 	{
 		FIoStoreShaderGroupEntry& GroupEntry = OutHeader.ShaderGroupEntries[ShaderGroupIndex];
@@ -1918,18 +1716,290 @@ void FIoStoreShaderCodeArchive::CreateIoStoreShaderCodeArchiveHeader(const FName
 			const TArray<uint32>& ShaderIndicesInGroup = StoredGroupShaderIndices[ShaderGroupIndex];
 			GroupEntry.ShaderIndicesOffset = OutHeader.ShaderIndices.Num();
 			OutHeader.ShaderIndices.Append(ShaderIndicesInGroup);
-			++Stats_GroupsThatAppendedToShaderIndices;
+			Stats.GroupsThatAppendedToShaderIndices++;
 		}
 	}
 
-	checkf(OutHeader.ShaderEntries.Num() == SerializedShaders.GetShaderEntries().Num(), TEXT("Error creating IoStoreShaderArchive header - shader entries differ (%d in IoStore, %d original). Bug in grouping logic?"),
-		OutHeader.ShaderEntries.Num(), SerializedShaders.GetShaderEntries().Num());
-	checkf(OutHeader.ShaderGroupIoHashes.Num() == OutHeader.ShaderGroupEntries.Num(), TEXT("Error creating IoStoreShaderArchive header - mismatch between shader group hashes and descriptors (%d descriptors, %d hashes). Bug in grouping logic?"),
+	// Validate and log stats
+	checkf(OutHeader.ShaderEntries.Num() == SerializedShaders->GetShaderEntries().Num(),
+		TEXT("Error creating IoStoreShaderArchive header - shader entries differ (%d in IoStore, %d original). Bug in grouping logic?"),
+		OutHeader.ShaderEntries.Num(), SerializedShaders->GetShaderEntries().Num());
+	checkf(OutHeader.ShaderGroupIoHashes.Num() == OutHeader.ShaderGroupEntries.Num(),
+		TEXT("Error creating IoStoreShaderArchive header - mismatch between shader group hashes and descriptors (%d descriptors, %d hashes). Bug in grouping logic?"),
 		OutHeader.ShaderGroupEntries.Num(), OutHeader.ShaderGroupIoHashes.Num());
 	checkf(OutHeader.ShaderGroupEntries.Num() != 0, TEXT("At least one group must have been created"));
 
-	UE_LOG(LogShaderLibrary, Display, TEXT("Created IoStoreShaderArchive header: shaders grouped in %d groups (%d of them didn't need new indices), average uncompressed size %llu bytes, min %u bytes, max %u bytes (r.ShaderCodeLibrary.MaxShaderGroupSize=%u)"),
-		OutHeader.ShaderGroupEntries.Num(), OutHeader.ShaderGroupEntries.Num() - Stats_GroupsThatAppendedToShaderIndices, Stats_TotalUncompressedMemory / static_cast<uint64>(OutHeader.ShaderGroupEntries.Num()), Stats_MinGroupSize, Stats_MaxGroupSize, MaxUncompressedShaderGroupSize);
+	UE_LOG(LogShaderLibrary, Display,
+		TEXT("Created IoStoreShaderArchive header: shaders grouped in %d groups (%d of them didn't need new indices), average uncompressed size %llu bytes, min %u bytes, max %u bytes (r.ShaderCodeLibrary.MaxShaderGroupSize=%u)"),
+		OutHeader.ShaderGroupEntries.Num(),
+		OutHeader.ShaderGroupEntries.Num() - Stats.GroupsThatAppendedToShaderIndices,
+		Stats.TotalUncompressedMemory / static_cast<uint64>(OutHeader.ShaderGroupEntries.Num()),
+		Stats.MinGroupSize,
+		Stats.MaxGroupSize,
+		MaxUncompressedShaderGroupSize);
+}
+
+void FShaderCodeArchiveProcessor::BuildShaderToShadermapsArray()
+{
+	ShaderToShadermapsArray.AddDefaulted(OutHeader.ShaderEntries.Num());
+
+	{
+		FCriticalSection ShaderLocks[1024];
+		// for each shader, find all the shadermaps it belongs to
+		ParallelFor(OutHeader.ShaderMapEntries.Num(),
+			[this, &ShaderLocks](int32 ShaderMapIndex)
+			{
+				const FIoStoreShaderMapEntry& ShaderMapEntry = OutHeader.ShaderMapEntries[ShaderMapIndex];
+				const int32 ShaderIndicesRange[2] =
+				{
+					ShaderMapEntry.ShaderIndicesOffset,
+					ShaderMapEntry.ShaderIndicesOffset + ShaderMapEntry.NumShaders
+				};
+				for (int32 ShaderIndicesEntryIndex = ShaderIndicesRange[0]; ShaderIndicesEntryIndex < ShaderIndicesRange[1]; ++ShaderIndicesEntryIndex)
+				{
+					// add this shadermap as a dependency.
+					const int32 ShaderIndex = OutHeader.ShaderIndices[ShaderIndicesEntryIndex];
+					const int32 ShaderLockNumber = ShaderIndex % UE_ARRAY_COUNT(ShaderLocks);
+					FScopeLock Locker(&ShaderLocks[ShaderLockNumber]);
+					ShaderToShadermapsArray[ShaderIndex].ShaderMapIndices.Add(ShaderMapIndex);
+				}
+			}
+		);
+	}
+
+	// sort shadermaps entries in shaders
+	{
+		constexpr int32 kShaderSortedPerThread = 1024;
+		const int32 NumThreads = (ShaderToShadermapsArray.Num() / kShaderSortedPerThread) + 1;
+		ParallelFor(NumThreads,
+			[this, kShaderSortedPerThread](int32 ThreadIndex)
+			{
+				const int32 ShaderIndexStart = ThreadIndex * kShaderSortedPerThread;
+				const int32 ShaderIndexEnd = FMath::Min(ShaderIndexStart + kShaderSortedPerThread, ShaderToShadermapsArray.Num());
+				for (int32 ShaderIndex = ShaderIndexStart; ShaderIndex < ShaderIndexEnd; ++ShaderIndex)
+				{
+					ShaderToShadermapsArray[ShaderIndex].ShaderMapIndices.Sort();
+				}
+			},
+			EParallelForFlags::Unbalanced
+		);
+	}
+
+	// now assigning the indices in the array so we can sort it
+	for (int32 ShaderIndex = 0, Num = ShaderToShadermapsArray.Num(); ShaderIndex < Num; ++ShaderIndex)
+	{
+		// check that no shader is unreferenced
+		checkf(!ShaderToShadermapsArray[ShaderIndex].ShaderMapIndices.IsEmpty(),
+			TEXT("Error converting to IoStore archive: shader (index=%d) is not referenced by any of the shadermaps!"), ShaderIndex);
+		ShaderToShadermapsArray[ShaderIndex].ShaderIndex = ShaderIndex;
+	}
+
+	// sort the mapping so the first are shaders that are referenced by a smaller number of shadermaps, then by index for determinism
+	Algo::Sort(ShaderToShadermapsArray, FShaderIndexToShaderMapIndicesPair::Sort);
+}
+
+void FShaderCodeArchiveProcessor::SplitShaderGroupForRaytracing(TArray<uint32>& CurrentShaderGroup)
+{
+	if (!bSeparateRaytracingShaders)
+	{
+		SplitShaderGroupBySize(CurrentShaderGroup);
+	}
+	else
+	{
+		// The streak changed. Create the group, but first, determine if the group needs to be split in two because of the raytracing shaders.
+		// We want to isolate them into separate groups so their preload can be skipped if raytracing is off.
+		TArray<uint32> RaytracingShaders;
+		TArray<uint32> NonraytracingShaders;
+		for (int32 ShaderIndex : CurrentShaderGroup)
+		{
+			if (LIKELY(!IsRayTracingShaderFrequency(static_cast<EShaderFrequency>(OutHeader.ShaderEntries[ShaderIndex].Frequency))))
+			{
+				NonraytracingShaders.Add(ShaderIndex);
+			}
+			else
+			{
+				RaytracingShaders.Add(ShaderIndex);
+			}
+		}
+		check(CurrentShaderGroup.Num() == NonraytracingShaders.Num() + RaytracingShaders.Num());
+
+		if (LIKELY(!NonraytracingShaders.IsEmpty()))
+		{
+			SplitShaderGroupBySize(NonraytracingShaders);
+		}
+		if (UNLIKELY(!RaytracingShaders.IsEmpty()))
+		{
+			SplitShaderGroupBySize(RaytracingShaders);
+		}
+	}
+}
+
+void FShaderCodeArchiveProcessor::SplitShaderGroupBySize(TArray<uint32>& CurrentShaderGroup)
+{
+	// calculate current group size
+	uint32 GroupSize = 0;
+	const TArrayView<const FShaderCodeEntry> ShaderEntries = SerializedShaders->GetShaderEntries();
+	for (uint32 ShaderIdx : CurrentShaderGroup)
+	{
+		GroupSize += ShaderEntries[ShaderIdx].UncompressedSize;
+	}
+
+	if (LIKELY(GroupSize <= MaxUncompressedShaderGroupSize || CurrentShaderGroup.Num() == 1))
+	{
+		AddNewShaderGroup(CurrentShaderGroup);
+	}
+	else
+	{
+		// split the shaders evenly into N new groups (don't allow more new groups than there are shaders)
+		int32 NumNewGroups = FMath::Min(static_cast<int32>(GroupSize / MaxUncompressedShaderGroupSize + 1), CurrentShaderGroup.Num());
+		checkf(NumNewGroups > 1, TEXT("Off by one error in group count calculation? NumNewGroups=%d, GroupSize=%u, MaxUncompressedShaderGroupSize=%u, CurrentShaderGroup.Num()=%d"), NumNewGroups, GroupSize, MaxUncompressedShaderGroupSize, CurrentShaderGroup.Num());
+
+		TArray<TArray<uint32>> NewGroups;
+		TArray<uint32> NewGroupSizes;
+		NewGroups.AddDefaulted(NumNewGroups);
+		NewGroupSizes.AddZeroed(NumNewGroups);
+
+		// sort the shaders descending, as this is easier to split (greedy algorithm)
+		CurrentShaderGroup.Sort(
+			[this](const int32 ShaderIndexA, const int32 ShaderIndexB)
+			{
+				const TArrayView<const FShaderCodeEntry> ShaderEntries = SerializedShaders->GetShaderEntries();
+				const FShaderCodeEntry& ShaderEntryA = ShaderEntries[ShaderIndexA];
+				const FShaderCodeEntry& ShaderEntryB = ShaderEntries[ShaderIndexB];
+				if (ShaderEntryA.UncompressedSize != ShaderEntryB.UncompressedSize)
+				{
+					return ShaderEntryA.UncompressedSize > ShaderEntryB.UncompressedSize;
+				}
+				if (ShaderEntryA.Size != ShaderEntryB.Size)
+				{
+					return ShaderEntryA.Size > ShaderEntryB.Size;
+				}
+				if (ShaderEntryA.Frequency != ShaderEntryB.Frequency)
+				{
+					return ShaderEntryA.Frequency > ShaderEntryB.Frequency;
+				}
+				return ShaderEntryA.Offset > ShaderEntryB.Offset;
+			}
+		);
+
+		for (int32 ShaderIdx : CurrentShaderGroup)
+		{
+			// add the shader to the group of smallest size
+			int32 SmallestNewGroupIdx = 0;
+			for (int32 IdxNewGroup = 1; IdxNewGroup < NumNewGroups; ++IdxNewGroup)
+			{
+				if (NewGroupSizes[IdxNewGroup] < NewGroupSizes[SmallestNewGroupIdx])
+				{
+					SmallestNewGroupIdx = IdxNewGroup;
+				}
+			}
+
+			NewGroups[SmallestNewGroupIdx].Add(ShaderIdx);
+			NewGroupSizes[SmallestNewGroupIdx] += ShaderEntries[ShaderIdx].UncompressedSize;
+		}
+
+#if DO_CHECK // sanity checks
+		uint32 NewGroupTotalSize = 0;
+		for (uint32 NewGroupSize : NewGroupSizes)
+		{
+			NewGroupTotalSize += NewGroupSize;
+		}
+		checkf(NewGroupTotalSize == GroupSize, TEXT("Original shader group was of size %u bytes, which was larger than limit %u, but it was split into %d group of total size %u, which is not %u - sizes must agree"),
+			GroupSize,
+			MaxUncompressedShaderGroupSize,
+			NumNewGroups,
+			NewGroupTotalSize, GroupSize
+		);
+#endif
+
+		for (TArray<uint32>& NewGroup : NewGroups)
+		{
+			// note there can be empty groups (take a very edge case of MaxUncompressedShaderGroupSize = 2 bytes and a shader group of 1 shader)
+			if (!NewGroup.IsEmpty())
+			{
+				AddNewShaderGroup(NewGroup);
+			}
+		}
+	}
+}
+
+void FShaderCodeArchiveProcessor::AddNewShaderGroup(TArray<uint32>& ShaderIndicesInGroup)
+{
+	// first, sort the shaders by uncompressed size, as this was found to compress better
+	ShaderIndicesInGroup.Sort(
+		[this](const int32 ShaderIndexA, const int32 ShaderIndexB)
+		{
+			const FShaderCodeEntry& ShaderEntryA = SerializedShaders->GetShaderEntries()[ShaderIndexA];
+			const FShaderCodeEntry& ShaderEntryB = SerializedShaders->GetShaderEntries()[ShaderIndexB];
+			if (ShaderEntryA.UncompressedSize != ShaderEntryB.UncompressedSize)
+			{
+				return ShaderEntryA.UncompressedSize < ShaderEntryB.UncompressedSize;
+			}
+			if (ShaderEntryA.Size != ShaderEntryB.Size)
+			{
+				return ShaderEntryA.Size < ShaderEntryB.Size;
+			}
+			if (ShaderEntryA.Frequency != ShaderEntryB.Frequency)
+			{
+				return ShaderEntryA.Frequency < ShaderEntryB.Frequency;
+			}
+			return ShaderEntryA.Offset < ShaderEntryB.Offset;
+		}
+	);
+
+	// add a new group entry
+	const int32 CurrentGroupIdx = OutHeader.ShaderGroupEntries.Num();
+	FIoStoreShaderGroupEntry& GroupEntry = OutHeader.ShaderGroupEntries.AddDefaulted_GetRef();
+	StoredGroupShaderIndices.Add(ShaderIndicesInGroup);
+	GroupEntry.NumShaders = ShaderIndicesInGroup.Num();
+	// ShaderIndicesOffset will be filled later, once we know all the groups (see comment about StoredGroupShaderIndices above).
+
+	const TArrayView<const FSHAHash> ShaderHashes = SerializedShaders->GetShaderHashes();
+	const TArrayView<const FShaderCodeEntry> ShaderEntries = SerializedShaders->GetShaderEntries();
+
+	// update shader entries both with the group number and their uncompressed offset in the group
+	FSHA1 GroupHasher;
+	uint32 CurrentGroupSize = 0;
+	for (int32 ShaderIdxIdx = 0, NumIdxIdx = ShaderIndicesInGroup.Num(); ShaderIdxIdx < NumIdxIdx; ++ShaderIdxIdx)
+	{
+		int32 ShaderIndex = ShaderIndicesInGroup[ShaderIdxIdx];
+		FIoStoreShaderCodeEntry& IoStoreShaderEntry = OutHeader.ShaderEntries[ShaderIndex];
+		IoStoreShaderEntry.ShaderGroupIndex = CurrentGroupIdx;
+		IoStoreShaderEntry.UncompressedOffsetInGroup = CurrentGroupSize;
+
+		// group hash is constructed from hashing the shaders in the group.
+		GroupHasher.Update(ShaderHashes[ShaderIndex].Hash, sizeof(FSHAHash));
+		// shader hash as of now excludes optional data, so we cannot rely on it, especially across the shader formats. Make the group hash a bit more robust by including the shader size in it.
+		GroupHasher.Update(reinterpret_cast<const uint8*>(&ShaderEntries[ShaderIndex].UncompressedSize), sizeof(FShaderCodeEntry::UncompressedSize));
+
+		CurrentGroupSize += ShaderEntries[ShaderIndex].UncompressedSize;
+	}
+
+	// Shader hashes cannot be used to uniquely identify across shader formats due to aforementioned exclusion of optional data from it.
+	// Include the shader format (in a platform-agnostic way) into the group hash to lower the risk of collision of shaders of different formats.
+	GroupHasher.Update(reinterpret_cast<const uint8*>(Utf8FormatName.Get()), Utf8FormatName.Length());
+	static_assert(sizeof(uint8) == sizeof(UTF8CHAR), "Unexpected UTF-8 char size.");
+
+	GroupEntry.UncompressedSize = CurrentGroupSize;
+	OutHeader.ShaderGroupIoHashes.Add(FIoStoreShaderCodeArchive::GetShaderCodeChunkId(GroupHasher.Finalize()));
+
+	Stats.TotalUncompressedMemory += CurrentGroupSize;
+	Stats.MinGroupSize = FMath::Min(Stats.MinGroupSize, CurrentGroupSize);
+	Stats.MaxGroupSize = FMath::Max(Stats.MaxGroupSize, CurrentGroupSize);
+}
+
+void FIoStoreShaderCodeArchive::CreateIoStoreShaderCodeArchiveHeader(const FName& Format, const FSerializedShaderArchive& SerializedShaders, FIoStoreShaderCodeArchiveHeader& OutHeader)
+{
+	// Higher level description of the group splitting algo that follows:
+	// We compress together shaders that are loaded together (all other strategies, like grouping similar shaders, were found to compress better but not reduce RAM usage).
+	// For that, we find for each shader which shadermaps are referencing it (often times it will be just one, but for some simple and shared shaders it can be thousands).
+	// We group the shaders by those sets of shadermaps - all shaders referenced by the same shadermap(s) are a candidate for being a single group. Then we potentially split this candidate group
+	// into raytracing and non-raytracing groups (so we can avoid preloading RTX shaders run-time if RTX is disabled), and then each of those is potentially split further by size
+	// (to avoid too large groups that will take too much time to decompress - this is regulated by r.ShaderCodeLibrary.MaxShaderGroupSize). The results of that process (note, it can still be
+	// a single group) is added to the header.
+	// Each group's indices, like in case of shadermaps, are stored in ShaderIndices array. Before we append a new group's indices however, we look if we can find an existing range that we can reuse.
+	FShaderCodeArchiveProcessor ArchiveProcessor(Format, OutHeader);
+	ArchiveProcessor.ProcessSerializedShaders(SerializedShaders);
 }
 
 FArchive& operator <<(FArchive& Ar, FIoStoreShaderCodeArchiveHeader& Ref)
diff --git a/Engine/Source/Runtime/RenderCore/Public/ShaderCodeArchive.h b/Engine/Source/Runtime/RenderCore/Public/ShaderCodeArchive.h
--- a/Engine/Source/Runtime/RenderCore/Public/ShaderCodeArchive.h
+++ b/Engine/Source/Runtime/RenderCore/Public/ShaderCodeArchive.h
@@ -497,7 +497,7 @@ namespace ShaderCodeArchive
 /** Descriptor of a shader map. This concept exists in run time, so this class describes the information stored in the library for a particular FShaderMap */
 struct FIoStoreShaderMapEntry
 {
-	/** Offset to an the first shader index referenced by this shader map in the array of shader indices. */
+	/** Offset to the first shader index referenced by this shader map in the array of shader indices. */
 	uint32 ShaderIndicesOffset = 0u;
 	/** Number of shaders in this shader map. */
 	uint32 NumShaders = 0u;
