[ShaderLibrary] Split ShaderCodeLibrary.cpp into smaller source files.

- Move `FEditorShaderCodeArchive`, `FEditorShaderStableInfo`, `FShaderLibraryInstance`, and `FShaderMapResource_SharedCode` into separate source files;
This reduced ShaderCodeLibrary.cpp file size from ~5k to under 3k LOC.
- Move all shader library relevant source files into RenderCore/Private/ShaderLibrary/ folder.
- Move shared functions, variables, and CVars for new source files into ShaderCodeLibraryUtilities.h/.cpp files.
- Add documentation to all functions in `FEditorShaderCodeArchive`.
- Fix compiler warnings by changing `& bOk` to `&& bOk` at various places.
- Make source archive parameter constant in `FEditorShaderCodeArchive::AppendFromArchive()` as this doesn't need to take ownership.

#rnx
#Synced-CL 48657858

diff --git a/Engine/Source/Runtime/RenderCore/Private/ShaderLibrary/EditorShaderCodeArchive.cpp b/Engine/Source/Runtime/RenderCore/Private/ShaderLibrary/EditorShaderCodeArchive.cpp
new file mode 100644
--- /dev/null
+++ b/Engine/Source/Runtime/RenderCore/Private/ShaderLibrary/EditorShaderCodeArchive.cpp
@@ -0,0 +1,1079 @@
+// Copyright Epic Games, Inc. All Rights Reserved.
+
+/*=============================================================================
+	EditorShaderCodeArchive.cpp: Implementation of FEditorShaderCodeArchive
+=============================================================================*/
+
+#if WITH_EDITOR
+
+#include "EditorShaderCodeArchive.h"
+
+#include "Async/ParallelFor.h"
+#include "HAL/FileManager.h"
+#include "Interfaces/IShaderFormat.h"
+#include "Interfaces/ITargetPlatformManagerModule.h"
+#include "Internationalization/Regex.h"
+#include "Math/UnitConversion.h"
+#include "RHIStrings.h"
+#include "PipelineCacheUtilities.h"
+#include "Serialization/CompactBinarySerialization.h"
+#include "Serialization/CompactBinaryWriter.h"
+#include "Shader.h"
+#include "ShaderCodeLibrary.h"
+#include "ShaderCompilerCore.h"
+#include "ShaderLibrary/ShaderCodeLibraryUtilities.h"
+#include "ShaderPipelineCache.h"
+#include "UObject/Class.h"
+#include "UObject/Object.h"
+
+using namespace UE::ShaderLibrary::Private;
+
+static const IShaderFormat* GetShaderFormatFromName(FName InFormatName)
+{
+	// Parse library format name into shader format name, e.g. "SF_VULKAN_ES31_ANDROID-VULKAN_ES3_1_ANDROID" to "SF_VULKAN_ES31_ANDROID".
+	TArray<FString> Components;
+	const FString Name = InFormatName.ToString();
+	Name.ParseIntoArray(Components, TEXT("-"));
+
+	if (Components.IsEmpty())
+	{
+		return nullptr;
+	}
+
+	const FName ShaderFormatName(Components[0]);
+	return GetTargetPlatformManagerRef().FindShaderFormat(ShaderFormatName);
+}
+
+FEditorShaderCodeArchive::FEditorShaderCodeArchive(FName InFormat, bool bInNeedsDeterministicOrder)
+	: FormatName(InFormat)
+	, Format(GetShaderFormatFromName(InFormat))
+	, bNeedsDeterministicOrder(bInNeedsDeterministicOrder)
+{
+	checkf(Format, TEXT("Failed to parse shader format from name '%s'"), *InFormat.ToString());
+
+	SerializedShaders.ShaderHashTable.Clear(0x10000);
+	SerializedShaders.ShaderMapHashTable.Clear(0x10000);
+}
+
+FEditorShaderCodeArchive::~FEditorShaderCodeArchive()
+{
+	// dummy
+}
+
+void FEditorShaderCodeArchive::OpenLibrary(FString const& Name)
+{
+	check(LibraryName.Len() == 0);
+	check(Name.Len() > 0);
+	LibraryName = Name;
+	SerializedShaders.Empty();
+	ShaderCode.Empty();
+}
+
+void FEditorShaderCodeArchive::CloseLibrary(FString const& Name)
+{
+	check(LibraryName == Name);
+	LibraryName = TEXT("");
+}
+
+int32 FEditorShaderCodeArchive::AddShaderCode(FShaderCodeStats& CodeStats, const FShaderMapResourceCode* Code, const FShaderMapAssetPaths& AssociatedAssets)
+{
+	int32 ShaderMapIndex = INDEX_NONE;
+
+	if (AssociatedAssets.Num() == 0 && LibraryName != TEXT("Global"))
+	{
+		UE_LOG(LogShaderLibrary, Warning, TEXT("Shadermap %s does not have assets associated with it, library layout may be inconsistent between builds"), *Code->ResourceHash.ToString());
+	}
+
+	const int32 NumShaders = Code->ShaderCodeResources.Num();
+	if (SerializedShaders.FindOrAddShaderMap(Code->ResourceHash, ShaderMapIndex, &AssociatedAssets))
+	{
+		CodeStats.NumShaderMaps++;
+		FShaderMapEntry& ShaderMapEntry = SerializedShaders.ShaderMapEntries[ShaderMapIndex];
+		ShaderMapEntry.NumShaders = NumShaders;
+		ShaderMapEntry.ShaderIndicesOffset = SerializedShaders.ShaderIndices.AddZeroed(NumShaders);
+
+		for (int32 i = 0; i < NumShaders; ++i)
+		{
+			int32 ShaderIndex = INDEX_NONE;
+			if (SerializedShaders.FindOrAddShader(Code->ShaderHashes[i], ShaderIndex))
+			{
+				const FShaderCodeResource& SourceShaderResource = Code->ShaderCodeResources[i];
+				FShaderCodeEntry& SerializedShaderEntry = SerializedShaders.ShaderEntries[ShaderIndex];
+				SerializedShaderEntry.Frequency = SourceShaderResource.GetFrequency();
+				FSharedBuffer CodeBuffer = SourceShaderResource.GetCodeBuffer();
+				SerializedShaderEntry.Size = CodeBuffer.GetSize();
+				SerializedShaderEntry.UncompressedSize = SourceShaderResource.GetUncompressedSize();
+				check(SerializedShaderEntry.Size > 0);
+				ShaderCode.Add(CodeBuffer);
+				check(ShaderCode.Num() == SerializedShaders.ShaderEntries.Num());
+			}
+			SerializedShaders.ShaderIndices[ShaderMapEntry.ShaderIndicesOffset + i] = ShaderIndex;
+		}
+	}
+
+	for (int32 i = 0; i < NumShaders; ++i)
+	{
+		for (uint64 ShaderTypeHash : Code->ShaderEditorOnlyDataEntries[i].ShaderTypeHashes)
+		{
+			int32 ShaderIndex = SerializedShaders.FindShader(Code->ShaderHashes[i]);
+
+			TArray<uint64>& SerializedShaderTypes = SerializedShaders.ShaderTypes[ShaderIndex].Data;
+			const int32 HashIndex = Algo::LowerBound(SerializedShaderTypes, ShaderTypeHash);
+			if (HashIndex >= SerializedShaderTypes.Num() || SerializedShaderTypes[HashIndex] != ShaderTypeHash)
+			{
+				SerializedShaderTypes.Insert(ShaderTypeHash, HashIndex);
+				FShaderCodeStats::FPerTypeStats& TypeCodeStats = CodeStats.ShaderTypeStats.FindOrAdd(ShaderTypeHash);
+				int32 CodeSize = SerializedShaders.ShaderEntries[ShaderIndex].Size;
+				TypeCodeStats.NumShaders++;
+				TypeCodeStats.ShadersSize += CodeSize;
+			}
+		}
+	}
+
+	// always mark the shadermap dirty, because it might have gotten new asset associations
+	MarkShaderMapDirty(ShaderMapIndex);
+	return ShaderMapIndex;
+}
+
+void FEditorShaderCodeArchive::MarkShaderMapDirty(int32 ShaderMapIndex)
+{
+	if (bHasCopiedAndCleared)
+	{
+		ShaderMapsToCopy.Add(ShaderMapIndex);
+	}
+}
+
+void FEditorShaderCodeArchive::MarkAllShaderMapsDirty()
+{
+	if (bHasCopiedAndCleared)
+	{
+		for (int32 ShaderMapIndex = 0; ShaderMapIndex < SerializedShaders.ShaderMapHashes.Num(); ++ShaderMapIndex)
+		{
+			ShaderMapsToCopy.Add(ShaderMapIndex);
+		}
+	}
+}
+
+bool FEditorShaderCodeArchive::HasDataToCopy() const
+{
+	// After the first copy dirty elements are added to ShaderMapsToCopy; during the first copy every ShaderMap in SerializedShaders is copied
+	return bHasCopiedAndCleared ? !ShaderMapsToCopy.IsEmpty() : !SerializedShaders.ShaderMapHashes.IsEmpty();
+}
+
+void FEditorShaderCodeArchive::CopyToCompactBinary(FCbWriter& Writer, FSerializedShaderArchive& TransferArchive, TArray<uint8>& TransferCode)
+{
+	Writer.BeginObject();
+	Writer << "SerializedShaders" << TransferArchive;
+	Writer << "ShaderCode";
+	Writer.AddBinary(FMemoryView(TransferCode.GetData(), TransferCode.Num()));
+	Writer.EndObject();
+}
+
+bool FEditorShaderCodeArchive::AppendFromCompactBinary(FCbFieldView Field, FShaderCodeStats& CodeStats)
+{
+	FSerializedShaderArchive TransferArchive;
+	if (!LoadFromCompactBinary(Field["SerializedShaders"], TransferArchive))
+	{
+		return false;
+	}
+
+	FMemoryView FlatShaderCodeView = Field["ShaderCode"].AsBinaryView();
+	TConstArrayView64<uint8> TransferCode(reinterpret_cast<const uint8*>(FlatShaderCodeView.GetData()), FlatShaderCodeView.GetSize());
+	return AppendFromArchive(TransferArchive, TransferCode, CodeStats);
+}
+
+void FEditorShaderCodeArchive::CopyToArchiveAndClear(FSerializedShaderArchive& TargetArchive, TArray<uint8>& TargetFlatShaderCode,
+	bool& bOutRanOutOfRoom, int64 MaxShaderSize, int64 MaxShaderCount)
+{
+	bOutRanOutOfRoom = false;
+	if (!bHasCopiedAndCleared)
+	{
+		// First CopyAndClear adds all shadermaps to ShaderMapsToCopy, because as an optimization we do not write to
+		// ShaderMapsToCopy until the first CopyAndClear call.
+		for (int32 ShaderMapIndex = 0; ShaderMapIndex < SerializedShaders.ShaderMapHashes.Num(); ++ShaderMapIndex)
+		{
+			ShaderMapsToCopy.Add(ShaderMapIndex);
+		}
+	}
+
+	TArray<int32> LocalShaderMapsToCopy = ShaderMapsToCopy.Array();
+	LocalShaderMapsToCopy.Sort(); // Maintain the same order that the shadermaps were added in
+	int32 NumShadersSentWithCode = 0;
+
+	const FSerializedShaderArchive& SourceArchive = this->SerializedShaders;
+	TArray<FSharedBuffer>& SourceShaderCodes = this->ShaderCode;
+	for (int32 SourceShaderMapIndex : LocalShaderMapsToCopy)
+	{
+		const FSHAHash& SourceShaderMapHash = SourceArchive.ShaderMapHashes[SourceShaderMapIndex];
+		const FShaderMapEntry& SourceEntry = SourceArchive.ShaderMapEntries[SourceShaderMapIndex];
+		const FShaderMapAssetPaths* const SourceAssetPaths = SourceArchive.ShaderCodeToAssets.Find(SourceShaderMapHash);
+
+		int32 TargetShaderMapIndex;
+		const bool bIsNewShaderMap = TargetArchive.FindOrAddShaderMap(SourceShaderMapHash, TargetShaderMapIndex, SourceAssetPaths);
+		if (!bIsNewShaderMap)
+		{
+			continue;
+		}
+		FShaderMapEntry& TargetEntry = TargetArchive.ShaderMapEntries[TargetShaderMapIndex];
+		const int32 NumShaders = SourceEntry.NumShaders;
+
+		TargetEntry.NumShaders = NumShaders;
+		TargetEntry.ShaderIndicesOffset = TargetArchive.ShaderIndices.Num();
+		TargetArchive.ShaderIndices.AddUninitialized(NumShaders);
+
+		for (int32 ShaderIndexIndex = 0; ShaderIndexIndex < NumShaders; ++ShaderIndexIndex)
+		{
+			const uint32 SourceShaderIndex = SourceArchive.ShaderIndices[SourceEntry.ShaderIndicesOffset + ShaderIndexIndex];
+			const FSHAHash& SourceShaderHash = SourceArchive.ShaderHashes[SourceShaderIndex];
+			int32 TargetShaderIndex;
+			const bool bIsNewShader = TargetArchive.FindOrAddShader(SourceShaderHash, TargetShaderIndex);
+			TargetArchive.ShaderIndices[TargetEntry.ShaderIndicesOffset + ShaderIndexIndex] = TargetShaderIndex;
+			if (bIsNewShader)
+			{
+				// We rely on the index of the newly added shader being at the end of the list of ShaderEntries,
+				// so we can pop off the added index if we overflow below
+				check(TargetShaderIndex == TargetArchive.ShaderEntries.Num() - 1);
+
+				const FShaderCodeEntry& SourceShaderEntry = SourceArchive.ShaderEntries[SourceShaderIndex];
+				FSharedBuffer& SourceShaderCode = SourceShaderCodes[SourceShaderIndex];
+				FShaderCodeEntry& TargetShaderEntry = TargetArchive.ShaderEntries[TargetShaderIndex];
+
+				TargetShaderEntry = SourceShaderEntry;
+				if (SourceShaderCode.GetSize() > 0)
+				{
+					TargetShaderEntry.Offset = TargetFlatShaderCode.Num();
+					check(SourceShaderEntry.Size == SourceShaderCode.GetSize());
+					if ((MaxShaderSize > 0 && TargetFlatShaderCode.Num() + (int64)SourceShaderCode.GetSize() > MaxShaderSize) ||
+						(MaxShaderCount > 0 && NumShadersSentWithCode > MaxShaderCount))
+					{
+						// We have to stop here to avoid overflowing the shader limit. Send the shaders we have accumulated
+						// but do not send&clear any other data.
+						// Remove all ShaderMap data
+						TargetArchive.EmptyShaderMaps();
+						// Remove the ShaderEntry we just added; we are not adding its code so we have to remove it from the
+						// list of shaders contained by the targetarchive
+						check(TargetShaderIndex == TargetArchive.ShaderHashes.Num() - 1);
+						TargetArchive.RemoveLastAddedShader();
+						bOutRanOutOfRoom = true;
+						// Keep any shaders we added before the one we just added
+						return;
+					}
+					++NumShadersSentWithCode;
+
+					TargetFlatShaderCode.Append(MakeArrayView(reinterpret_cast<const uint8*>(SourceShaderCode.GetData()), SourceShaderCode.GetSize()));
+
+					// Reset the ShaderCode reference to (potentially) save memory in the local process. The consumer of the TargetArchive and
+					// TargetFlatShaderCode will be the only one that needs to read it.
+					SourceShaderCode.Reset();
+				}
+				else
+				{
+					// The shadercode was already copied in an earlier CopyAndClear operation; we just need to note that
+					// ShaderMaps in this call to CopyAndClear reference it.
+					TargetShaderEntry.Offset = INDEX_NONE;
+				}
+			}
+
+			for (uint64 ShaderTypeHash : SourceArchive.ShaderTypes[SourceShaderIndex].Data)
+			{
+				TArray<uint64>& TargetShaderTypes = TargetArchive.ShaderTypes[TargetShaderIndex].Data;
+				const int32 HashIndex = Algo::LowerBound(TargetShaderTypes, ShaderTypeHash);
+				if (HashIndex >= TargetShaderTypes.Num() || TargetShaderTypes[HashIndex] != ShaderTypeHash)
+				{
+					TargetShaderTypes.Insert(ShaderTypeHash, HashIndex);
+				}
+			}
+		}
+	}
+
+	ShaderMapsToCopy.Empty();
+	bHasCopiedAndCleared = true;
+}
+
+bool FEditorShaderCodeArchive::AppendFromArchive(
+	const FSerializedShaderArchive& SourceArchive, TConstArrayView64<uint8> SourceFlatShaderCode, FShaderCodeStats& CodeStats)
+{
+	bool bOk = true;
+	FSerializedShaderArchive& TargetArchive = this->SerializedShaders;
+	TArray<FSharedBuffer>& TargetShaderCodes = this->ShaderCode;
+
+	// Add all the shaders; we can sometimes get messages that send the shaders in advance without sending the shadermaps that use them
+	for (int32 SourceShaderIndex = 0; SourceShaderIndex < SourceArchive.ShaderHashes.Num(); ++SourceShaderIndex)
+	{
+		const FSHAHash& SourceShaderHash = SourceArchive.ShaderHashes[SourceShaderIndex];
+		int32 TargetShaderIndex = INDEX_NONE;
+		const bool bShaderIsNew = TargetArchive.FindOrAddShader(SourceShaderHash, TargetShaderIndex);
+		const TArray<uint64>& ShaderTypeHashes = SourceArchive.ShaderTypes[SourceShaderIndex].Data;
+		const FShaderCodeEntry& SourceShaderEntry = SourceArchive.ShaderEntries[SourceShaderIndex];
+
+		TArray<uint64>& ExistingShaderTypeHashes = TargetArchive.ShaderTypes[TargetShaderIndex].Data;
+		for (uint64 ShaderTypeHash : ShaderTypeHashes)
+		{
+			const int32 HashIndex = Algo::LowerBound(ExistingShaderTypeHashes, ShaderTypeHash);
+			if (HashIndex >= ExistingShaderTypeHashes.Num() || ExistingShaderTypeHashes[HashIndex] != ShaderTypeHash)
+			{
+				ExistingShaderTypeHashes.Insert(ShaderTypeHash, HashIndex);
+				FShaderCodeStats::FPerTypeStats& TypeCodeStats = CodeStats.ShaderTypeStats.FindOrAdd(ShaderTypeHash);
+				TypeCodeStats.NumShaders++;
+				TypeCodeStats.ShadersSize += SourceShaderEntry.Size;
+			}
+		}
+
+		if (!bShaderIsNew)
+		{
+			continue;
+		}
+		check(TargetShaderIndex == TargetArchive.ShaderEntries.Num() - 1 && TargetShaderCodes.Num() == TargetArchive.ShaderEntries.Num() - 1);
+
+		FShaderCodeEntry& TargetShaderEntry = TargetArchive.ShaderEntries[TargetShaderIndex];
+		TargetShaderEntry = SourceShaderEntry;
+		TargetShaderEntry.Offset = 0;
+
+		if (SourceShaderEntry.Offset == INDEX_NONE)
+		{
+			UE_LOG(LogShaderLibrary, Error, TEXT("ShaderMapLibrary transfer received from a remote machine has incomplete record for Shader %s. ")
+				TEXT("The remote machine thought this machine already had the shader and did not send the ShaderCode for the Shader, but the shader is not found. ")
+				TEXT("The ShaderMaps using the shader will be corrupt."),
+				*SourceShaderHash.ToString());
+			bOk = false;
+			TargetShaderEntry.Size = 0;
+		}
+		else if (SourceShaderEntry.Offset + SourceShaderEntry.Size > static_cast<uint64>(SourceFlatShaderCode.Num()))
+		{
+			UE_LOG(LogShaderLibrary, Error, TEXT("ShaderMapLibrary transfer received from a remote machine has corrupt record for Shader %s. ")
+				TEXT("The (Offset, Size) specified by the ShaderEntry does not fit in the TransferCode (an array of bytes that should contain the shader code for all transferred shaders. ")
+				TEXT("The ShaderMaps using the shader will be corrupt."),
+				*SourceShaderHash.ToString());
+			bOk = false;
+			TargetShaderEntry.Size = 0;
+		}
+		else
+		{
+			// Copy from source's flat list to the target's separate FSharedBuffer for each shader
+			TargetShaderCodes.Add(FSharedBuffer::Clone(SourceFlatShaderCode.GetData() + SourceShaderEntry.Offset, SourceShaderEntry.Size));
+		}
+	}
+
+	// Add all the shadermaps
+	for (int32 SourceShaderMapIndex = 0; SourceShaderMapIndex < SourceArchive.ShaderMapHashes.Num(); ++SourceShaderMapIndex)
+	{
+		const FSHAHash& SourceShaderMapHash = SourceArchive.ShaderMapHashes[SourceShaderMapIndex];
+		const FShaderMapAssetPaths* const SourceAssetPaths = SourceArchive.ShaderCodeToAssets.Find(SourceShaderMapHash);
+		int32 TargetShaderMapIndex;
+		const bool bIsNewShaderMap = TargetArchive.FindOrAddShaderMap(SourceShaderMapHash, TargetShaderMapIndex, SourceAssetPaths);
+		// always mark the shadermap dirty, because it might have gotten new asset associations
+		MarkShaderMapDirty(TargetShaderMapIndex);
+		if (!bIsNewShaderMap)
+		{
+			// ShaderMap has already been loaded in this process, e.g. from another CookWorker loading the same Material
+			continue;
+		}
+
+		const FShaderMapEntry& SourceEntry = SourceArchive.ShaderMapEntries[SourceShaderMapIndex];
+		FShaderMapEntry& TargetEntry = TargetArchive.ShaderMapEntries[TargetShaderMapIndex];
+		const int32 NumShaders = SourceEntry.NumShaders;
+
+		TargetEntry.NumShaders = NumShaders;
+		TargetEntry.ShaderIndicesOffset = TargetArchive.ShaderIndices.Num();
+		TargetArchive.ShaderIndices.AddUninitialized(NumShaders);
+
+		for (int32 ShaderIndexIndex = 0; ShaderIndexIndex < NumShaders; ++ShaderIndexIndex)
+		{
+			const uint32 SourceShaderIndex = SourceArchive.ShaderIndices[SourceEntry.ShaderIndicesOffset + ShaderIndexIndex];
+			const FSHAHash& SourceShaderHash = SourceArchive.ShaderHashes[SourceShaderIndex];
+			int32 TargetShaderIndex;
+			const bool bShaderIsNew = TargetArchive.FindOrAddShader(SourceShaderHash, TargetShaderIndex);
+			TargetArchive.ShaderIndices[TargetEntry.ShaderIndicesOffset + ShaderIndexIndex] = TargetShaderIndex;
+			// Every shader in the SourceArchive should have already been added by the loop above over SourceArchive.ShaderHashes
+			check(!bShaderIsNew);
+		}
+		CodeStats.NumShaderMaps++;
+	}
+	return bOk;
+}
+
+FEditorShaderCodeArchive* FEditorShaderCodeArchive::CreateChunk(int ChunkId, const TSet<FName>& PackagesInChunk)
+{
+	checkf(!bHasCopiedAndCleared, TEXT("It is not valid to call CreateChunk on an FEditorShaderCodeArchive that has sent its shaders to another process by calling CopyAndClear."));
+
+	FEditorShaderCodeArchive* NewChunk = new FEditorShaderCodeArchive(FormatName, bNeedsDeterministicOrder);
+	NewChunk->OpenLibrary(UE::ShaderLibrary::Private::GetShaderLibraryNameForChunk(LibraryName, ChunkId));
+
+	TArray<int32> ShaderCodeEntriesNeeded;	// this array is filled with the indices from the existing ShaderCode that will need to be taken
+	NewChunk->SerializedShaders.CreateAsChunkFrom(SerializedShaders, PackagesInChunk, ShaderCodeEntriesNeeded);
+	// extra integrity check
+	checkf(ShaderCodeEntriesNeeded.Num() == NewChunk->SerializedShaders.ShaderHashes.Num(), TEXT("FSerializedShaderArchive for the new chunk did not create a valid shader code mapping"));
+	checkf(ShaderCodeEntriesNeeded.Num() == NewChunk->SerializedShaders.ShaderEntries.Num(), TEXT("FSerializedShaderArchive for the new chunk did not create a valid shader code mapping"));
+
+	// copy the shader code
+	NewChunk->ShaderCode.Empty();
+	for (int32 NewArchiveIdx = 0, NumIndices = ShaderCodeEntriesNeeded.Num(); NewArchiveIdx < NumIndices; ++NewArchiveIdx)
+	{
+		FSharedBuffer& SourceShaderCodeEntry = ShaderCode[ShaderCodeEntriesNeeded[NewArchiveIdx]];
+		check(SourceShaderCodeEntry.GetSize() > 0);
+		NewChunk->ShaderCode.Add(SourceShaderCodeEntry);
+	}
+
+	return NewChunk;
+}
+
+int32 FEditorShaderCodeArchive::AddShaderCode(int32 OtherShaderMapIndex, const FEditorShaderCodeArchive& OtherArchive)
+{
+	checkf(!OtherArchive.bHasCopiedAndCleared,
+		TEXT("It is not valid to call AddShaderCode from an FEditorShaderCodeArchive that has sent its shaders to another process by calling CopyAndClear."));
+
+	int32 ShaderMapIndex = 0;
+	if (SerializedShaders.FindOrAddShaderMap(OtherArchive.SerializedShaders.ShaderMapHashes[OtherShaderMapIndex], ShaderMapIndex,
+		OtherArchive.SerializedShaders.ShaderCodeToAssets.Find(OtherArchive.SerializedShaders.ShaderMapHashes[OtherShaderMapIndex])))
+	{
+		const FShaderMapEntry& PrevShaderMapEntry = OtherArchive.SerializedShaders.ShaderMapEntries[OtherShaderMapIndex];
+		FShaderMapEntry& ShaderMapEntry = SerializedShaders.ShaderMapEntries[ShaderMapIndex];
+		ShaderMapEntry.NumShaders = PrevShaderMapEntry.NumShaders;
+		ShaderMapEntry.ShaderIndicesOffset = SerializedShaders.ShaderIndices.AddZeroed(ShaderMapEntry.NumShaders);
+
+		for (uint32 i = 0; i < ShaderMapEntry.NumShaders; ++i)
+		{
+			const int32 OtherShaderIndex = OtherArchive.SerializedShaders.ShaderIndices[PrevShaderMapEntry.ShaderIndicesOffset + i];
+			int32 ShaderIndex = 0;
+			if (SerializedShaders.FindOrAddShader(OtherArchive.SerializedShaders.ShaderHashes[OtherShaderIndex], ShaderIndex))
+			{
+				const FShaderCodeEntry& OtherShaderEntry = OtherArchive.SerializedShaders.ShaderEntries[OtherShaderIndex];
+				SerializedShaders.ShaderEntries[ShaderIndex] = OtherShaderEntry;
+				SerializedShaders.ShaderTypes[ShaderIndex] = OtherArchive.SerializedShaders.ShaderTypes[OtherShaderIndex];
+
+				const FSharedBuffer& OtherShaderCodeEntry = OtherArchive.ShaderCode[OtherShaderIndex];
+				check(OtherShaderCodeEntry.GetSize() > 0);
+				ShaderCode.Add(OtherShaderCodeEntry);
+				check(ShaderCode.Num() == SerializedShaders.ShaderEntries.Num());
+			}
+			SerializedShaders.ShaderIndices[ShaderMapEntry.ShaderIndicesOffset + i] = ShaderIndex;
+		}
+	}
+	// always mark the shadermap dirty, because it might have gotten new asset associations
+	MarkShaderMapDirty(ShaderMapIndex);
+	return ShaderMapIndex;
+}
+
+int32 FEditorShaderCodeArchive::PrepareAddShaderCode(FShaderCodeStats& CodeStats,
+	int32 OtherShaderMapIndex,
+	const FSerializedShaderArchive& OtherShaders,
+	TArray<FShaderCodeReadInfo>& ReadInfo)
+{
+	int32 ShaderMapIndex = 0;
+	const FSHAHash& OtherShaderMapHash = OtherShaders.ShaderMapHashes[OtherShaderMapIndex];
+	const FShaderMapAssetPaths* OtherAssets = OtherShaders.ShaderCodeToAssets.Find(OtherShaderMapHash);
+	bool bIsNew = SerializedShaders.FindOrAddShaderMap(OtherShaderMapHash, ShaderMapIndex, OtherAssets);
+	if (bIsNew)
+	{
+		CodeStats.NumShaderMaps++;
+
+		const FShaderMapEntry& OtherEntry = OtherShaders.ShaderMapEntries[OtherShaderMapIndex];
+		FShaderMapEntry& ShaderMapEntry = SerializedShaders.ShaderMapEntries[ShaderMapIndex];
+		ShaderMapEntry.NumShaders = OtherEntry.NumShaders;
+		ShaderMapEntry.ShaderIndicesOffset = SerializedShaders.ShaderIndices.AddZeroed(ShaderMapEntry.NumShaders);
+
+		for (uint32 i = 0; i < ShaderMapEntry.NumShaders; ++i)
+		{
+			const int32 OtherShaderIndex = OtherShaders.ShaderIndices[OtherEntry.ShaderIndicesOffset + i];
+			int32 ShaderIndex = 0;
+			if (SerializedShaders.FindOrAddShader(OtherShaders.ShaderHashes[OtherShaderIndex], ShaderIndex))
+			{
+				const FShaderCodeEntry& OtherShaderEntry = OtherShaders.ShaderEntries[OtherShaderIndex];
+				SerializedShaders.ShaderEntries[ShaderIndex] = OtherShaderEntry;
+
+				check(ShaderCode.Num() == SerializedShaders.GetNumShaders() - 1);
+
+				// Save information required to read in the old shader code
+				check(OtherShaderEntry.Size > 0);
+				FShaderCodeReadInfo Info{ ShaderCode.Num(), OtherShaderEntry.Offset, OtherShaderEntry.Size };
+				ReadInfo.Add(Info);
+
+				// We will add this old shader code in a separate pass
+				ShaderCode.AddZeroed(1);
+			}
+
+			TArray<uint64>& SerializedShaderTypes = SerializedShaders.ShaderTypes[ShaderIndex].Data;
+			for (uint64 OtherShaderTypeHash : OtherShaders.ShaderTypes[OtherShaderIndex].Data)
+			{
+				const int32 HashIndex = Algo::LowerBound(SerializedShaderTypes, OtherShaderTypeHash);
+				if (HashIndex >= SerializedShaderTypes.Num() || SerializedShaderTypes[HashIndex] != OtherShaderTypeHash)
+				{
+					SerializedShaderTypes.Insert(OtherShaderTypeHash, HashIndex);
+					FShaderCodeStats::FPerTypeStats& TypeCodeStats = CodeStats.ShaderTypeStats.FindOrAdd(OtherShaderTypeHash);
+					int32 CodeSize = SerializedShaders.ShaderEntries[ShaderIndex].Size;
+					TypeCodeStats.NumShaders++;
+					TypeCodeStats.ShadersSize += CodeSize;
+				}
+			}
+			SerializedShaders.ShaderIndices[ShaderMapEntry.ShaderIndicesOffset + i] = ShaderIndex;
+		}
+	}
+	// always mark the shadermap dirty, because it might have gotten new asset associations
+	MarkShaderMapDirty(ShaderMapIndex);
+	return ShaderMapIndex;
+}
+
+FEditorShaderCodeArchive::EShaderCodeReadStatus FEditorShaderCodeArchive::ReadShaderCodeBlock(FArchive& Ar, int64 StartOffset, const TArrayView<FShaderCodeReadInfo>& BlockReadInfo)
+{
+	if (BlockReadInfo.Num() > 0)
+	{
+		const FShaderCodeReadInfo& First = BlockReadInfo[0];
+		const FShaderCodeReadInfo& Last = BlockReadInfo[BlockReadInfo.Num() - 1];
+		int64 ReadSize = Last.Offset - First.Offset + Last.Size;
+
+		FUniqueBuffer BlockOfShaderCode = FUniqueBuffer::Alloc(ReadSize);
+		int64 FileLocation = First.Offset + StartOffset;
+		if (FileLocation != Ar.Tell())
+		{
+			Ar.Seek(FileLocation);
+		}
+		Ar.Serialize(BlockOfShaderCode.GetData(), ReadSize);
+
+		for (auto& Block : BlockReadInfo)
+		{
+			FUniqueBuffer Code = FUniqueBuffer::Clone(static_cast<const uint8*>(BlockOfShaderCode.GetData()) + Block.Offset - First.Offset, Block.Size);
+			ShaderCode[Block.ShaderCodeIndex] = Code.MoveToShared();
+		}
+
+		if (Ar.GetError())
+		{
+			return EShaderCodeReadStatus::MalformedArchive;
+		}
+	}
+	return EShaderCodeReadStatus::OK;
+}
+
+FEditorShaderCodeArchive::EShaderCodeReadStatus FEditorShaderCodeArchive::ReadShaderCode(FArchive& Ar, TArray<FShaderCodeReadInfo>& ReadInfo)
+{
+	EShaderCodeReadStatus Status = EShaderCodeReadStatus::OK;
+	if (ReadInfo.Num() > 0)
+	{
+		// Sort by offset
+		ReadInfo.Sort([](FShaderCodeReadInfo const& A, FShaderCodeReadInfo const& B) { return A.Offset < B.Offset; });
+
+		// Fill in all the shader code data, reading old data in blocks
+		const int64 DesiredBlockSize = 10 * 1024 * 1024;
+		int64 PrevCookedShadersCodeStart = Ar.Tell();
+
+		// Setup starting block, then iterate the rest
+		int32 BlockStartIndex = 0;
+		int32 IndexCount = 1;
+
+		for (int64 i = 1; i < ReadInfo.Num(); ++i)
+		{
+			const FShaderCodeReadInfo& Next = ReadInfo[i];
+
+			// See how much we'd have to read to add this block, if it adds too much, read and start a new block, otherwise include in current block
+			int64 PotentialSize = Next.Offset - ReadInfo[BlockStartIndex].Offset + Next.Size;
+			if (PotentialSize > DesiredBlockSize)
+			{
+				Status = ReadShaderCodeBlock(Ar, PrevCookedShadersCodeStart, MakeArrayView(&ReadInfo[BlockStartIndex], IndexCount));
+				BlockStartIndex = i;
+				IndexCount = 1;
+				if (Status != EShaderCodeReadStatus::OK)
+				{
+					return Status;
+				}
+			}
+			else
+			{
+				IndexCount += 1;
+			}
+		}
+
+		// Read last block
+		Status = ReadShaderCodeBlock(Ar, PrevCookedShadersCodeStart, MakeArrayView(&ReadInfo[BlockStartIndex], IndexCount));
+	}
+	return Status;
+}
+
+FEditorShaderCodeArchive::EShaderCodeReadStatus FEditorShaderCodeArchive::AddShaderCodeLibraryByName(const FString& BaseDir, const FString& InLibraryName, FShaderCodeStats& CodeStats)
+{
+	FString CodeArchiveFileName = GetCodeArchiveFilename(BaseDir, InLibraryName, FormatName);
+	if (FArchive* PrevCookedAr = CreateShaderFileReader(*CodeArchiveFileName))
+	{
+		EShaderCodeReadStatus Status;
+		uint32 Version = 0;
+		*PrevCookedAr << Version;
+
+		if (Version == GShaderCodeArchiveVersion)
+		{
+			FSerializedShaderArchive PrevCookedShaders;
+
+			*PrevCookedAr << PrevCookedShaders;
+
+			// check if it also contains the asset info file
+			FString AssetInfoFileName = GetShaderAssetInfoFilename(BaseDir, InLibraryName, FormatName);
+			TUniquePtr<FArchive> ShaderAssetInfoReader(CreateShaderFileReader(*AssetInfoFileName));
+			if (PrevCookedShaders.LoadAssetInfo(ShaderAssetInfoReader.Get()))
+			{
+				UE_LOG(LogShaderLibrary, Display, TEXT("Loaded asset info %s for the shader library %s: %d entries"),
+					*AssetInfoFileName,
+					*CodeArchiveFileName,
+					PrevCookedShaders.ShaderCodeToAssets.Num()
+				);
+			}
+			else
+			{
+				UE_LOG(LogShaderLibrary, Warning, TEXT("Could not find or load asset info %s for the shader library %s"),
+					*AssetInfoFileName,
+					*CodeArchiveFileName
+				);
+			}
+
+			// check if it also contains the shader type info file
+			FString TypeInfoFileName = GetShaderTypeInfoFilename(BaseDir, InLibraryName, FormatName);
+			TUniquePtr<FArchive> ShaderTypeInfoReader(CreateShaderFileReader(*TypeInfoFileName));
+			if (ShaderTypeInfoReader)
+			{
+				*ShaderTypeInfoReader << PrevCookedShaders.ShaderTypes;
+				check(PrevCookedShaders.ShaderTypes.Num() == PrevCookedShaders.ShaderEntries.Num());
+				UE_LOG(LogShaderLibrary, Display, TEXT("Loaded shader type info %s for the shader library %s: %d entries"),
+					*TypeInfoFileName,
+					*CodeArchiveFileName,
+					PrevCookedShaders.ShaderTypes.Num()
+				);
+			}
+			else
+			{
+				UE_LOG(LogShaderLibrary, Warning, TEXT("Could not load shader type info %s for the shader library %s"),
+					*TypeInfoFileName,
+					*CodeArchiveFileName
+				);
+			}
+
+			// Setup current shader library information and accumulate the reads we will need to do. Then read them.
+			TArray<FShaderCodeReadInfo> ReadInfo;
+			for (int32 PrevShaderMapIndex = 0; PrevShaderMapIndex < PrevCookedShaders.ShaderMapEntries.Num(); ++PrevShaderMapIndex)
+			{
+				PrepareAddShaderCode(CodeStats, PrevShaderMapIndex, PrevCookedShaders, ReadInfo);
+			}
+			Status = ReadShaderCode(*PrevCookedAr, ReadInfo);
+		}
+		else
+		{
+			Status = EShaderCodeReadStatus::VersionMismatch;
+		}
+
+		PrevCookedAr->Close();
+		delete PrevCookedAr;
+		return Status;
+	}
+	return EShaderCodeReadStatus::FileNotFound;
+}
+
+bool FEditorShaderCodeArchive::LoadExistingShaderCodeLibrary(FString const& MetaDataDir)
+{
+	FString BaseDir = MetaDataDir / TEXT("ShaderLibrarySource");
+	FShaderCodeStats CodeStats;
+	EShaderCodeReadStatus Status = AddShaderCodeLibraryByName(BaseDir, LibraryName, CodeStats);
+	switch (Status)
+	{
+	case EShaderCodeReadStatus::FileNotFound:
+	{
+		UE_LOG(LogShaderLibrary, Error, TEXT("Failed to load shader code library from %s; file not found"), *GetCodeArchiveFilename(BaseDir, LibraryName, FormatName));
+		return false;
+	}
+	case EShaderCodeReadStatus::MalformedArchive:
+	{
+		UE_LOG(LogShaderLibrary, Error, TEXT("Failed to deserialized code from library %s"), *GetCodeArchiveFilename(BaseDir, LibraryName, FormatName));
+		return false;
+	}
+	case EShaderCodeReadStatus::VersionMismatch:
+	{
+		UE_LOG(LogShaderLibrary, Error, TEXT("Failed to deserialize shader code from %s because the archive format is incompatible with the current version %u"), *GetCodeArchiveFilename(BaseDir, LibraryName, FormatName), GShaderCodeArchiveVersion);
+		return false;
+	}
+	case EShaderCodeReadStatus::OK:
+	default:
+		return true;
+	}
+}
+
+void FEditorShaderCodeArchive::AddExistingShaderCodeLibrary(FString const& OutputDir, FShaderCodeStats& CodeStats)
+{
+	check(LibraryName.Len() > 0);
+	static const FRegexPattern ShaderCodePattern(ShaderCodePatternStr);
+	static const FRegexPattern ShaderArchivePattern(ShaderArchivePatternStr);
+
+	const FString FormatNameStr = FormatName.ToString();
+	const FString ShaderIntermediateLocation = GetShaderCodeLibraryDebugInfoPath() / FormatNameStr;
+
+	TArray<FString> ShaderFiles;
+	ShaderFindFiles(ShaderFiles, *ShaderIntermediateLocation, *ShaderExtension);
+	ShaderFindFiles(ShaderFiles, *OutputDir, *ShaderExtension);
+
+	for (const FString& ShaderFileName : ShaderFiles)
+	{
+		FRegexMatcher FindShaderFormat(ShaderCodePattern, ShaderFileName);
+		if (FindShaderFormat.FindNext())
+		{
+			const FString& FoundLibraryName = FindShaderFormat.GetCaptureGroup(1);
+			const FString& FoundFormatName = FindShaderFormat.GetCaptureGroup(2);
+
+			if (FoundLibraryName.StartsWith(LibraryName) && FoundFormatName.Equals(FormatNameStr, ESearchCase::IgnoreCase))
+			{
+				AddShaderCodeLibraryByName(OutputDir, FoundLibraryName, CodeStats);
+			}
+		}
+
+		FRegexMatcher FindArchiveFormat(ShaderArchivePattern, ShaderFileName);
+		if (FindArchiveFormat.FindNext())
+		{
+			const FString& FoundArchiveName = FindArchiveFormat.GetCaptureGroup(1);
+			const FString& FoundFormatName = FindArchiveFormat.GetCaptureGroup(2);
+
+			if (FoundArchiveName.StartsWith(LibraryName) && FoundFormatName.Equals(FormatNameStr, ESearchCase::IgnoreCase))
+			{
+				AddShaderCodeLibraryByName(OutputDir, FoundArchiveName, CodeStats);
+			}
+		}
+	}
+}
+
+bool FEditorShaderCodeArchive::SaveToDisk(const FString& OutputDir, const FString& MetaOutputDir, bool bSaveOnlyAssetInfo, TArray<FString>* OutputFilenames)
+{
+	check(LibraryName.Len() > 0);
+	checkf(!bHasCopiedAndCleared, TEXT("It is not valid to call SaveToDisk on an FEditorShaderCodeArchive that has sent its shaders to another process by calling CopyAndClear."));
+
+	bool bSuccess = IFileManager::Get().MakeDirectory(*OutputDir, true);
+
+	auto CopyFile = [this](const FString& DestinationPath, const FString& SourcePath, TArray<FString>* OutputFilenames) -> bool
+		{
+			uint32 Result = IFileManager::Get().Copy(*DestinationPath, *SourcePath, true, true);
+			if (Result != COPY_OK)
+			{
+				UE_LOG(LogShaderLibrary, Error, TEXT("FEditorShaderCodeArchive copying %s to %s failed. Failed to finalize Shared Shader Library %s with format %s"),
+					*SourcePath, *DestinationPath, *LibraryName, *FormatName.ToString());
+				return false;
+			}
+
+			if (OutputFilenames)
+			{
+				OutputFilenames->Add(DestinationPath);
+			}
+			return true;
+		};
+
+	// Shader library
+	if (bSuccess && SerializedShaders.GetNumShaderMaps() > 0)
+	{
+		// Write to a intermediate file
+		const FString DebugInfoPath = GetShaderCodeLibraryDebugInfoPath();
+		FString IntermediateFormatPath = GetShaderCodeFilename(DebugInfoPath / FormatName.ToString(), LibraryName, FormatName);
+		FString AssetInfoIntermediatePath = GetShaderAssetInfoFilename(DebugInfoPath / FormatName.ToString(), LibraryName, FormatName);
+		FString TypeInfoIntermediatePath = GetShaderTypeInfoFilename(DebugInfoPath / FormatName.ToString(), LibraryName, FormatName);
+
+		// save the actual shader code code
+		if (!bSaveOnlyAssetInfo)
+		{
+			FArchive* FileWriter = IFileManager::Get().CreateFileWriter(*IntermediateFormatPath, FILEWRITE_NoFail);
+			if (FileWriter)
+			{
+				check(Format);
+
+				SerializedShaders.Finalize();
+
+				int32 ArchiveVersion = GShaderCodeArchiveVersion;
+				*FileWriter << ArchiveVersion;
+
+				// Write shader library
+				*FileWriter << SerializedShaders;
+				for (FSharedBuffer& Code : ShaderCode)
+				{
+					check(Code.GetSize() > 0);
+					FileWriter->Serialize(const_cast<void*>(Code.GetData()), Code.GetSize());
+				}
+
+				FileWriter->Close();
+				delete FileWriter;
+
+				// Copy to output location - support for iterative native library cooking
+				if (!CopyFile(GetCodeArchiveFilename(OutputDir, LibraryName, FormatName), IntermediateFormatPath, OutputFilenames))
+				{
+					bSuccess = false;
+				}
+				else if (MetaOutputDir.Len())
+				{
+					if (!CopyFile(GetCodeArchiveFilename(MetaOutputDir / TEXT("../ShaderLibrarySource"), LibraryName, FormatName), IntermediateFormatPath, nullptr))
+					{
+						bSuccess = false;
+					}
+				}
+			}
+
+			FArchive* ShaderTypeFileWriter = IFileManager::Get().CreateFileWriter(*TypeInfoIntermediatePath, FILEWRITE_NoFail);
+			if (ShaderTypeFileWriter)
+			{
+				*ShaderTypeFileWriter << SerializedShaders.ShaderTypes;
+				ShaderTypeFileWriter->Close();
+				delete ShaderTypeFileWriter;
+
+				// Copy to output location - support for iterative native library cooking
+				if (!CopyFile(GetShaderTypeInfoFilename(OutputDir, LibraryName, FormatName), TypeInfoIntermediatePath, OutputFilenames))
+				{
+					bSuccess = false;
+				}
+				else if (MetaOutputDir.Len())
+				{
+					if (!CopyFile(GetShaderTypeInfoFilename(MetaOutputDir / TEXT("../ShaderLibrarySource"), LibraryName, FormatName), TypeInfoIntermediatePath, nullptr))
+					{
+						bSuccess = false;
+					}
+				}
+			}
+		}
+
+		// save asset info
+		{
+			FArchive* AssetInfoWriter = IFileManager::Get().CreateFileWriter(*AssetInfoIntermediatePath, FILEWRITE_NoFail);
+			if (AssetInfoWriter)
+			{
+				SerializedShaders.SaveAssetInfo(*AssetInfoWriter);
+				AssetInfoWriter->Close();
+				delete AssetInfoWriter;
+
+				if (!CopyFile(GetShaderAssetInfoFilename(OutputDir, LibraryName, FormatName), AssetInfoIntermediatePath, nullptr))
+				{
+					bSuccess = false;
+				}
+				else if (MetaOutputDir.Len())
+				{
+					// copy asset info as well for debugging
+					if (!CopyFile(GetShaderAssetInfoFilename(MetaOutputDir / TEXT("../ShaderLibrarySource"), LibraryName, FormatName), AssetInfoIntermediatePath, nullptr))
+					{
+						bSuccess = false;
+					}
+				}
+			}
+		}
+	}
+
+	return bSuccess;
+}
+
+bool FEditorShaderCodeArchive::PackageNativeShaderLibrary(const FString& ShaderCodeDir, TArray<FString>* OutputFilenames)
+{
+	if (SerializedShaders.GetNumShaders() == 0)
+	{
+		return true;
+	}
+	checkf(!bHasCopiedAndCleared, TEXT("It is not valid to call PackageNativeShaderLibrary on an FEditorShaderCodeArchive that has sent its shaders to another process by calling CopyAndClear."));
+
+	FString IntermediateFormatPath = GetShaderDebugFolder(GetShaderCodeLibraryDebugInfoPath() / FormatName.ToString(), LibraryName, FormatName);
+	FString TempPath = IntermediateFormatPath / TEXT("NativeLibrary");
+
+	IFileManager::Get().MakeDirectory(*TempPath, true);
+	IFileManager::Get().MakeDirectory(*ShaderCodeDir, true);
+
+	EShaderPlatform Platform = ShaderFormatToLegacyShaderPlatform(FormatName);
+	const bool bOK = Format->CreateShaderArchive(LibraryName, FormatName, TempPath, ShaderCodeDir, IntermediateFormatPath, SerializedShaders, ShaderCode, OutputFilenames);
+
+	if (bOK)
+	{
+		// Delete Shader code library / pipelines as we now have native versions
+		// [RCL] 2020-12-02 FIXME: check if this doesn't ruin iterative cooking during Launch On
+		{
+			FString OutputFilePath = GetCodeArchiveFilename(ShaderCodeDir, LibraryName, FormatName);
+			IFileManager::Get().Delete(*OutputFilePath);
+		}
+		{
+			FString OutputFilePath = GetPipelinesArchiveFilename(ShaderCodeDir, LibraryName, FormatName);
+			IFileManager::Get().Delete(*OutputFilePath);
+		}
+	}
+
+	// Clean up the saved directory of temporary files
+	IFileManager::Get().DeleteDirectory(*IntermediateFormatPath, false, true);
+	IFileManager::Get().DeleteDirectory(*TempPath, false, true);
+
+	return bOK;
+}
+
+void FEditorShaderCodeArchive::MakePatchLibrary(TArray<FEditorShaderCodeArchive*> const& OldLibraries, FEditorShaderCodeArchive const& NewLibrary)
+{
+	for (int32 NewShaderMapIndex = 0; NewShaderMapIndex < NewLibrary.SerializedShaders.ShaderMapHashes.Num(); ++NewShaderMapIndex)
+	{
+		const FSHAHash& Hash = NewLibrary.SerializedShaders.ShaderMapHashes[NewShaderMapIndex];
+		if (!HasShaderMap(Hash))
+		{
+			bool bInPreviousPatch = false;
+			for (FEditorShaderCodeArchive const* OldLibrary : OldLibraries)
+			{
+				bInPreviousPatch |= OldLibrary->HasShaderMap(Hash);
+				if (bInPreviousPatch)
+				{
+					break;
+				}
+			}
+			if (!bInPreviousPatch)
+			{
+				AddShaderCode(NewShaderMapIndex, NewLibrary);
+			}
+		}
+	}
+}
+	
+bool FEditorShaderCodeArchive::CreatePatchLibrary(FName FormatName, FString const& LibraryName, TArray<FString> const& OldMetaDataDirs, FString const& NewMetaDataDir, FString const& OutDir, bool bNativeFormat, bool bNeedsDeterministicOrder)
+{
+	TArray<FEditorShaderCodeArchive*> OldLibraries;
+	for (FString const& OldMetaDataDir : OldMetaDataDirs)
+	{
+		FEditorShaderCodeArchive* OldLibrary = new FEditorShaderCodeArchive(FormatName, bNeedsDeterministicOrder);
+		OldLibrary->OpenLibrary(LibraryName);
+		if (OldLibrary->LoadExistingShaderCodeLibrary(OldMetaDataDir))
+		{
+			OldLibraries.Add(OldLibrary);
+		}
+	}
+
+	FEditorShaderCodeArchive NewLibrary(FormatName, bNeedsDeterministicOrder);
+	NewLibrary.OpenLibrary(LibraryName);
+	bool bOK = NewLibrary.LoadExistingShaderCodeLibrary(NewMetaDataDir);
+	if (bOK)
+	{
+		FEditorShaderCodeArchive OutLibrary(FormatName, bNeedsDeterministicOrder);
+		FShaderCodeStats OutStats;
+		OutLibrary.OpenLibrary(LibraryName);
+		OutLibrary.MakePatchLibrary(OldLibraries, NewLibrary);
+		bOK = OutLibrary.SerializedShaders.GetNumShaderMaps() > 0;
+		if (bOK)
+		{
+			FString Empty;
+			OutLibrary.FinishPopulate(OutDir, OutStats);
+			bOK = OutLibrary.SaveToDisk(OutDir, Empty);
+			UE_CLOG(!bOK, LogShaderLibrary, Error, TEXT("Failed to save %s shader patch library %s, %s, %s"), bNativeFormat ? TEXT("native") : TEXT(""), *FormatName.ToString(), *LibraryName, *OutDir);
+
+			if (bOK && bNativeFormat && OutLibrary.GetFormat()->SupportsShaderArchives())
+			{
+				bOK = OutLibrary.PackageNativeShaderLibrary(OutDir);
+				UE_CLOG(!bOK, LogShaderLibrary, Error, TEXT("Failed to package native shader patch library %s, %s, %s"), *FormatName.ToString(), *LibraryName, *OutDir);
+			}
+		}
+		else
+		{
+			UE_LOG(LogShaderLibrary, Verbose, TEXT("No shaders to patch for library %s, %s, %s"), *FormatName.ToString(), *LibraryName, *OutDir);
+		}
+	}
+	else
+	{
+		UE_LOG(LogShaderLibrary, Error, TEXT("Failed to open the shader library to patch against %s, %s, %s"), *FormatName.ToString(), *LibraryName, *NewMetaDataDir);
+	}
+
+	for (FEditorShaderCodeArchive* Lib : OldLibraries)
+	{
+		delete Lib;
+	}
+	return bOK;
+}
+
+void FEditorShaderCodeArchive::DumpStatsAndDebugInfo()
+{
+	bool bUseExtendedDebugInfo = UE::ShaderLibrary::Private::GProduceExtendedStats != 0;
+
+	UE_LOG(LogShaderLibrary, Display, TEXT(""));
+	UE_LOG(LogShaderLibrary, Display, TEXT("Shader Library '%s' (%s) Stats:"), *LibraryName, *FormatName.ToString());
+	UE_LOG(LogShaderLibrary, Display, TEXT("================="));
+
+	FSerializedShaderArchive::FDebugStats Stats;
+	FSerializedShaderArchive::FExtendedDebugStats ExtendedStats;
+	SerializedShaders.CollectStatsAndDebugInfo(Stats, bUseExtendedDebugInfo ? &ExtendedStats : nullptr);
+
+	UE_LOG(LogShaderLibrary, Display, TEXT("Assets: %d, Unique Shadermaps: %d (%.2f%%)"),
+		Stats.NumAssets, Stats.NumShaderMaps, (Stats.NumAssets > 0) ? 100.0 * static_cast<double>(Stats.NumShaderMaps) / static_cast<double>(Stats.NumAssets) : 0.0);
+	UE_LOG(LogShaderLibrary, Display, TEXT("Total Shaders: %d, Unique Shaders: %d (%.2f%%)"),
+		Stats.NumShaders, Stats.NumUniqueShaders, (Stats.NumShaders > 0) ? 100.0 * static_cast<double>(Stats.NumUniqueShaders) / static_cast<double>(Stats.NumShaders) : 0.0);
+	UE_LOG(LogShaderLibrary, Display, TEXT("Total Shader Size: %.2fmb, Unique Shaders Size: %.2fmb (%.2f%%)"),
+		FUnitConversion::Convert(static_cast<double>(Stats.ShadersSize), EUnit::Bytes, EUnit::Megabytes),
+		FUnitConversion::Convert(static_cast<double>(Stats.ShadersUniqueSize), EUnit::Bytes, EUnit::Megabytes),
+		(Stats.ShadersSize > 0) ? 100.0 * static_cast<double>(Stats.ShadersUniqueSize) / static_cast<double>(Stats.ShadersSize) : 0.0
+	);
+
+	if (bUseExtendedDebugInfo)
+	{
+		UE_LOG(LogShaderLibrary, Display, TEXT("=== Extended info:"));
+		UE_LOG(LogShaderLibrary, Display, TEXT("Minimum number of shaders in shadermap: %d"), ExtendedStats.MinNumberOfShadersPerSM);
+		UE_LOG(LogShaderLibrary, Display, TEXT("Median number of shaders in shadermap: %d"), ExtendedStats.MedianNumberOfShadersPerSM);
+		UE_LOG(LogShaderLibrary, Display, TEXT("Maximum number of shaders in shadermap: %d"), ExtendedStats.MaxNumberofShadersPerSM);
+		if (ExtendedStats.TopShaderUsages.Num() > 0)
+		{
+			FString UsageString;
+			UE_LOG(LogShaderLibrary, Display, TEXT("Number of shadermaps referencing top %d most shared shaders:"), ExtendedStats.TopShaderUsages.Num());
+			for (int IdxUsage = 0; IdxUsage < ExtendedStats.TopShaderUsages.Num() - 1; ++IdxUsage)
+			{
+				UsageString += FString::Printf(TEXT("%d, "), ExtendedStats.TopShaderUsages[IdxUsage]);
+			}
+			UE_LOG(LogShaderLibrary, Display, TEXT("    %s%d"), *UsageString, ExtendedStats.TopShaderUsages[ExtendedStats.TopShaderUsages.Num() - 1]);
+
+			// print per-frequency stats
+			UE_LOG(LogShaderLibrary, Display, TEXT("Unique shaders itemization (sorted by compressed size):"));
+			// sort by compressed size
+			TArray<int32> SortedIndices;
+			for (int32 IdxFreq = 0; IdxFreq < SF_NumFrequencies; ++IdxFreq)
+			{
+				SortedIndices.Add(IdxFreq);
+			}
+			SortedIndices.StableSort([&ExtendedStats](int32 IndexA, int32 IndexB) { return ExtendedStats.CompressedSizePerFrequency[IndexA] >= ExtendedStats.CompressedSizePerFrequency[IndexB]; });
+			for (int32 Freq : SortedIndices)
+			{
+				if (ExtendedStats.NumShadersPerFrequency[Freq] > 0)
+				{
+					UE_LOG(LogShaderLibrary, Display, TEXT("%s: %d shaders (%.2f%%), compressed size: %.2f MB (%.2f KB avg per shader), uncompressed size: %.2f MB (%.2f KB avg per shader)"),
+						GetShaderFrequencyString(static_cast<EShaderFrequency>(Freq)),
+						ExtendedStats.NumShadersPerFrequency[Freq],
+						100.0 * ExtendedStats.NumShadersPerFrequency[Freq] / double(Stats.NumUniqueShaders),
+						double(ExtendedStats.CompressedSizePerFrequency[Freq]) / (1024.0 * 1024.0), double(ExtendedStats.CompressedSizePerFrequency[Freq]) / (double(ExtendedStats.NumShadersPerFrequency[Freq]) * 1024.0),
+						double(ExtendedStats.UncompressedSizePerFrequency[Freq]) / (1024.0 * 1024.0), double(ExtendedStats.UncompressedSizePerFrequency[Freq]) / (double(ExtendedStats.NumShadersPerFrequency[Freq]) * 1024.0)
+					);
+				}
+			}
+
+		}
+		else
+		{
+			UE_LOG(LogShaderLibrary, Display, TEXT("No shader usage info is provided"));
+		}
+
+		const FString DebugInfoPath = GetShaderCodeLibraryDebugInfoPath();
+		const FString DebugLibFolder = GetShaderDebugFolder(DebugInfoPath / FormatName.ToString(), LibraryName, FormatName);
+		IFileManager::Get().MakeDirectory(*DebugLibFolder, true);
+
+		{
+			const FString DumpFilename = DebugLibFolder / TEXT("Dump.txt");
+			if (TUniquePtr<FArchive> DumpAr = TUniquePtr<FArchive>(IFileManager::Get().CreateFileWriter(*DumpFilename)))
+			{
+				FTCHARToUTF8 Converter(*ExtendedStats.TextualRepresentation);
+				DumpAr->Serialize(const_cast<UTF8CHAR*>(reinterpret_cast<const UTF8CHAR*>(Converter.Get())), Converter.Length());
+				UE_LOG(LogShaderLibrary, Display, TEXT("Textual dump saved to '%s'"), *DumpFilename);
+			}
+		}
+#if 0 // creating a graphviz graph - maybe one day we'll return to this
+		FString DebugGraphFolder = GetShaderDebugFolder(DebugInfoPath / FormatName.ToString(), LibraryName, FormatName);
+		FString DebugGraphFile = DebugGraphFolder / TEXT("RelationshipGraph.gv");
+
+		IFileManager::Get().MakeDirectory(*DebugGraphFolder, true);
+		TUniquePtr<FArchive> GraphVizAr(IFileManager::Get().CreateFileWriter(*DebugGraphFile));
+
+		TAnsiStringBuilder<512> LineBuffer;
+		LineBuffer << "digraph ShaderLibrary {\n";
+		GraphVizAr->Serialize(const_cast<ANSICHAR*>(LineBuffer.ToString()), LineBuffer.Len() * sizeof(ANSICHAR));
+		for (TTuple<FString, FString>& Edge : RelationshipGraph)
+		{
+			LineBuffer.Reset();
+			LineBuffer << "\t \"";
+			LineBuffer << TCHAR_TO_UTF8(*Edge.Key);
+			LineBuffer << "\" -> \"";
+			LineBuffer << TCHAR_TO_UTF8(*Edge.Value);
+			LineBuffer << "\";\n";
+			GraphVizAr->Serialize(const_cast<ANSICHAR*>(LineBuffer.ToString()), LineBuffer.Len() * sizeof(ANSICHAR));
+		}
+		LineBuffer.Reset();
+		LineBuffer << "}\n";
+		GraphVizAr->Serialize(const_cast<ANSICHAR*>(LineBuffer.ToString()), LineBuffer.Len() * sizeof(ANSICHAR));
+#endif//
+	}
+
+	UE_LOG(LogShaderLibrary, Display, TEXT("================="));
+}
+
+#endif // WITH_EDITOR
