Remove profiler related macros that were deprecated in UE 5.5 or UE 5.6

#Synced-CL 48676711

diff --git a/Engine/Source/Runtime/Core/Public/HAL/PlatformMisc.h b/Engine/Source/Runtime/Core/Public/HAL/PlatformMisc.h
--- a/Engine/Source/Runtime/Core/Public/HAL/PlatformMisc.h
+++ b/Engine/Source/Runtime/Core/Public/HAL/PlatformMisc.h
@@ -212,8 +212,6 @@ private:
 
 #endif
 
-#define SCOPED_PROFILER_COLOR(...) UE_DEPRECATED_MACRO(5.5, "SCOPED_PROFILER_COLOR is deprecated and there is no replacement.")
-
 // For timing OnEnterBackground tasks. This can be time sensitive on some platforms
 class FScopedEnterBackgroundEvent
 {
diff --git a/Engine/Source/Runtime/RenderCore/Public/ProfilingDebugging/RealtimeGPUProfiler.h b/Engine/Source/Runtime/RenderCore/Public/ProfilingDebugging/RealtimeGPUProfiler.h
--- a/Engine/Source/Runtime/RenderCore/Public/ProfilingDebugging/RealtimeGPUProfiler.h
+++ b/Engine/Source/Runtime/RenderCore/Public/ProfilingDebugging/RealtimeGPUProfiler.h
@@ -82,16 +82,12 @@
 	#define RHI_BREADCRUMB_EVENT_GAMETHREAD_F(                       StaticName, Format, ...) RHI_BREADCRUMB_EVENT_GAMETHREAD_PRIVATE_IMPL(RHI_GPU_STAT_ARGS_NONE,      true, TEXT(StaticName), TEXT(Format), ##__VA_ARGS__)
 	#define RHI_BREADCRUMB_EVENT_CONDITIONAL_GAMETHREAD_F(Condition, StaticName, Format, ...) RHI_BREADCRUMB_EVENT_GAMETHREAD_PRIVATE_IMPL(RHI_GPU_STAT_ARGS_NONE, Condition, TEXT(StaticName), TEXT(Format), ##__VA_ARGS__)
 
-	// Used only for back compat with SCOPED_DRAW_EVENTF_GAMETHREAD
-	#define RHI_BREADCRUMB_EVENT_GAMETHREAD_STR_DEPRECATED(          StaticName, Format, ...) RHI_BREADCRUMB_EVENT_GAMETHREAD_PRIVATE_IMPL(RHI_GPU_STAT_ARGS_NONE,      true, TEXT(StaticName),      Format , ##__VA_ARGS__)
-
 #else
 
 	#define RHI_BREADCRUMB_EVENT_GAMETHREAD(...)                do { } while(0)
 	#define RHI_BREADCRUMB_EVENT_CONDITIONAL_GAMETHREAD(...)    do { } while(0)
 	#define RHI_BREADCRUMB_EVENT_GAMETHREAD_F(...)			    do { } while(0)
 	#define RHI_BREADCRUMB_EVENT_CONDITIONAL_GAMETHREAD_F(...)  do { } while(0)
-	#define RHI_BREADCRUMB_EVENT_GAMETHREAD_STR_DEPRECATED(...) do { } while(0)
 
 #endif
 
@@ -102,55 +98,6 @@
 #define SCOPED_CONDITIONAL_DRAW_EVENT(RHICmdList, Name, Condition)               RHI_BREADCRUMB_EVENT_CONDITIONAL(RHICmdList, Condition, #Name);
 #define SCOPED_CONDITIONAL_DRAW_EVENTF(RHICmdList, Name, Condition, Format, ...) RHI_BREADCRUMB_EVENT_F_CONDITIONAL_STR_DEPRECATED(RHICmdList, Condition, #Name, Format, ##__VA_ARGS__);
 
-// Non-render-thread event macros:
-#define SCOPED_DRAW_EVENT_GAMETHREAD(Name)															 UE_DEPRECATED_MACRO(5.6, "SCOPED_DRAW_EVENT_GAMETHREAD has been deprecated. Use RHI_BREADCRUMB_EVENT_GAMETHREAD instead."                           ) RHI_BREADCRUMB_EVENT_GAMETHREAD(#Name);
-#define SCOPED_DRAW_EVENTF_GAMETHREAD(Name, Format, ...)											 UE_DEPRECATED_MACRO(5.6, "SCOPED_DRAW_EVENTF_GAMETHREAD has been deprecated. Use RHI_BREADCRUMB_EVENT_GAMETHREAD_F instead."                        ) RHI_BREADCRUMB_EVENT_GAMETHREAD_STR_DEPRECATED(#Name, Format, ##__VA_ARGS__);
-#define SCOPED_CONDITIONAL_DRAW_EVENT_GAMETHREAD(Name, Condition)									 UE_DEPRECATED_MACRO(5.6, "SCOPED_CONDITIONAL_DRAW_EVENT_GAMETHREAD has been deprecated. Use RHI_BREADCRUMB_EVENT_CONDITIONAL_GAMETHREAD instead."   ) RHI_BREADCRUMB_EVENT_CONDITIONAL_GAMETHREAD(Condition, #Name);
-#define SCOPED_CONDITIONAL_DRAW_EVENTF_GAMETHREAD(Name, Condition, Format, ...)						 UE_DEPRECATED_MACRO(5.6, "SCOPED_CONDITIONAL_DRAW_EVENTF_GAMETHREAD has been deprecated. Use RHI_BREADCRUMB_EVENT_CONDITIONAL_GAMETHREAD_F instead.") RHI_BREADCRUMB_EVENT_CONDITIONAL_GAMETHREAD_F(Condition, #Name, Format, ##__VA_ARGS__);
-
-// Deprecated macros
-#define BEGIN_DRAW_EVENTF(RHICmdList, Name, Event, Format, ...)                                      UE_DEPRECATED_MACRO(5.5, "BEGIN_DRAW_EVENTF has been deprecated. Equivalent functionality can be implemented by constructing / destructing an instance of FRHIBreadcrumbEventManual."                      )
-#define STOP_DRAW_EVENT(Event)                                                                       UE_DEPRECATED_MACRO(5.5, "STOP_DRAW_EVENT has been deprecated. Equivalent functionality can be implemented by constructing / destructing an instance of FRHIBreadcrumbEventManual."                        )
-#define STOP_DRAW_EVENT_GAMETHREAD(...)                                                              UE_DEPRECATED_MACRO(5.5, "STOP_DRAW_EVENT_GAMETHREAD has been deprecated. Equivalent functionality can be implemented by constructing / destructing an instance of FRHIBreadcrumbScope_GameThread."        )
-#define BEGIN_DRAW_EVENTF_GAMETHREAD(...)                                                            UE_DEPRECATED_MACRO(5.5, "BEGIN_DRAW_EVENTF_GAMETHREAD has been deprecated. Equivalent functionality can be implemented by constructing / destructing an instance of FRHIBreadcrumbScope_GameThread."      )
-#define BEGIN_DRAW_EVENTF_COLOR_GAMETHREAD(...)                                                      UE_DEPRECATED_MACRO(5.5, "BEGIN_DRAW_EVENTF_COLOR_GAMETHREAD has been deprecated. Equivalent functionality can be implemented by constructing / destructing an instance of FRHIBreadcrumbScope_GameThread.")
-#define SCOPED_DRAW_EVENT_COLOR(RHICmdList, Color, Name)                                             UE_DEPRECATED_MACRO(5.5, "SCOPED_DRAW_EVENT_COLOR has been deprecated. Use RHI_BREADCRUMB_EVENT instead."                                               ) RHI_BREADCRUMB_EVENT(RHICmdList, #Name)
-#define SCOPED_GPU_EVENT(RHICmdList, Name)                                                           UE_DEPRECATED_MACRO(5.5, "SCOPED_GPU_EVENT has been deprecated. Use RHI_BREADCRUMB_EVENT instead."                                                      ) RHI_BREADCRUMB_EVENT(RHICmdList, #Name)
-#define SCOPED_GPU_EVENT_COLOR(RHICmdList, Color, Name)                                              UE_DEPRECATED_MACRO(5.5, "SCOPED_GPU_EVENT_COLOR has been deprecated. Use RHI_BREADCRUMB_EVENT instead."                                                ) RHI_BREADCRUMB_EVENT(RHICmdList, #Name)
-#define SCOPED_COMPUTE_EVENT(RHICmdList, Name)                                                       UE_DEPRECATED_MACRO(5.5, "SCOPED_COMPUTE_EVENT has been deprecated. Use RHI_BREADCRUMB_EVENT instead."                                                  ) RHI_BREADCRUMB_EVENT(RHICmdList, #Name)
-#define SCOPED_COMPUTE_EVENT_COLOR(RHICmdList, Color, Name)                                          UE_DEPRECATED_MACRO(5.5, "SCOPED_COMPUTE_EVENT_COLOR has been deprecated. Use RHI_BREADCRUMB_EVENT instead."                                            ) RHI_BREADCRUMB_EVENT(RHICmdList, #Name)
-#define SCOPED_DRAW_EVENTF_COLOR(RHICmdList, Color, Name, Format, ...)                               UE_DEPRECATED_MACRO(5.5, "SCOPED_DRAW_EVENTF_COLOR has been deprecated. Use RHI_BREADCRUMB_EVENT instead."                                              ) RHI_BREADCRUMB_EVENT(RHICmdList, Format, ##__VA_ARGS__)
-#define SCOPED_GPU_EVENTF(RHICmdList, Name, Format, ...)                                             UE_DEPRECATED_MACRO(5.5, "SCOPED_GPU_EVENTF has been deprecated. Use RHI_BREADCRUMB_EVENT instead."                                                     ) RHI_BREADCRUMB_EVENT(RHICmdList, Format, ##__VA_ARGS__)
-#define SCOPED_GPU_EVENTF_COLOR(RHICmdList, Color, Name, Format, ...)                                UE_DEPRECATED_MACRO(5.5, "SCOPED_GPU_EVENTF_COLOR has been deprecated. Use RHI_BREADCRUMB_EVENT instead."                                               ) RHI_BREADCRUMB_EVENT(RHICmdList, Format, ##__VA_ARGS__)
-#define SCOPED_COMPUTE_EVENTF(RHICmdList, Name, Format, ...)                                         UE_DEPRECATED_MACRO(5.5, "SCOPED_COMPUTE_EVENTF has been deprecated. Use RHI_BREADCRUMB_EVENT instead."                                                 ) RHI_BREADCRUMB_EVENT(RHICmdList, Format, ##__VA_ARGS__)
-#define SCOPED_COMPUTE_EVENTF_COLOR(RHICmdList, Color, Name, Format, ...)                            UE_DEPRECATED_MACRO(5.5, "SCOPED_COMPUTE_EVENTF_COLOR has been deprecated. Use RHI_BREADCRUMB_EVENT instead."                                           ) RHI_BREADCRUMB_EVENT(RHICmdList, Format, ##__VA_ARGS__)
-#define SCOPED_CONDITIONAL_DRAW_EVENT_COLOR(RHICmdList, Name, Color, Condition)                      UE_DEPRECATED_MACRO(5.5, "SCOPED_CONDITIONAL_DRAW_EVENT_COLOR has been deprecated. Use RHI_BREADCRUMB_EVENT_CONDITIONAL instead."                       ) RHI_BREADCRUMB_EVENT_CONDITIONAL(RHICmdList, Condition, #Name)
-#define SCOPED_CONDITIONAL_GPU_EVENT(RHICmdList, Name, Condition)                                    UE_DEPRECATED_MACRO(5.5, "SCOPED_CONDITIONAL_GPU_EVENT has been deprecated. Use RHI_BREADCRUMB_EVENT_CONDITIONAL instead."                              ) RHI_BREADCRUMB_EVENT_CONDITIONAL(RHICmdList, Condition, #Name)
-#define SCOPED_CONDITIONAL_GPU_EVENT_COLOR(RHICmdList, Name, Color, Condition)                       UE_DEPRECATED_MACRO(5.5, "SCOPED_CONDITIONAL_GPU_EVENT_COLOR has been deprecated. Use RHI_BREADCRUMB_EVENT_CONDITIONAL instead."                        ) RHI_BREADCRUMB_EVENT_CONDITIONAL(RHICmdList, Condition, #Name)
-#define SCOPED_CONDITIONAL_COMPUTE_EVENT(RHICmdList, Name, Condition)                                UE_DEPRECATED_MACRO(5.5, "SCOPED_CONDITIONAL_COMPUTE_EVENT has been deprecated. Use RHI_BREADCRUMB_EVENT_CONDITIONAL instead."                          ) RHI_BREADCRUMB_EVENT_CONDITIONAL(RHICmdList, Condition, #Name)
-#define SCOPED_CONDITIONAL_COMPUTE_EVENT_COLOR(RHICmdList, Name, Color, Condition)                   UE_DEPRECATED_MACRO(5.5, "SCOPED_CONDITIONAL_COMPUTE_EVENT_COLOR has been deprecated. Use RHI_BREADCRUMB_EVENT_CONDITIONAL instead."                    ) RHI_BREADCRUMB_EVENT_CONDITIONAL(RHICmdList, Condition, #Name)
-#define SCOPED_CONDITIONAL_DRAW_EVENTF_COLOR(RHICmdList, Color, Name, Condition, Format, ...)        UE_DEPRECATED_MACRO(5.5, "SCOPED_CONDITIONAL_DRAW_EVENTF_COLOR has been deprecated. Use RHI_BREADCRUMB_EVENT_CONDITIONAL instead."                      ) RHI_BREADCRUMB_EVENT_CONDITIONAL(RHICmdList, Condition, Format, ##__VA_ARGS__)
-#define SCOPED_CONDITIONAL_GPU_EVENTF(RHICmdList, Name, Condition, Format, ...)                      UE_DEPRECATED_MACRO(5.5, "SCOPED_CONDITIONAL_GPU_EVENTF has been deprecated. Use RHI_BREADCRUMB_EVENT_CONDITIONAL instead."                             ) RHI_BREADCRUMB_EVENT_CONDITIONAL(RHICmdList, Condition, Format, ##__VA_ARGS__)
-#define SCOPED_CONDITIONAL_GPU_EVENTF_COLOR(RHICmdList, Color, Name, Condition, Format, ...)         UE_DEPRECATED_MACRO(5.5, "SCOPED_CONDITIONAL_GPU_EVENTF_COLOR has been deprecated. Use RHI_BREADCRUMB_EVENT_CONDITIONAL instead."                       ) RHI_BREADCRUMB_EVENT_CONDITIONAL(RHICmdList, Condition, Format, ##__VA_ARGS__)
-#define SCOPED_CONDITIONAL_COMPUTE_EVENTF(RHICmdList, Name, Condition, Format, ...)                  UE_DEPRECATED_MACRO(5.5, "SCOPED_CONDITIONAL_COMPUTE_EVENTF has been deprecated. Use RHI_BREADCRUMB_EVENT_CONDITIONAL instead."                         ) RHI_BREADCRUMB_EVENT_CONDITIONAL(RHICmdList, Condition, Format, ##__VA_ARGS__)
-#define SCOPED_CONDITIONAL_COMPUTE_EVENTF_COLOR(RHICmdList, Color, Name, Condition, Format, ...)     UE_DEPRECATED_MACRO(5.5, "SCOPED_CONDITIONAL_COMPUTE_EVENTF_COLOR has been deprecated. Use RHI_BREADCRUMB_EVENT_CONDITIONAL instead."                   ) RHI_BREADCRUMB_EVENT_CONDITIONAL(RHICmdList, Condition, Format, ##__VA_ARGS__)
-#define BEGIN_DRAW_EVENTF_COLOR(RHICmdList, Color, Name, Event, Format, ...)                         UE_DEPRECATED_MACRO(5.5, "BEGIN_DRAW_EVENTF_COLOR has been deprecated. Use BEGIN_DRAW_EVENTF instead."                                                  ) BEGIN_DRAW_EVENTF(RHICmdList, Name, Event, Format, ##__VA_ARGS__)
-#define BEGIN_GPU_EVENTF(RHICmdList, Name, Event, Format, ...)                                       UE_DEPRECATED_MACRO(5.5, "BEGIN_GPU_EVENTF has been deprecated. Use BEGIN_DRAW_EVENTF instead."                                                         ) BEGIN_DRAW_EVENTF(RHICmdList, Name, Event, Format, ##__VA_ARGS__)
-#define BEGIN_GPU_EVENTF_COLOR(RHICmdList, Color, Name, Event, Format, ...)                          UE_DEPRECATED_MACRO(5.5, "BEGIN_GPU_EVENTF_COLOR has been deprecated. Use BEGIN_DRAW_EVENTF instead."                                                   ) BEGIN_DRAW_EVENTF(RHICmdList, Name, Event, Format, ##__VA_ARGS__)
-#define STOP_GPU_EVENT(Event)                                                                        UE_DEPRECATED_MACRO(5.5, "STOP_GPU_EVENT has been deprecated. Use STOP_DRAW_EVENT instead."                                                             ) STOP_DRAW_EVENT(Event)
-#define SCOPED_DRAW_EVENT_COLOR_GAMETHREAD(Color, Name)                                              UE_DEPRECATED_MACRO(5.5, "SCOPED_DRAW_EVENT_COLOR_GAMETHREAD has been deprecated. Use RHI_BREADCRUMB_EVENT_GAMETHREAD instead."                         ) RHI_BREADCRUMB_EVENT_GAMETHREAD(#Name)
-#define SCOPED_DRAW_EVENTF_COLOR_GAMETHREAD(Color, Name, Format, ...)                                UE_DEPRECATED_MACRO(5.5, "SCOPED_DRAW_EVENTF_COLOR_GAMETHREAD has been deprecated. Use RHI_BREADCRUMB_EVENT_GAMETHREAD instead."                        ) RHI_BREADCRUMB_EVENT_GAMETHREAD(Format, ##__VA_ARGS__)
-#define SCOPED_CONDITIONAL_DRAW_EVENT_COLOR_GAMETHREAD(Name, Color, Condition)                       UE_DEPRECATED_MACRO(5.5, "SCOPED_CONDITIONAL_DRAW_EVENT_COLOR_GAMETHREAD has been deprecated. Use RHI_BREADCRUMB_EVENT_CONDITIONAL_GAMETHREAD instead." ) RHI_BREADCRUMB_EVENT_CONDITIONAL_GAMETHREAD(Condition, #Name)
-#define SCOPED_CONDITIONAL_DRAW_EVENTF_COLOR_GAMETHREAD(Color, Name, Condition, Format, ...)         UE_DEPRECATED_MACRO(5.5, "SCOPED_CONDITIONAL_DRAW_EVENTF_COLOR_GAMETHREAD has been deprecated. Use RHI_BREADCRUMB_EVENT_CONDITIONAL_GAMETHREAD instead.") RHI_BREADCRUMB_EVENT_CONDITIONAL_GAMETHREAD(Condition, Format, ##__VA_ARGS__)
-#define SCOPED_RHI_DRAW_EVENT(RHICmdContext, Name)                                                   UE_DEPRECATED_MACRO(5.5, "SCOPED_RHI_DRAW_EVENT has been deprecated. Use standard RHI breadcrumb events instead."                   )
-#define SCOPED_RHI_DRAW_EVENTF(RHICmdContext, Name, Format, ...)                                     UE_DEPRECATED_MACRO(5.5, "SCOPED_RHI_DRAW_EVENTF has been deprecated. Use standard RHI breadcrumb events instead."                  )
-#define SCOPED_RHI_CONDITIONAL_DRAW_EVENT(RHICmdContext, Name, Condition)                            UE_DEPRECATED_MACRO(5.5, "SCOPED_RHI_CONDITIONAL_DRAW_EVENT has been deprecated. Use standard RHI breadcrumb events instead."       )
-#define SCOPED_RHI_CONDITIONAL_DRAW_EVENTF(RHICmdContext, Name, Condition, Format, ...)              UE_DEPRECATED_MACRO(5.5, "SCOPED_RHI_CONDITIONAL_DRAW_EVENTF has been deprecated. Use standard RHI breadcrumb events instead."      )
-#define SCOPED_RHI_DRAW_EVENT_COLOR(RHICmdContext, Color, Name)                                      UE_DEPRECATED_MACRO(5.5, "SCOPED_RHI_DRAW_EVENT_COLOR has been deprecated. Use standard RHI breadcrumb events instead."             )
-#define SCOPED_RHI_DRAW_EVENTF_COLOR(RHICmdContext, Color, Name, Format, ...)                        UE_DEPRECATED_MACRO(5.5, "SCOPED_RHI_DRAW_EVENTF_COLOR has been deprecated. Use standard RHI breadcrumb events instead."            )
-#define SCOPED_RHI_CONDITIONAL_DRAW_EVENT_COLOR(RHICmdContext, Color, Name, Condition)               UE_DEPRECATED_MACRO(5.5, "SCOPED_RHI_CONDITIONAL_DRAW_EVENT_COLOR has been deprecated. Use standard RHI breadcrumb events instead." )
-#define SCOPED_RHI_CONDITIONAL_DRAW_EVENTF_COLOR(RHICmdContext, Color, Name, Condition, Format, ...) UE_DEPRECATED_MACRO(5.5, "SCOPED_RHI_CONDITIONAL_DRAW_EVENTF_COLOR has been deprecated. Use standard RHI breadcrumb events instead.")
-
 #if HAS_GPU_STATS
 
 	#define DECLARE_GPU_STAT_NAME_TYPE(StatName, NameString) \
