[Shaders] Fix removal of unused interpolators if any of the shader pipeline stages failed this optimization.

When interpolator removal of any of the stage jobs from a pipeline job failed, none of the stages must remove their interpolators.
Otherwise, mismatching interpolators may cause the respective PSO to fail at runtime.
This change marks the shader compile job output with a new flag to let the pipeline job know when such an event occurs.

This also generalizes the error handling of failed interpolator removal by moving these procedures into their respective utility functions (in ShaderCompilerCommonInternal.cpp/.h files).

#test Ran Lyra and EngineTest with debug code to always fail interpolator removal but only for one of the shader stages. The expected warnings showed up but launching worked fine. The same type of warnings now show up in the ReplayRun cooks, exposing the materials that cannot be optimized for unused interpolator on certain platforms.
#rnx

#Synced-CL 48841409

diff --git a/Engine/Shaders/Public/Platform/D3D/D3DCommon.ush b/Engine/Shaders/Public/Platform/D3D/D3DCommon.ush
--- a/Engine/Shaders/Public/Platform/D3D/D3DCommon.ush
+++ b/Engine/Shaders/Public/Platform/D3D/D3DCommon.ush
@@ -5,7 +5,7 @@ D3DCommon.ush: Common defines for the D3D shader compiler
 =============================================================================*/
 
 // Update this GUID to invalidate D3D shaders and force a shader recompilation for D3D RHIs.
-#pragma message("UESHADERMETADATA_VERSION FFA46460-2B9A-4B17-8877-E071C7A31822")
+#pragma message("UESHADERMETADATA_VERSION ACD85E2E-2A94-4B29-B94B-608B96CABB3C")
 
 #include "/Engine/Shared/HLSLReservedSpaces.h"
 
diff --git a/Engine/Source/Developer/ShaderCompilerCommon/Internal/ShaderCompilerCommonInternal.cpp b/Engine/Source/Developer/ShaderCompilerCommon/Internal/ShaderCompilerCommonInternal.cpp
new file mode 100644
--- /dev/null
+++ b/Engine/Source/Developer/ShaderCompilerCommon/Internal/ShaderCompilerCommonInternal.cpp
@@ -0,0 +1,71 @@
+// Copyright Epic Games, Inc. All Rights Reserved.
+
+#include "ShaderCompilerCommonInternal.h"
+#include "HlslParserInternal.h"
+
+
+namespace UE::ShaderCompilerCommon
+{
+
+bool RemoveUnusedOutputs(
+	const FShaderCompilerInput& CompilerInput,
+	FShaderCompilerOutput& CompilerOutput,
+	FString& InOutSourceCode,
+	FString& InOutEntryPoint,
+	TConstArrayView<FStringView> InUsedOutputs,
+	TConstArrayView<FStringView> InExceptions,
+	TConstArrayView<UE::HlslParser::FScopedDeclarations> InScopedDeclarations)
+{
+	TArray<FString> HlslParserErrors;
+	if (!UE::HlslParser::RemoveUnusedOutputs(InOutSourceCode, InUsedOutputs, InExceptions, InScopedDeclarations, InOutEntryPoint, HlslParserErrors))
+	{
+		CompilerOutput.bInterpolatorRemovalFailed = true;
+		CompilerOutput.Errors.Add(FShaderCompilerError(FString::Printf(TEXT("Failed to remove unused outputs from shader: %s"), *CompilerInput.GenerateShaderName())));
+		for (FString& Error : HlslParserErrors)
+		{
+			CompilerOutput.Errors.Add(FShaderCompilerError(MoveTemp(Error)));
+		}
+		return false;
+	}
+	else
+	{
+		// Unused interpolators were removed, so HlslParser generated new HLSL with a new entry point that wraps the original entry point.
+		CompilerOutput.ModifiedShaderSource = InOutSourceCode;
+		CompilerOutput.ModifiedEntryPointName = InOutEntryPoint;
+		return true;
+	}
+}
+
+bool RemoveUnusedInputs(
+	const FShaderCompilerInput& CompilerInput,
+	FShaderCompilerOutput& CompilerOutput,
+	FStringView InSourceCode,
+	FStringView InEntryPointName,
+	TConstArrayView<FString> InUsedInputs,
+	TConstArrayView<UE::HlslParser::FScopedDeclarations> InScopedDeclarations)
+{
+	TArray<FString> HlslParserErrors;
+	FString ModifiedSourceCode(InSourceCode);
+	FString ModifiedEntryPoint(InEntryPointName);
+	const TArray<FStringView> InUsedInputsView(MakeArrayView(InUsedInputs));
+	if (!UE::HlslParser::RemoveUnusedInputs(ModifiedSourceCode, InUsedInputsView, InScopedDeclarations, ModifiedEntryPoint, HlslParserErrors))
+	{
+		CompilerOutput.bInterpolatorRemovalFailed = true;
+		CompilerOutput.Errors.Add(FShaderCompilerError(FString::Printf(TEXT("Failed to remove unused inputs from shader: %s"), *CompilerInput.GenerateShaderName())));
+		for (FString& Error : HlslParserErrors)
+		{
+			CompilerOutput.Errors.Add(FShaderCompilerError(MoveTemp(Error)));
+		}
+		return false;
+	}
+	else
+	{
+		// Unused interpolators were removed, so HlslParser generated new HLSL with a new entry point that wraps the original entry point.
+		CompilerOutput.ModifiedShaderSource = MoveTemp(ModifiedSourceCode);
+		CompilerOutput.ModifiedEntryPointName = MoveTemp(ModifiedEntryPoint);
+		return true;
+	}
+}
+
+} // UE::ShaderCompilerCommon
+
diff --git a/Engine/Source/Developer/ShaderCompilerCommon/Internal/ShaderCompilerCommonInternal.h b/Engine/Source/Developer/ShaderCompilerCommon/Internal/ShaderCompilerCommonInternal.h
new file mode 100644
--- /dev/null
+++ b/Engine/Source/Developer/ShaderCompilerCommon/Internal/ShaderCompilerCommonInternal.h
@@ -0,0 +1,32 @@
+// Copyright Epic Games, Inc. All Rights Reserved.
+
+#pragma once
+
+#include "CoreMinimal.h"
+#include "ShaderCompilerCore.h"
+#include "HlslParserInternal.h"
+
+namespace UE::ShaderCompilerCommon
+{
+
+	extern SHADERCOMPILERCOMMON_API bool RemoveUnusedOutputs(
+		const FShaderCompilerInput& CompilerInput,
+		FShaderCompilerOutput& CompilerOutput,
+		FString& InOutSourceCode,
+		FString& InOutEntryPoint,
+		TConstArrayView<FStringView> InUsedOutputs,
+		TConstArrayView<FStringView> InExceptions,
+		TConstArrayView<UE::HlslParser::FScopedDeclarations> InScopedDeclarations
+	);
+
+	extern SHADERCOMPILERCOMMON_API bool RemoveUnusedInputs(
+		const FShaderCompilerInput& CompilerInput,
+		FShaderCompilerOutput& CompilerOutput,
+		FStringView InSourceCode,
+		FStringView InEntryPointName,
+		TConstArrayView<FString> InUsedInputs,
+		TConstArrayView<UE::HlslParser::FScopedDeclarations> InScopedDeclarations = {}
+	);
+
+}
+
diff --git a/Engine/Source/Developer/Windows/ShaderFormatD3D/Private/D3DShaderCompiler.cpp b/Engine/Source/Developer/Windows/ShaderFormatD3D/Private/D3DShaderCompiler.cpp
--- a/Engine/Source/Developer/Windows/ShaderFormatD3D/Private/D3DShaderCompiler.cpp
+++ b/Engine/Source/Developer/Windows/ShaderFormatD3D/Private/D3DShaderCompiler.cpp
@@ -8,6 +8,7 @@
 #include "Serialization/MemoryWriter.h"
 #include "ShaderFormatD3D.h"
 #include "ShaderCompilerCommon.h"
+#include "ShaderCompilerCommonInternal.h"
 #include "ShaderCompilerDefinitions.h"
 #include "ShaderCompileWorkerUtil.h"
 #include "ShaderMinifier.h"
@@ -1518,24 +1519,7 @@ void CompileD3DShader(const FShaderCompilerInput& Input, const FShaderPreprocess
 		};
 		ScopedDeclarations.Emplace(TConstArrayView<FStringView>(), GlobalSymbols);
 
-		TArray<FString> Errors;
-		if (!UE::HlslParser::RemoveUnusedOutputs(PreprocessedSource, UsedOutputs, Exceptions, ScopedDeclarations, EntryPointName, Errors))
-		{
-			UE_LOG(LogD3DShaderCompiler, Warning, TEXT("Failed to remove unused outputs from shader: %s"), *Input.GenerateShaderName());
-			for (const FString& ErrorReport : Errors)
-			{
-				// Add error to shader output but also make sure the error shows up on build farm by emitting a log entry
-				UE_LOG(LogD3DShaderCompiler, Warning, TEXT("%s"), *ErrorReport);
-				FShaderCompilerError NewError;
-				NewError.StrippedErrorMessage = ErrorReport;
-				Output.Errors.Add(NewError);
-			}
-		}
-		else
-		{
-			Output.ModifiedEntryPointName = EntryPointName;
-			Output.ModifiedShaderSource = PreprocessedSource;
-		}
+		(void)UE::ShaderCompilerCommon::RemoveUnusedOutputs(Input, Output, PreprocessedSource, EntryPointName, UsedOutputs, Exceptions, ScopedDeclarations);
 	}
 
 	const bool bSuccess = DoesShaderModelRequireDXC(ShaderModel)
diff --git a/Engine/Source/Developer/Windows/ShaderFormatD3D/Private/D3DShaderCompiler.inl b/Engine/Source/Developer/Windows/ShaderFormatD3D/Private/D3DShaderCompiler.inl
--- a/Engine/Source/Developer/Windows/ShaderFormatD3D/Private/D3DShaderCompiler.inl
+++ b/Engine/Source/Developer/Windows/ShaderFormatD3D/Private/D3DShaderCompiler.inl
@@ -672,18 +672,13 @@ bool RemoveUnusedInterpolators(
 				// Since unused inputs are passed to the next stage, that will cause us to generate a vertex shader that does not output B, but our pixel shader will still be expecting B on input,
 				// as it was rewritten based on the pass1 results.
 
-				FShaderCompilerOutput OriginalOutput = Output;
+				const FShaderCompilerOutput OriginalOutput = Output;
 				const int kMaxReasonableAttempts = 64;
 				for (int32 Attempt = 0; Attempt < kMaxReasonableAttempts; ++Attempt)
 				{
-					TArray<FString> RemoveErrors;
-					FString ModifiedShaderSource = PreprocessedShaderSource;
-					FString ModifiedEntryPointName = Input.EntryPointName;
-					if (UE::HlslParser::RemoveUnusedInputs(ModifiedShaderSource, CompileData.ShaderInputs, ModifiedEntryPointName, RemoveErrors))
+					if (UE::ShaderCompilerCommon::RemoveUnusedInputs(Input, Output, PreprocessedShaderSource, Input.EntryPointName, CompileData.ShaderInputs))
 					{
-						Output = OriginalOutput;
-						
-						CompileResult = ShaderCompileLambda(Input, ModifiedShaderSource, ModifiedEntryPointName, ShaderParameterParser, ShaderProfile, ShaderModel, true, Output);
+						CompileResult = ShaderCompileLambda(Input, Output.ModifiedShaderSource, Output.ModifiedEntryPointName, ShaderParameterParser, ShaderProfile, ShaderModel, true, Output);
 						if (!CompileResult)
 						{
 							// if we failed to compile the shader, propagate the error up
@@ -693,9 +688,6 @@ bool RemoveUnusedInterpolators(
 						// check if the ShaderInputs changed - if not, we're done here
 						if (Output.UsedAttributes.Num() == CompileData.ShaderInputs.Num())
 						{
-							Output.ModifiedShaderSource = MoveTemp(ModifiedShaderSource);
-							Output.ModifiedEntryPointName = MoveTemp(ModifiedEntryPointName);
-
 							return true;
 						}
 
@@ -728,21 +720,9 @@ bool RemoveUnusedInterpolators(
 						}
 
 						CompileData.ShaderInputs = Output.UsedAttributes;
+						Output = OriginalOutput;
 						// go around to remove newly identified unused inputs
 					}
-					else
-					{
-						UE_LOG(LogD3DShaderCompiler, Warning, TEXT("Failed to remove unused inputs from shader: %s"), *Input.GenerateShaderName());
-						for (const FString& ErrorMessage : RemoveErrors)
-						{
-							// Add error to shader output but also make sure the error shows up on build farm by emitting a log entry
-							UE_LOG(LogD3DShaderCompiler, Warning, TEXT("%s"), *ErrorMessage);
-							FShaderCompilerError NewError;
-							NewError.StrippedErrorMessage = ErrorMessage;
-							Output.Errors.Add(NewError);
-						}
-						break;
-					}
 				}
 			}
 		}
diff --git a/Engine/Source/Developer/Windows/ShaderFormatD3D/Private/D3DShaderCompilerDXC.cpp b/Engine/Source/Developer/Windows/ShaderFormatD3D/Private/D3DShaderCompilerDXC.cpp
--- a/Engine/Source/Developer/Windows/ShaderFormatD3D/Private/D3DShaderCompilerDXC.cpp
+++ b/Engine/Source/Developer/Windows/ShaderFormatD3D/Private/D3DShaderCompilerDXC.cpp
@@ -3,6 +3,7 @@
 #include "ShaderFormatD3D.h"
 #include "ShaderPreprocessor.h"
 #include "ShaderCompilerCommon.h"
+#include "ShaderCompilerCommonInternal.h"
 #include "ShaderCompileWorkerUtil.h"
 #include "ShaderParameterParser.h"
 #include "D3D12RHI.h"
diff --git a/Engine/Source/Runtime/RenderCore/Private/ShaderCompilerCore.cpp b/Engine/Source/Runtime/RenderCore/Private/ShaderCompilerCore.cpp
--- a/Engine/Source/Runtime/RenderCore/Private/ShaderCompilerCore.cpp
+++ b/Engine/Source/Runtime/RenderCore/Private/ShaderCompilerCore.cpp
@@ -2,4 +2,4 @@
 
 #include "ShaderCompilerCore.h"
 
-const FGuid UE_SHADER_CACHE_VERSION{"B04528D6-1A2C-44C8-A742-E7D497D3562C"};
+const FGuid UE_SHADER_CACHE_VERSION{"24752250-148B-42A5-AE80-B8C513ED34B6"};
diff --git a/Engine/Source/Runtime/RenderCore/Private/ShaderCore.cpp b/Engine/Source/Runtime/RenderCore/Private/ShaderCore.cpp
--- a/Engine/Source/Runtime/RenderCore/Private/ShaderCore.cpp
+++ b/Engine/Source/Runtime/RenderCore/Private/ShaderCore.cpp
@@ -2118,24 +2118,9 @@ void CompileShader(const TArray<const IShaderFormat*>& ShaderFormats, FShaderCom
 	Job.bSucceeded = Job.Output.bSucceeded;
 }
 
-void CompileShaderPipeline(const TArray<const IShaderFormat*>& ShaderFormats, FShaderPipelineCompileJob* PipelineJob, const FString& Dir, int32* CompileCount)
+// Compiles all shader pipeline stage jobs and returns true if all succeeded.
+static bool CompileShaderPipelineStages(const TArray<const IShaderFormat*>& ShaderFormats, FShaderPipelineCompileJob* PipelineJob, const FString& WorkingDirectory, int32* CompileCount, bool* bOutInterpolatorRemovalFailed = nullptr)
 {
-	checkf(PipelineJob->StageJobs.Num() > 0, TEXT("Pipeline %s has zero jobs!"), PipelineJob->Key.ShaderPipeline->GetName());
-
-	// This tells the shader compiler we do want to remove unused outputs
-	bool bEnableRemovingUnused = true;
-
-	//#todo-rco: Currently only removes for pure VS & PS stages
-	for (int32 Index = 0; Index < PipelineJob->StageJobs.Num(); ++Index)
-	{
-		auto Stage = PipelineJob->StageJobs[Index]->GetSingleShaderJob()->Input.Target.Frequency;
-		if (Stage != SF_Vertex && Stage != SF_Pixel)
-		{
-			bEnableRemovingUnused = false;
-			break;
-		}
-	}
-
 	FShaderCompileJob* CurrentJob = nullptr;
 	for (int32 Index = 0; Index < PipelineJob->StageJobs.Num(); ++Index)
 	{
@@ -2150,7 +2135,7 @@ void CompileShaderPipeline(const TArray<const IShaderFormat*>& ShaderFormats, FS
 			check(CurrentJob->Input.bCompilingForShaderPipeline == true);
 		}
 
-		if (bEnableRemovingUnused && !bFirstJob && PreviousJob->Output.bSupportsQueryingUsedAttributes)
+		if (bOutInterpolatorRemovalFailed && !bFirstJob && PreviousJob->Output.bSupportsQueryingUsedAttributes)
 		{
 			CurrentJob->Input.bIncludeUsedOutputs = true;
 			CurrentJob->Input.bCompilingForShaderPipeline = true;
@@ -2158,7 +2143,7 @@ void CompileShaderPipeline(const TArray<const IShaderFormat*>& ShaderFormats, FS
 		}
 
 		// Compile the shader directly through the platform dll (directly from the shader dir as the working directory)
-		CompileShader(ShaderFormats, *CurrentJob, Dir, CompileCount);
+		CompileShader(ShaderFormats, *CurrentJob, WorkingDirectory, CompileCount);
 
 		CurrentJob->bSucceeded = CurrentJob->Output.bSucceeded;
 		if (!CurrentJob->Output.bSucceeded)
@@ -2171,6 +2156,88 @@ void CompileShaderPipeline(const TArray<const IShaderFormat*>& ShaderFormats, FS
 				Job.Output.Target = Job.Input.Target;
 				Job.Output.ValidateInputHash = Job.Input.Hash;
 			}
+			return false;
+		}
+		else if (CurrentJob->Output.bInterpolatorRemovalFailed && bOutInterpolatorRemovalFailed)
+		{
+			*bOutInterpolatorRemovalFailed = true;
+		}
+	}
+
+	return true;
+}
+
+// Returns false if the specified pipeline job contains any shader stage other than Vertex or Pixel.
+static bool HasOnlyVertexOrPixelShaderStages(const FShaderPipelineCompileJob& PipelineJob)
+{
+	for (const TRefCountPtr<FShaderCompileJob>& StageJob : PipelineJob.StageJobs)
+	{
+		const EShaderFrequency Stage = StageJob->Input.Target.GetFrequency();
+		if (Stage != SF_Vertex && Stage != SF_Pixel)
+		{
+			return false;
+		}
+	}
+	return true;
+}
+
+// Resets the specified shader compiler output, but retains all errors so they can be reported as warnings.
+static void ResetShaderCompilerOutputWithRetainedErrors(FShaderCompilerOutput& Output)
+{
+	TArray<FShaderCompilerError> RetainedErrors = MoveTemp(Output.Errors);
+	bool bInterpolatorRemovalFailed = Output.bInterpolatorRemovalFailed;
+
+	Output = FShaderCompilerOutput();
+
+	Output.Errors = MoveTemp(RetainedErrors);
+	Output.bInterpolatorRemovalFailed = bInterpolatorRemovalFailed;
+}
+
+void CompileShaderPipeline(const TArray<const IShaderFormat*>& ShaderFormats, FShaderPipelineCompileJob* PipelineJob, const FString& WorkingDirectory, int32* CompileCount)
+{
+	checkf(PipelineJob->StageJobs.Num() > 0, TEXT("Pipeline %s has zero jobs!"), PipelineJob->Key.ShaderPipeline->GetName());
+
+	// This tells the shader compiler we do want to remove unused outputs
+	const bool bEnableRemovingUnused = HasOnlyVertexOrPixelShaderStages(*PipelineJob);
+	if (bEnableRemovingUnused)
+	{
+		bool bInterpolatorRemovalFailed = false;
+		if (!CompileShaderPipelineStages(ShaderFormats, PipelineJob, WorkingDirectory, CompileCount, &bInterpolatorRemovalFailed))
+		{
+			return;
+		}
+
+		// If interpolator removal failed in any shader stage, compile pipeline again with interpolator removal disabled.
+		// Otherwise, pipeline stages might end up with mismatching sets of interpolators.
+		// If the pipeline consists of only a single stage (for DepthOnly pipelines for instance), there's no need to compile again, just take the successful result.
+		if (bInterpolatorRemovalFailed && PipelineJob->StageJobs.Num() > 1)
+		{
+			// Reset compile states and strip interpolator removal option from each pipeline stage job
+			PipelineJob->ForEachSingleShaderJob(
+				[](FShaderCompileJob& StageJob)
+				{
+					// Reset information about used outputs
+					StageJob.Input.bIncludeUsedOutputs = false;
+					StageJob.Input.UsedOutputs.Empty();
+
+					// Strip interpolator removal option from input, so our next compilation attempt avoids this optimization for the entire pipeline
+					StageJob.Input.Environment.CompilerFlags.Remove(CFLAG_ForceRemoveUnusedInterpolators);
+
+					// Reset job output before compiling again but retain errors so they can be reported as warnings
+					ResetShaderCompilerOutputWithRetainedErrors(StageJob.Output);
+				}
+			);
+
+			if (!CompileShaderPipelineStages(ShaderFormats, PipelineJob, WorkingDirectory, CompileCount))
+			{
+				return;
+			}
+		}
+	}
+	else
+	{
+		if (!CompileShaderPipelineStages(ShaderFormats, PipelineJob, WorkingDirectory, CompileCount))
+		{
 			return;
 		}
 	}
diff --git a/Engine/Source/Runtime/RenderCore/Public/ShaderCompilerCore.h b/Engine/Source/Runtime/RenderCore/Public/ShaderCompilerCore.h
--- a/Engine/Source/Runtime/RenderCore/Public/ShaderCompilerCore.h
+++ b/Engine/Source/Runtime/RenderCore/Public/ShaderCompilerCore.h
@@ -33,7 +33,7 @@ typedef TSharedPtr<TArray<ANSICHAR>, ESPMode::ThreadSafe> FShaderSharedAnsiStrin
 // this is for the protocol, not the data, bump if FShaderCompilerInput/FShaderPreprocessOutput serialization, SerializeWorkerInput or ProcessInputFromArchive changes.
 inline const int32 ShaderCompileWorkerInputVersion = 32;
 // this is for the protocol, not the data, bump if FShaderCompilerOutput or WriteToOutputArchive changes.
-inline const int32 ShaderCompileWorkerOutputVersion = 28;
+inline const int32 ShaderCompileWorkerOutputVersion = 29;
 // this is for the protocol, not the data.
 inline const int32 ShaderCompileWorkerSingleJobHeader = 'S';
 // this is for the protocol, not the data.
@@ -272,6 +272,9 @@ struct FShaderCompilerInput
 	/** Oodle-specific compression level - used if CompressionFormat is set to NAME_Oodle. */
 	FOodleDataCompression::ECompressionLevel OodleLevel;
 
+	// Marks the compile job for being used with a ShaderPipeline. This tells the shader backend to remove unused interpolators from the Pixel input and Vertex output (in that order).
+	// Therefore, ShaderPipelines compile their Pixel shader first, to determine what input is actually used and then compile the Vertex shader, to remove output interpolators that are no longer needed by the Pixel input.
+	// Only the first stage in a ShaderPipeline job enables this attribute when the job is constructed and the next stage propagates its value in CompileShaderPipeline().
 	bool bCompilingForShaderPipeline = false;
 	bool bIncludeUsedOutputs = false;
 
@@ -502,6 +505,7 @@ struct FShaderCompilerOutput
 	,	bSucceeded(false)
 	,	bSupportsQueryingUsedAttributes(false)
 	,	bSerializingForCache(false)
+	,	bInterpolatorRemovalFailed(false)
 	{
 	}
 
@@ -518,6 +522,7 @@ struct FShaderCompilerOutput
 	bool bSucceeded;
 	bool bSupportsQueryingUsedAttributes;
 	bool bSerializingForCache;
+	bool bInterpolatorRemovalFailed;
 	TArray<FString> UsedAttributes;
 
 	TArray<FShaderCodeValidationStride> ParametersStrideToValidate;
@@ -606,6 +611,7 @@ struct FShaderCompilerOutput
 			Ar << Output.ValidateInputHash;
 			Ar << Output.CompileTime;
 		}
+		Ar << Output.bInterpolatorRemovalFailed;
 		Ar << Output.OutputHash;
 		Ar << Output.NumInstructions;
 		Ar << Output.NumTextureSamplers;
