IncrementalCook: Shaders: When Materials are recooked, prune the old shaders they referenced from the shader library if they are no longer used.
This fixes the cook error "Shader library contains a reference to an FShaderType with hash %llu which does not exist" that occurs in incremental cook when a shadertype is deleted.
#rnx
#Synced-CL 48939402

diff --git a/Engine/Source/Editor/UnrealEd/Private/Cooker/CookIncrementalVersion.cpp b/Engine/Source/Editor/UnrealEd/Private/Cooker/CookIncrementalVersion.cpp
--- a/Engine/Source/Editor/UnrealEd/Private/Cooker/CookIncrementalVersion.cpp
+++ b/Engine/Source/Editor/UnrealEd/Private/Cooker/CookIncrementalVersion.cpp
@@ -8,6 +8,6 @@ namespace UE::Cook
 {
 
 // Change CookIncrementalVersion to a new guid when all packages in an incremental cook need to be invalidated
-FGuid CookIncrementalVersion(0x71B523B5, 0x596C49F0, 0x99F52FDC, 0x82EFF942);
+FGuid CookIncrementalVersion(0xC70D6607, 0xCF3A48C0, 0xB484B5B8, 0x1DEA4AFD);
 
 }
diff --git a/Engine/Source/Runtime/RenderCore/Private/ShaderLibrary/EditorShaderCodeArchive.cpp b/Engine/Source/Runtime/RenderCore/Private/ShaderLibrary/EditorShaderCodeArchive.cpp
--- a/Engine/Source/Runtime/RenderCore/Private/ShaderLibrary/EditorShaderCodeArchive.cpp
+++ b/Engine/Source/Runtime/RenderCore/Private/ShaderLibrary/EditorShaderCodeArchive.cpp
@@ -9,6 +9,7 @@
 #include "EditorShaderCodeArchive.h"
 
 #include "Async/ParallelFor.h"
+#include "Cooker/CookArtifact.h"
 #include "HAL/FileManager.h"
 #include "Interfaces/IShaderFormat.h"
 #include "Interfaces/ITargetPlatformManagerModule.h"
@@ -75,7 +76,8 @@ void FEditorShaderCodeArchive::CloseLibrary(FString const& Name)
 	LibraryName = TEXT("");
 }
 
-int32 FEditorShaderCodeArchive::AddShaderCode(FShaderCodeStats& CodeStats, const FShaderMapResourceCode* Code, const FShaderMapAssetPaths& AssociatedAssets)
+int32 FEditorShaderCodeArchive::AddShaderCode(FShaderCodeStats& CodeStats, const FShaderMapResourceCode* Code,
+	const FShaderMapAssetPaths& AssociatedAssets, ECookShaderLibrarySource CookSource)
 {
 	int32 ShaderMapIndex = INDEX_NONE;
 
@@ -85,7 +87,7 @@ int32 FEditorShaderCodeArchive::AddShaderCode(FShaderCodeStats& CodeStats, const
 	}
 
 	const int32 NumShaders = Code->ShaderCodeResources.Num();
-	if (SerializedShaders.FindOrAddShaderMap(Code->ResourceHash, ShaderMapIndex, &AssociatedAssets))
+	if (SerializedShaders.FindOrAddShaderMapEditor(Code->ResourceHash, ShaderMapIndex, &AssociatedAssets, CookSource))
 	{
 		CodeStats.NumShaderMaps++;
 		FShaderMapEntry& ShaderMapEntry = SerializedShaders.ShaderMapEntries[ShaderMapIndex];
@@ -169,7 +171,8 @@ void FEditorShaderCodeArchive::CopyToCompactBinary(FCbWriter& Writer, FSerialize
 	Writer.EndObject();
 }
 
-bool FEditorShaderCodeArchive::AppendFromCompactBinary(FCbFieldView Field, FShaderCodeStats& CodeStats)
+bool FEditorShaderCodeArchive::AppendFromCompactBinary(FCbFieldView Field, FShaderCodeStats& CodeStats,
+	ECookShaderLibrarySource CookSource)
 {
 	FSerializedShaderArchive TransferArchive;
 	if (!LoadFromCompactBinary(Field["SerializedShaders"], TransferArchive))
@@ -179,11 +182,11 @@ bool FEditorShaderCodeArchive::AppendFromCompactBinary(FCbFieldView Field, FShad
 
 	FMemoryView FlatShaderCodeView = Field["ShaderCode"].AsBinaryView();
 	TConstArrayView64<uint8> TransferCode(reinterpret_cast<const uint8*>(FlatShaderCodeView.GetData()), FlatShaderCodeView.GetSize());
-	return AppendFromArchive(TransferArchive, TransferCode, CodeStats);
+	return AppendFromArchive(TransferArchive, TransferCode, CodeStats, CookSource);
 }
 
 void FEditorShaderCodeArchive::CopyToArchiveAndClear(FSerializedShaderArchive& TargetArchive, TArray<uint8>& TargetFlatShaderCode,
-	bool& bOutRanOutOfRoom, int64 MaxShaderSize, int64 MaxShaderCount)
+	bool& bOutRanOutOfRoom, int64 MaxShaderSize, int64 MaxShaderCount, ECookShaderLibrarySource CookSource)
 {
 	bOutRanOutOfRoom = false;
 	if (!bHasCopiedAndCleared)
@@ -206,10 +209,12 @@ void FEditorShaderCodeArchive::CopyToArchiveAndClear(FSerializedShaderArchive& T
 	{
 		const FSHAHash& SourceShaderMapHash = SourceArchive.ShaderMapHashes[SourceShaderMapIndex];
 		const FShaderMapEntry& SourceEntry = SourceArchive.ShaderMapEntries[SourceShaderMapIndex];
-		const FShaderMapAssetPaths* const SourceAssetPaths = SourceArchive.ShaderCodeToAssets.Find(SourceShaderMapHash);
+		const FShaderMapAssetPaths* const SourceAssetPaths
+			= SourceArchive.GetShaderMapAssetAssociations().FindShaderMap(SourceShaderMapHash);
 
 		int32 TargetShaderMapIndex;
-		const bool bIsNewShaderMap = TargetArchive.FindOrAddShaderMap(SourceShaderMapHash, TargetShaderMapIndex, SourceAssetPaths);
+		const bool bIsNewShaderMap = TargetArchive.FindOrAddShaderMapEditor(SourceShaderMapHash, TargetShaderMapIndex,
+			SourceAssetPaths, CookSource);
 		if (!bIsNewShaderMap)
 		{
 			continue;
@@ -291,7 +296,8 @@ void FEditorShaderCodeArchive::CopyToArchiveAndClear(FSerializedShaderArchive& T
 }
 
 bool FEditorShaderCodeArchive::AppendFromArchive(
-	const FSerializedShaderArchive& SourceArchive, TConstArrayView64<uint8> SourceFlatShaderCode, FShaderCodeStats& CodeStats)
+	const FSerializedShaderArchive& SourceArchive, TConstArrayView64<uint8> SourceFlatShaderCode,
+	FShaderCodeStats& CodeStats, ECookShaderLibrarySource CookSource)
 {
 	bool bOk = true;
 	FSerializedShaderArchive& TargetArchive = this->SerializedShaders;
@@ -358,9 +364,11 @@ bool FEditorShaderCodeArchive::AppendFromArchive(
 	for (int32 SourceShaderMapIndex = 0; SourceShaderMapIndex < SourceArchive.ShaderMapHashes.Num(); ++SourceShaderMapIndex)
 	{
 		const FSHAHash& SourceShaderMapHash = SourceArchive.ShaderMapHashes[SourceShaderMapIndex];
-		const FShaderMapAssetPaths* const SourceAssetPaths = SourceArchive.ShaderCodeToAssets.Find(SourceShaderMapHash);
+		const FShaderMapAssetPaths* const SourceAssetPaths
+			= SourceArchive.GetShaderMapAssetAssociations().FindShaderMap(SourceShaderMapHash);
 		int32 TargetShaderMapIndex;
-		const bool bIsNewShaderMap = TargetArchive.FindOrAddShaderMap(SourceShaderMapHash, TargetShaderMapIndex, SourceAssetPaths);
+		const bool bIsNewShaderMap = TargetArchive.FindOrAddShaderMapEditor(SourceShaderMapHash, TargetShaderMapIndex,
+			SourceAssetPaths, CookSource);
 		// always mark the shadermap dirty, because it might have gotten new asset associations
 		MarkShaderMapDirty(TargetShaderMapIndex);
 		if (!bIsNewShaderMap)
@@ -417,29 +425,32 @@ FEditorShaderCodeArchive* FEditorShaderCodeArchive::CreateChunk(int ChunkId, con
 	return NewChunk;
 }
 
-int32 FEditorShaderCodeArchive::AddShaderCode(int32 OtherShaderMapIndex, const FEditorShaderCodeArchive& OtherArchive)
+int32 FEditorShaderCodeArchive::AddShaderCode(int32 OtherShaderMapIndex, const FEditorShaderCodeArchive& OtherArchive,
+	ECookShaderLibrarySource CookSource)
 {
 	checkf(!OtherArchive.bHasCopiedAndCleared,
 		TEXT("It is not valid to call AddShaderCode from an FEditorShaderCodeArchive that has sent its shaders to another process by calling CopyAndClear."));
 
+	const FSerializedShaderArchive& OtherShaders = OtherArchive.SerializedShaders;
+	const FSHAHash& OtherShaderMapHash = OtherShaders.ShaderMapHashes[OtherShaderMapIndex];
+	const FShaderMapAssetPaths* OtherAssets = OtherShaders.GetShaderMapAssetAssociations().FindShaderMap(OtherShaderMapHash);
 	int32 ShaderMapIndex = 0;
-	if (SerializedShaders.FindOrAddShaderMap(OtherArchive.SerializedShaders.ShaderMapHashes[OtherShaderMapIndex], ShaderMapIndex,
-		OtherArchive.SerializedShaders.ShaderCodeToAssets.Find(OtherArchive.SerializedShaders.ShaderMapHashes[OtherShaderMapIndex])))
+	if (SerializedShaders.FindOrAddShaderMapEditor(OtherShaderMapHash, ShaderMapIndex, OtherAssets, CookSource))
 	{
-		const FShaderMapEntry& PrevShaderMapEntry = OtherArchive.SerializedShaders.ShaderMapEntries[OtherShaderMapIndex];
+		const FShaderMapEntry& PrevShaderMapEntry = OtherShaders.ShaderMapEntries[OtherShaderMapIndex];
 		FShaderMapEntry& ShaderMapEntry = SerializedShaders.ShaderMapEntries[ShaderMapIndex];
 		ShaderMapEntry.NumShaders = PrevShaderMapEntry.NumShaders;
 		ShaderMapEntry.ShaderIndicesOffset = SerializedShaders.ShaderIndices.AddZeroed(ShaderMapEntry.NumShaders);
 
 		for (uint32 i = 0; i < ShaderMapEntry.NumShaders; ++i)
 		{
-			const int32 OtherShaderIndex = OtherArchive.SerializedShaders.ShaderIndices[PrevShaderMapEntry.ShaderIndicesOffset + i];
+			const int32 OtherShaderIndex = OtherShaders.ShaderIndices[PrevShaderMapEntry.ShaderIndicesOffset + i];
 			int32 ShaderIndex = 0;
-			if (SerializedShaders.FindOrAddShader(OtherArchive.SerializedShaders.ShaderHashes[OtherShaderIndex], ShaderIndex))
+			if (SerializedShaders.FindOrAddShader(OtherShaders.ShaderHashes[OtherShaderIndex], ShaderIndex))
 			{
-				const FShaderCodeEntry& OtherShaderEntry = OtherArchive.SerializedShaders.ShaderEntries[OtherShaderIndex];
+				const FShaderCodeEntry& OtherShaderEntry = OtherShaders.ShaderEntries[OtherShaderIndex];
 				SerializedShaders.ShaderEntries[ShaderIndex] = OtherShaderEntry;
-				SerializedShaders.ShaderTypes[ShaderIndex] = OtherArchive.SerializedShaders.ShaderTypes[OtherShaderIndex];
+				SerializedShaders.ShaderTypes[ShaderIndex] = OtherShaders.ShaderTypes[OtherShaderIndex];
 
 				const FSharedBuffer& OtherShaderCodeEntry = OtherArchive.ShaderCode[OtherShaderIndex];
 				check(OtherShaderCodeEntry.GetSize() > 0);
@@ -457,12 +468,13 @@ int32 FEditorShaderCodeArchive::AddShaderCode(int32 OtherShaderMapIndex, const F
 int32 FEditorShaderCodeArchive::PrepareAddShaderCode(FShaderCodeStats& CodeStats,
 	int32 OtherShaderMapIndex,
 	const FSerializedShaderArchive& OtherShaders,
-	TArray<FShaderCodeReadInfo>& ReadInfo)
+	TArray<FShaderCodeReadInfo>& ReadInfo,
+	ECookShaderLibrarySource CookSource)
 {
 	int32 ShaderMapIndex = 0;
 	const FSHAHash& OtherShaderMapHash = OtherShaders.ShaderMapHashes[OtherShaderMapIndex];
-	const FShaderMapAssetPaths* OtherAssets = OtherShaders.ShaderCodeToAssets.Find(OtherShaderMapHash);
-	bool bIsNew = SerializedShaders.FindOrAddShaderMap(OtherShaderMapHash, ShaderMapIndex, OtherAssets);
+	const FShaderMapAssetPaths* OtherAssets = OtherShaders.GetShaderMapAssetAssociations().FindShaderMap(OtherShaderMapHash);
+	bool bIsNew = SerializedShaders.FindOrAddShaderMapEditor(OtherShaderMapHash, ShaderMapIndex, OtherAssets, CookSource);
 	if (bIsNew)
 	{
 		CodeStats.NumShaderMaps++;
@@ -587,7 +599,8 @@ FEditorShaderCodeArchive::EShaderCodeReadStatus FEditorShaderCodeArchive::ReadSh
 	return Status;
 }
 
-FEditorShaderCodeArchive::EShaderCodeReadStatus FEditorShaderCodeArchive::AddShaderCodeLibraryByName(const FString& BaseDir, const FString& InLibraryName, FShaderCodeStats& CodeStats)
+FEditorShaderCodeArchive::EShaderCodeReadStatus FEditorShaderCodeArchive::AddShaderCodeLibraryByName(
+	const FString& BaseDir, const FString& InLibraryName, FShaderCodeStats& CodeStats, ECookShaderLibrarySource CookSource)
 {
 	FString CodeArchiveFileName = GetCodeArchiveFilename(BaseDir, InLibraryName, FormatName);
 	if (FArchive* PrevCookedAr = CreateShaderFileReader(*CodeArchiveFileName))
@@ -605,20 +618,29 @@ FEditorShaderCodeArchive::EShaderCodeReadStatus FEditorShaderCodeArchive::AddSha
 			// check if it also contains the asset info file
 			FString AssetInfoFileName = GetShaderAssetInfoFilename(BaseDir, InLibraryName, FormatName);
 			TUniquePtr<FArchive> ShaderAssetInfoReader(CreateShaderFileReader(*AssetInfoFileName));
-			if (PrevCookedShaders.LoadAssetInfo(ShaderAssetInfoReader.Get()))
+			if (PrevCookedShaders.LoadAssetInfo(ShaderAssetInfoReader.Get(), CookSource))
 			{
 				UE_LOG(LogShaderLibrary, Display, TEXT("Loaded asset info %s for the shader library %s: %d entries"),
 					*AssetInfoFileName,
 					*CodeArchiveFileName,
-					PrevCookedShaders.ShaderCodeToAssets.Num()
+					PrevCookedShaders.GetShaderMapAssetAssociations().ViewShaderMaps().Num()
 				);
 			}
 			else
 			{
-				UE_LOG(LogShaderLibrary, Warning, TEXT("Could not find or load asset info %s for the shader library %s"),
-					*AssetInfoFileName,
-					*CodeArchiveFileName
-				);
+				FString ErrorMessage = FString::Printf(
+					TEXT("Could not find or load asset info %s for the shader library %s"),
+					*AssetInfoFileName, *CodeArchiveFileName);
+				if (CookSource == ECookShaderLibrarySource::PreviousIncremental)
+				{
+					// In incremental cooks, we require the AssetInfo so we know which shadermaps to keep, lack of it
+					// should report a cook error.
+					UE_LOG(LogShaderLibrary, Error, TEXT("%s"), *ErrorMessage);
+				}
+				else
+				{
+					UE_LOG(LogShaderLibrary, Warning, TEXT("%s"), *ErrorMessage);
+				}
 			}
 
 			// check if it also contains the shader type info file
@@ -646,7 +668,7 @@ FEditorShaderCodeArchive::EShaderCodeReadStatus FEditorShaderCodeArchive::AddSha
 			TArray<FShaderCodeReadInfo> ReadInfo;
 			for (int32 PrevShaderMapIndex = 0; PrevShaderMapIndex < PrevCookedShaders.ShaderMapEntries.Num(); ++PrevShaderMapIndex)
 			{
-				PrepareAddShaderCode(CodeStats, PrevShaderMapIndex, PrevCookedShaders, ReadInfo);
+				PrepareAddShaderCode(CodeStats, PrevShaderMapIndex, PrevCookedShaders, ReadInfo, CookSource);
 			}
 			Status = ReadShaderCode(*PrevCookedAr, ReadInfo);
 		}
@@ -662,11 +684,11 @@ FEditorShaderCodeArchive::EShaderCodeReadStatus FEditorShaderCodeArchive::AddSha
 	return EShaderCodeReadStatus::FileNotFound;
 }
 
-bool FEditorShaderCodeArchive::LoadExistingShaderCodeLibrary(FString const& MetaDataDir)
+bool FEditorShaderCodeArchive::LoadExistingShaderCodeLibrary(FString const& MetaDataDir, ECookShaderLibrarySource CookSource)
 {
 	FString BaseDir = MetaDataDir / TEXT("ShaderLibrarySource");
 	FShaderCodeStats CodeStats;
-	EShaderCodeReadStatus Status = AddShaderCodeLibraryByName(BaseDir, LibraryName, CodeStats);
+	EShaderCodeReadStatus Status = AddShaderCodeLibraryByName(BaseDir, LibraryName, CodeStats, CookSource);
 	switch (Status)
 	{
 	case EShaderCodeReadStatus::FileNotFound:
@@ -690,7 +712,8 @@ bool FEditorShaderCodeArchive::LoadExistingShaderCodeLibrary(FString const& Meta
 	}
 }
 
-void FEditorShaderCodeArchive::AddExistingShaderCodeLibrary(FString const& OutputDir, FShaderCodeStats& CodeStats)
+void FEditorShaderCodeArchive::AddExistingShaderCodeLibrary(FString const& OutputDir, FShaderCodeStats& CodeStats,
+	ECookShaderLibrarySource CookSource)
 {
 	check(LibraryName.Len() > 0);
 	static const FRegexPattern ShaderCodePattern(ShaderCodePatternStr);
@@ -713,7 +736,7 @@ void FEditorShaderCodeArchive::AddExistingShaderCodeLibrary(FString const& Outpu
 
 			if (FoundLibraryName.StartsWith(LibraryName) && FoundFormatName.Equals(FormatNameStr, ESearchCase::IgnoreCase))
 			{
-				AddShaderCodeLibraryByName(OutputDir, FoundLibraryName, CodeStats);
+				AddShaderCodeLibraryByName(OutputDir, FoundLibraryName, CodeStats, CookSource);
 			}
 		}
 
@@ -725,12 +748,72 @@ void FEditorShaderCodeArchive::AddExistingShaderCodeLibrary(FString const& Outpu
 
 			if (FoundArchiveName.StartsWith(LibraryName) && FoundFormatName.Equals(FormatNameStr, ESearchCase::IgnoreCase))
 			{
-				AddShaderCodeLibraryByName(OutputDir, FoundArchiveName, CodeStats);
+				AddShaderCodeLibraryByName(OutputDir, FoundArchiveName, CodeStats, CookSource);
 			}
 		}
 	}
 }
 
+void FEditorShaderCodeArchive::UpdateOplogPackages(UE::Cook::Artifact::FUpdateOplogPackagesContext& Context)
+{
+	using FShaderMapAssetAssociations = FSerializedShaderArchive::FShaderMapAssetAssociations;
+	using FOplogPackageData = UE::Cook::Artifact::FOplogPackageData;
+
+	const TMap<FName, FOplogPackageData>& OplogPackages = Context.GetOplogPackages();
+	FShaderMapAssetAssociations& AssetMap = SerializedShaders.GetShaderMapAssetAssociations();
+
+	for (TPair<FName, FShaderMapAssetAssociations::FAssociatedAssetData>& Pair : AssetMap.ViewAssets())
+	{
+		FName PackageName = Pair.Key;
+		const FOplogPackageData* OplogData = OplogPackages.Find(PackageName);
+		if (OplogData)
+		{
+			Pair.Value.bReferencedByOplog = true;
+			Pair.Value.bReferencedByStaging |= OplogData->IsReferenced();
+		}
+	}
+}
+
+void FEditorShaderCodeArchive::Prune()
+{
+	using FShaderMapAssetAssociations = FSerializedShaderArchive::FShaderMapAssetAssociations;
+
+	constexpr bool bImplemented = false;
+	if (!bImplemented)
+	{
+		return;
+	}
+
+	FShaderMapAssetAssociations& AssetMap = SerializedShaders.GetShaderMapAssetAssociations();
+	TArray<FName> RemoveAssets;
+	for (const TPair<FName, FShaderMapAssetAssociations::FAssociatedAssetData>& AssetPair : AssetMap.ViewAssets())
+	{
+		FName PackageName = AssetPair.Key;
+		if (!AssetPair.Value.bReferencedByOplog)
+		{
+			RemoveAssets.Add(PackageName);
+		}
+	}
+
+	TSet<FSHAHash> RemoveShaderMaps;
+	for (FSHAHash& ShaderMapHash : SerializedShaders.ShaderMapHashes)
+	{
+		if (!AssetMap.FindShaderMap(ShaderMapHash))
+		{
+			RemoveShaderMaps.Add(ShaderMapHash);
+		}
+	}
+
+	// TODO: Implement RemoveShaderMaps.
+	// TODO: In the Library write functions, move the ShaderMaps that do not have any bReferencedByStaging Assets into
+	// a separate pair of files: IncrementalCookArchive.uincrementalcookshaderbytecode, IncrementalCookAssetInfo.json that
+	// are not used by staging. Note that we cannot use .ushaderbytecode as the extension, because staging searches
+	// for *.ushaderbytecode.
+	// TODO: In FinishPopulate, read the IncrementalCookArchive files in addition to the .ushaderbytecode files.
+
+	// FSerializedShaderArchive::RemoveShaderMaps(RemoveShaderMaps);
+}
+
 bool FEditorShaderCodeArchive::SaveToDisk(const FString& OutputDir, const FString& MetaOutputDir, bool bSaveOnlyAssetInfo, TArray<FString>* OutputFilenames)
 {
 	check(LibraryName.Len() > 0);
@@ -908,7 +991,7 @@ void FEditorShaderCodeArchive::MakePatchLibrary(TArray<FEditorShaderCodeArchive*
 			}
 			if (!bInPreviousPatch)
 			{
-				AddShaderCode(NewShaderMapIndex, NewLibrary);
+				AddShaderCode(NewShaderMapIndex, NewLibrary, ECookShaderLibrarySource::PatchBase);
 			}
 		}
 	}
@@ -921,7 +1004,7 @@ bool FEditorShaderCodeArchive::CreatePatchLibrary(FName FormatName, FString cons
 	{
 		FEditorShaderCodeArchive* OldLibrary = new FEditorShaderCodeArchive(FormatName, bNeedsDeterministicOrder);
 		OldLibrary->OpenLibrary(LibraryName);
-		if (OldLibrary->LoadExistingShaderCodeLibrary(OldMetaDataDir))
+		if (OldLibrary->LoadExistingShaderCodeLibrary(OldMetaDataDir, ECookShaderLibrarySource::PatchBase))
 		{
 			OldLibraries.Add(OldLibrary);
 		}
@@ -929,7 +1012,7 @@ bool FEditorShaderCodeArchive::CreatePatchLibrary(FName FormatName, FString cons
 
 	FEditorShaderCodeArchive NewLibrary(FormatName, bNeedsDeterministicOrder);
 	NewLibrary.OpenLibrary(LibraryName);
-	bool bOK = NewLibrary.LoadExistingShaderCodeLibrary(NewMetaDataDir);
+	bool bOK = NewLibrary.LoadExistingShaderCodeLibrary(NewMetaDataDir, ECookShaderLibrarySource::PatchBase);
 	if (bOK)
 	{
 		FEditorShaderCodeArchive OutLibrary(FormatName, bNeedsDeterministicOrder);
@@ -940,7 +1023,7 @@ bool FEditorShaderCodeArchive::CreatePatchLibrary(FName FormatName, FString cons
 		if (bOK)
 		{
 			FString Empty;
-			OutLibrary.FinishPopulate(OutDir, OutStats);
+			OutLibrary.FinishPopulate(OutDir, OutStats, ECookShaderLibrarySource::PatchBase);
 			bOK = OutLibrary.SaveToDisk(OutDir, Empty);
 			UE_CLOG(!bOK, LogShaderLibrary, Error, TEXT("Failed to save %s shader patch library %s, %s, %s"), bNativeFormat ? TEXT("native") : TEXT(""), *FormatName.ToString(), *LibraryName, *OutDir);
 
diff --git a/Engine/Source/Runtime/RenderCore/Private/ShaderLibrary/EditorShaderCodeArchive.h b/Engine/Source/Runtime/RenderCore/Private/ShaderLibrary/EditorShaderCodeArchive.h
--- a/Engine/Source/Runtime/RenderCore/Private/ShaderLibrary/EditorShaderCodeArchive.h
+++ b/Engine/Source/Runtime/RenderCore/Private/ShaderLibrary/EditorShaderCodeArchive.h
@@ -20,6 +20,7 @@
 
 
 class IShaderFormat;
+namespace UE::Cook::Artifact { struct FUpdateOplogPackagesContext; }
 
 /**
  * Shader code statistics structure used by FEditorShaderCodeArchive.
@@ -42,6 +43,8 @@ struct FShaderCodeStats
  */
 class FEditorShaderCodeArchive
 {
+	using ECookShaderLibrarySource = ShaderCodeArchive::ECookShaderLibrarySource;
+
 	const FName FormatName;
 	FString LibraryName;
 	FSerializedShaderArchive SerializedShaders;
@@ -90,7 +93,8 @@ public:
 	}
 
 	/** Adds and deduplicates shader code from the specified shader map resoure to the list of serialized shaders. */
-	int32 AddShaderCode(FShaderCodeStats& CodeStats, const FShaderMapResourceCode* Code, const FShaderMapAssetPaths& AssociatedAssets);
+	int32 AddShaderCode(FShaderCodeStats& CodeStats, const FShaderMapResourceCode* Code,
+		const FShaderMapAssetPaths& AssociatedAssets, ECookShaderLibrarySource CookSource);
 
 	/** Returns true if this archive has any data that can be copied in CopyToArchiveAndClear(). */
 	bool HasDataToCopy() const;
@@ -100,35 +104,45 @@ public:
 	void CopyToCompactBinary(FCbWriter& Writer, FSerializedShaderArchive& TransferArchive, TArray<uint8>& TransferCode);
 
 	/** Loads (LoadFromCompactBinary) and appends (AppendFromArchive) all serialized shaders from the specified compact binary field to this archive. */
-	bool AppendFromCompactBinary(FCbFieldView Field, FShaderCodeStats& CodeStats);
+	bool AppendFromCompactBinary(FCbFieldView Field, FShaderCodeStats& CodeStats, ECookShaderLibrarySource CookSource);
 
 	/** Copies all serialized shaders to the target archive and its flat shader code array.
 	 * This also flushes the shared buffer of all serialized shaders in this archive to release potentially large amounts of memory.
 	 * If output parameter bOutRanOutOfRoom is true after this call, additional serialized shaders should be copied into the next target archive. */
 	void CopyToArchiveAndClear(FSerializedShaderArchive& TargetArchive, TArray<uint8>& TargetFlatShaderCode,
-		bool& bOutRanOutOfRoom, int64 MaxShaderSize, int64 MaxShaderCount);
+		bool& bOutRanOutOfRoom, int64 MaxShaderSize, int64 MaxShaderCount, ECookShaderLibrarySource CookSource);
 
 	/** Appends all entries of the source archive to the list of serialized shaders in this archive, including copies of the shader code. */
-	bool AppendFromArchive(const FSerializedShaderArchive& SourceArchive, TConstArrayView64<uint8> SourceFlatShaderCode, FShaderCodeStats& CodeStats);
+	bool AppendFromArchive(const FSerializedShaderArchive& SourceArchive, TConstArrayView64<uint8> SourceFlatShaderCode,
+		FShaderCodeStats& CodeStats, ECookShaderLibrarySource CookSource);
 
 	/** Produces another archive that contains the code only for these assets */
 	FEditorShaderCodeArchive* CreateChunk(int ChunkId, const TSet<FName>& PackagesInChunk);
 
 	/** Adds entries for the specified shader map from another archive with all its individual shaders and their code to this archive. */
-	int32 AddShaderCode(int32 OtherShaderMapIndex, const FEditorShaderCodeArchive& OtherArchive);
+	int32 AddShaderCode(int32 OtherShaderMapIndex, const FEditorShaderCodeArchive& OtherArchive,
+		ECookShaderLibrarySource CookSource);
 
 	/** Shortcut to add shader code library from directory . */
-	FORCEINLINE void AddShaderCodeLibraryFromDirectory(const FString& BaseDir, FShaderCodeStats& CodeStats)
+	FORCEINLINE void AddShaderCodeLibraryFromDirectory(const FString& BaseDir, FShaderCodeStats& CodeStats,
+		ECookShaderLibrarySource CookSource)
 	{
-		AddShaderCodeLibraryByName(BaseDir, LibraryName, CodeStats);
+		AddShaderCodeLibraryByName(BaseDir, LibraryName, CodeStats, CookSource);
 	}
 
 	/** Reads all existing shader code entries from file and adds them to this library. This marks the end of populating this archive. */
-	FORCEINLINE void FinishPopulate(const FString& OutputDir, FShaderCodeStats& CodeStats)
+	FORCEINLINE void FinishPopulate(const FString& OutputDir, FShaderCodeStats& CodeStats,
+		ECookShaderLibrarySource CookSource)
 	{
-		AddExistingShaderCodeLibrary(OutputDir, CodeStats);
+		AddExistingShaderCodeLibrary(OutputDir, CodeStats, CookSource);
 	}
 
+	/** Called by the cooker to report which assets are referenced in the current cook or in incrementalcook oplog. */
+	void UpdateOplogPackages(UE::Cook::Artifact::FUpdateOplogPackagesContext& Context);
+
+	/** Called by the cooker to ask the archive to prune no-longer referenced shaders. */
+	void Prune();
+
 	/** Saves this shader archive to disk:
 	 * Shader entries and all their code to *.ushaderbytecode file,
 	 * Asset-shader associations to *.assetinfo.json file,
@@ -157,9 +171,12 @@ private:
 		int32 Size;
 	};
 
-	void AddExistingShaderCodeLibrary(FString const& OutputDir, FShaderCodeStats& CodeStats);
-	bool LoadExistingShaderCodeLibrary(FString const& MetaDataDir);
-	int32 PrepareAddShaderCode(FShaderCodeStats& CodeStats, int32 OtherShaderMapIndex, const FSerializedShaderArchive& OtherShaders, TArray<FShaderCodeReadInfo>& ReadInfo);
+	void AddExistingShaderCodeLibrary(FString const& OutputDir, FShaderCodeStats& CodeStats,
+		ECookShaderLibrarySource CookSource);
+	bool LoadExistingShaderCodeLibrary(FString const& MetaDataDir, ECookShaderLibrarySource CookSource);
+	int32 PrepareAddShaderCode(FShaderCodeStats& CodeStats, int32 OtherShaderMapIndex,
+		const FSerializedShaderArchive& OtherShaders, TArray<FShaderCodeReadInfo>& ReadInfo,
+		ECookShaderLibrarySource CookSource);
 
 	enum class EShaderCodeReadStatus : uint8
 	{
@@ -169,7 +186,8 @@ private:
 		VersionMismatch,
 	};
 
-	EShaderCodeReadStatus AddShaderCodeLibraryByName(const FString& BaseDir, const FString& InLibraryName, FShaderCodeStats& CodeStats);
+	EShaderCodeReadStatus AddShaderCodeLibraryByName(const FString& BaseDir, const FString& InLibraryName,
+		FShaderCodeStats& CodeStats, ECookShaderLibrarySource CookSource);
 	EShaderCodeReadStatus ReadShaderCode(FArchive& Ar, TArray<FShaderCodeReadInfo>& ReadInfo);
 	EShaderCodeReadStatus ReadShaderCodeBlock(FArchive& Ar, int64 StartOffset, const TArrayView<FShaderCodeReadInfo>& BlockReadInfo);
 
diff --git a/Engine/Source/Runtime/RenderCore/Private/ShaderLibrary/ShaderCodeArchive.cpp b/Engine/Source/Runtime/RenderCore/Private/ShaderLibrary/ShaderCodeArchive.cpp
--- a/Engine/Source/Runtime/RenderCore/Private/ShaderLibrary/ShaderCodeArchive.cpp
+++ b/Engine/Source/Runtime/RenderCore/Private/ShaderLibrary/ShaderCodeArchive.cpp
@@ -137,7 +137,7 @@ int32 FSerializedShaderArchive::FindShaderMap(const FSHAHash& Hash) const
 }
 
 #if !USE_MMAPPED_SHADERARCHIVE
-bool FSerializedShaderArchive::FindOrAddShaderMap(const FSHAHash& Hash, int32& OutIndex, const FShaderMapAssetPaths* AssociatedAssets)
+bool FSerializedShaderArchive::FindOrAddShaderMapNonEditorInternal(const FSHAHash& Hash, int32& OutIndex)
 {
 	const uint32 Key = GetTypeHash(Hash);
 	int32 Index = FindShaderMapWithKey(Hash, Key);
@@ -149,18 +149,93 @@ bool FSerializedShaderArchive::FindOrAddShaderMap(const FSHAHash& Hash, int32& O
 		check(ShaderMapEntries.Num() == ShaderMapHashes.Num());
 		ShaderMapHashTable.Add(Key, Index);
 	}
+	OutIndex = Index;
+	return bAdded;
+}
+
+bool FSerializedShaderArchive::FindOrAddShaderMap(const FSHAHash& Hash, int32& OutIndex,
+	const FShaderMapAssetPaths* AssociatedAssets)
+{
 #if WITH_EDITOR
-	if (AssociatedAssets && AssociatedAssets->Num())
+	FShaderMapAssetPaths Buffer;
+	if (!AssociatedAssets)
 	{
-		ShaderCodeToAssets.FindOrAdd(Hash).Append(*AssociatedAssets);
+		AssociatedAssets = &Buffer;
 	}
+	return FindOrAddShaderMapEditor(Hash, OutIndex, AssociatedAssets, ECookShaderLibrarySource::CurrentCook);
+#else
+	return FindOrAddShaderMapNonEditor(Hash, OutIndex);
 #endif
+}
 
+#if WITH_EDITOR
+bool FSerializedShaderArchive::FindOrAddShaderMapEditor(const FSHAHash& Hash, int32& OutIndex,
+	const FShaderMapAssetPaths* AssociatedAssets, ECookShaderLibrarySource CookSource)
+{
+	bool bResult = FindOrAddShaderMapNonEditorInternal(Hash, OutIndex);
 
-	OutIndex = Index;
-	return bAdded;
+	TOptional<FShaderMapAssetPaths> Buffer;
+	const FShaderMapAssetPaths* NormalizedAssets = AssociatedAssets;
+	if (!AssociatedAssets || AssociatedAssets->IsEmpty())
+	{
+		UE_LOG(LogShaderLibrary, Error,
+			TEXT("FindOrAddShaderMapEditor must be called with non-empty AssociatedAssets using the ShaderMap. ")
+			TEXT("ShaderMaps without assets will be pruned. Pass NAME_None for ShaderMaps that are readded on every cook. ")
+			TEXT("Hash %s is being assigned to NAME_None and will be pruned next cook."),
+			*Hash.ToString());
+		Buffer.Emplace();
+		Buffer->Add(NAME_None);
+		NormalizedAssets = Buffer.GetPtrOrNull();
+	}
+	for (FName Asset : *NormalizedAssets)
+	{
+		FShaderMapAssetAssociations::FAssociatedAssetData* AssetData = ShaderMapAssetAssociations.FindAsset(Asset);
+		if (AssetData)
+		{
+			// When loading ShaderArchive data from a previous cook, if we have recooked the Asset and
+			// already have data from the current cook, do not add the stale references.
+			if (CookSource >= ECookShaderLibrarySource::PreviousIncremental
+				&& AssetData->LatestSource < ECookShaderLibrarySource::PreviousIncremental)
+			{
+				continue;
+			}
+			// When recooking an asset, the first time we add a ShaderMap for it remove any Assets added from a
+			// previous cook.
+			if (CookSource == ECookShaderLibrarySource::CurrentCook
+				&& AssetData->LatestSource >= ECookShaderLibrarySource::PreviousIncremental)
+			{
+				ShaderMapAssetAssociations.RemoveAsset(Asset);
+				// Add the current ShaderMap to ShaderMapAssetAssociations.
+				AssetData = &ShaderMapAssetAssociations.FindOrAddAsset(Asset, Hash);
+			}
+			else
+			{
+				// We keep both the previous and current ShaderMap in ShaderMapAssetAssociations, so add the current
+				// ShaderMap to the existing ShaderMaps.
+				ShaderMapAssetAssociations.FindOrAddAsset(Asset, Hash);
+			}
+
+			// Report the lower of previous and current source as the current source; lower sources are newer
+			// and we report the newest.
+			AssetData->MergeSource(CookSource);
+		}
+		else
+		{
+			// Add the new Asset to ShaderMapAssetAssociations with the current ShaderMap.
+			ShaderMapAssetAssociations.FindOrAddAsset(Asset, Hash).LatestSource = CookSource;
+		}
+	}
+
+	return bResult;
 }
-#endif
+#else // WITH_EDITOR
+bool FSerializedShaderArchive::FindOrAddShaderMapNonEditor(const FSHAHash& Hash, int32& OutIndex)
+{
+	return FindOrAddShaderMapNonEditorInternal(Hash, OutIndex);
+}
+#endif // #else !WITH_EDITOR
+
+#endif // !USE_MMAPPED_SHADERARCHIVE
 
 int32 FSerializedShaderArchive::FindShaderWithKey(const FSHAHash& Hash, uint32 Key) const
 {
@@ -233,15 +308,14 @@ FCbWriter& operator<<(FCbWriter& Writer, const FSerializedShaderArchive& Archive
 		SerializedBytes.Empty();
 
 		// Serialize is meant for runtime fields only, so copy the editor-only fields separately
-		Writer.BeginArray("ShaderCodeToAssets");
-		for (const TPair<FSHAHash, FShaderMapAssetPaths>& Pair : Archive.ShaderCodeToAssets)
+		Writer.BeginArray("ShaderMapAssetAssociations");
+		for (const TPair<FName, FSerializedShaderArchive::FShaderMapAssetAssociations::FAssociatedAssetData>& Pair
+			: Archive.ShaderMapAssetAssociations.ViewAssets())
 		{
-			Writer << Pair.Key;
 			Writer.BeginArray();
-			for (FName AssetName : Pair.Value)
-			{
-				Writer << AssetName;
-			}
+			Writer << Pair.Key;
+			Writer << (uint8)Pair.Value.LatestSource;
+			Writer << Pair.Value.ShaderMaps;
 			Writer.EndArray();
 		}
 		Writer.EndArray();
@@ -276,31 +350,57 @@ bool LoadFromCompactBinary(FCbFieldView Field, FSerializedShaderArchive& OutArch
 	}
 
 	bool bOk = true;
-
 	{
 		FCbFieldView ShaderCodeToAssetsField = Field["ShaderCodeToAssets"];
-		// Map size is array size divided by two because pairs are written as successive elements
-		int32 NumShaderCodeToAssets = ShaderCodeToAssetsField.AsArrayView().Num() / 2;
-		bOk = !ShaderCodeToAssetsField.HasError();
-		OutArchive.ShaderCodeToAssets.Empty(NumShaderCodeToAssets);
-		FCbFieldViewIterator It = ShaderCodeToAssetsField.CreateViewIterator();
-		while (It)
+		if (ShaderCodeToAssetsField.AsArrayView().Num() > 0)
+		{
+			// Legacy field, no longer supported
+			OutArchive = FSerializedShaderArchive();
+			return false;
+		}
+	}
+
+	{
+
+		FCbFieldView ShaderMapAssetAssociationsField = Field["ShaderMapAssetAssociations"];
+		int32 NumAssets = ShaderMapAssetAssociationsField.AsArrayView().Num() / 2;
+		bOk = !ShaderMapAssetAssociationsField.HasError() && bOk;
+		OutArchive.ShaderMapAssetAssociations.ReserveAssets(NumAssets);
+
+		for (FCbFieldView AssetFieldIter : ShaderMapAssetAssociationsField)
 		{
-			FSHAHash ShaderMapHash;
-			if (!LoadFromCompactBinary(*It++, ShaderMapHash))
+			FCbArrayView AssetField = AssetFieldIter.AsArrayView();
+			if (AssetField.Num() < 3)
 			{
 				bOk = false;
 				continue;
 			}
-			FShaderMapAssetPaths& Paths = OutArchive.ShaderCodeToAssets.FindOrAdd(ShaderMapHash);
-			FCbFieldView AssetNameArrayField = *It++;
-			Paths.Reserve((*It).AsArrayView().Num());
-			bOk = (!AssetNameArrayField.HasError()) && bOk;
-			for (FCbFieldView AssetNameField : AssetNameArrayField)
+			FCbFieldViewIterator ArrayIter = AssetField.CreateViewIterator();
+			FName AssetName;
+			TArray<FSHAHash> ShaderMaps;
+			ShaderCodeArchive::ECookShaderLibrarySource LatestSource
+				= ShaderCodeArchive::ECookShaderLibrarySource::CurrentCook;
+			bOk = LoadFromCompactBinary(*ArrayIter++, AssetName) && bOk;
+			uint8 LatestSourceInt;
+			if (LoadFromCompactBinary(*ArrayIter++, LatestSourceInt)
+				&& LatestSourceInt < (uint8) ShaderCodeArchive::ECookShaderLibrarySource::Count)
+			{
+				LatestSource = (ShaderCodeArchive::ECookShaderLibrarySource)LatestSourceInt;
+			}
+			else
+			{
+				bOk = false;
+			}
+			bOk = LoadFromCompactBinary(*ArrayIter++, ShaderMaps) && bOk;
+			if (bOk && ShaderMaps.Num() > 0)
 			{
-				FName AssetName;
-				bOk = LoadFromCompactBinary(AssetNameField, AssetName) && bOk;
-				Paths.Add(AssetName);
+				FSerializedShaderArchive::FShaderMapAssetAssociations::FAssociatedAssetData* AssetData;
+				AssetData = &OutArchive.ShaderMapAssetAssociations.FindOrAddAsset(AssetName, ShaderMaps[0]);
+				AssetData->LatestSource = LatestSource;
+				for (int32 Index = 1; Index < ShaderMaps.Num(); ++Index)
+				{
+					OutArchive.ShaderMapAssetAssociations.FindOrAddAsset(AssetName, ShaderMaps[Index]);
+				}
 			}
 		}
 	}
@@ -664,16 +764,12 @@ void FSerializedShaderArchive::SaveAssetInfo(FArchive& Ar)
 			Writer->WriteValue(TEXT("AssetInfoVersion"), static_cast<int32>(EAssetInfoVersion::CurrentVersion));
 
 			TArray<const TPair<FSHAHash, FShaderMapAssetPaths>*> SortedData;
-			SortedData.Reserve(ShaderCodeToAssets.Num());
-			for (TPair<FSHAHash, FShaderMapAssetPaths>& Pair : ShaderCodeToAssets)
+			ShaderMapAssetAssociations.SortForSaving();
+			SortedData.Reserve(ShaderMapAssetAssociations.ViewShaderMaps().Num());
+			for (const TPair<FSHAHash, FShaderMapAssetPaths>& Pair : ShaderMapAssetAssociations.ViewShaderMaps())
 			{
 				SortedData.Add(&Pair);
-				Pair.Value.Sort(FNameLexicalLess());
 			}
-			Algo::Sort(SortedData, [](const TPair<FSHAHash, FShaderMapAssetPaths>* const a, const TPair<FSHAHash, FShaderMapAssetPaths>* const b)
-				{
-					return GetTypeHash(a->Key) < GetTypeHash(b->Key);
-				});
 
 			Writer->WriteArrayStart(TEXT("ShaderCodeToAssets"));
 			for (const TPair<FSHAHash, FShaderMapAssetPaths>* Pair : SortedData)
@@ -702,13 +798,13 @@ void FSerializedShaderArchive::SaveAssetInfo(FArchive& Ar)
 	}
 }
 
-bool FSerializedShaderArchive::LoadAssetInfo(const FString& Filename)
+bool FSerializedShaderArchive::LoadAssetInfo(const FString& Filename, ECookShaderLibrarySource CookSource)
 {
 	TUniquePtr<FArchive> Reader(IFileManager::Get().CreateFileReader(*Filename));
-	return LoadAssetInfo(Reader.Get());
+	return LoadAssetInfo(Reader.Get(), CookSource);
 }
 
-bool FSerializedShaderArchive::LoadAssetInfo(FArchive* Ar)
+bool FSerializedShaderArchive::LoadAssetInfo(FArchive* Ar, ECookShaderLibrarySource CookSource)
 {
 	if (!Ar)
 	{
@@ -789,20 +885,74 @@ bool FSerializedShaderArchive::LoadAssetInfo(FArchive* Ar)
 			);
 			return false;
 		}
-			
-		FShaderMapAssetPaths Paths;
+
 		TArray<TSharedPtr<FJsonValue>> AssetPathsArray = AssetPathsArrayValue->AsArray();
-		for (int32 IdxAsset = 0, NumAssets = AssetPathsArray.Num(); IdxAsset < NumAssets; ++IdxAsset)
+		for (TSharedPtr<FJsonValue>& AssetPath : AssetPathsArray)
 		{
-			Paths.Add(FName(*AssetPathsArray[IdxAsset]->AsString()));
+			FShaderMapAssetAssociations::FAssociatedAssetData& AssetData
+				= ShaderMapAssetAssociations.FindOrAddAsset(FName(*AssetPath->AsString()), ShaderMapHash);
+			AssetData.MergeSource(CookSource);
 		}
-
-		ShaderCodeToAssets.Add(ShaderMapHash, Paths);
 	}
 
 	return true;
 }
 
+void FSerializedShaderArchive::FShaderMapAssetAssociations::ReserveAssets(int32 NumAssets)
+{
+	ShaderMapToAssets.Reserve(NumAssets);
+	AssetToShaderMaps.Reserve(NumAssets);
+}
+
+FSerializedShaderArchive::FShaderMapAssetAssociations::FAssociatedAssetData&
+FSerializedShaderArchive::FShaderMapAssetAssociations::FindOrAddAsset(FName Asset, const FSHAHash& ShaderMap)
+{
+	FAssociatedAssetData& Result = AssetToShaderMaps.FindOrAdd(Asset);
+	FShaderMapAssetPaths& Assets = ShaderMapToAssets.FindOrAdd(ShaderMap);
+	bool bAlreadyExists = false;
+	Assets.Add(Asset, &bAlreadyExists);
+	if (!bAlreadyExists)
+	{
+		Result.ShaderMaps.AddUnique(ShaderMap);
+	}
+	check(AssetToShaderMaps.IsEmpty() == ShaderMapToAssets.IsEmpty());
+	return Result;
+}
+
+void FSerializedShaderArchive::FShaderMapAssetAssociations::RemoveAsset(FName Asset)
+{
+	FAssociatedAssetData Removed;
+	if (AssetToShaderMaps.RemoveAndCopyValue(Asset, Removed))
+	{
+		for (const FSHAHash& ShaderMap : Removed.ShaderMaps)
+		{
+			FShaderMapAssetPaths* AssetsUsingShader = ShaderMapToAssets.Find(ShaderMap);
+			if (AssetsUsingShader)
+			{
+				AssetsUsingShader->Remove(Asset);
+				if (AssetsUsingShader->IsEmpty())
+				{
+					ShaderMapToAssets.Remove(ShaderMap);
+				}
+			}
+		}
+	}
+	check(AssetToShaderMaps.IsEmpty() == ShaderMapToAssets.IsEmpty());
+}
+
+void FSerializedShaderArchive::FShaderMapAssetAssociations::SortForSaving()
+{
+	ShaderMapToAssets.KeySort(
+		[](const FSHAHash& A, const FSHAHash& B)
+		{
+			return GetTypeHash(A) < GetTypeHash(B);
+		});
+	for (TPair<FSHAHash, FShaderMapAssetPaths>& Pair : ShaderMapToAssets)
+	{
+		Pair.Value.Sort(FNameLexicalLess());
+	}
+}
+
 void FSerializedShaderArchive::CreateAsChunkFrom(const FSerializedShaderArchive& Parent, const TSet<FName>& PackagesInChunk, TArray<int32>& OutShaderCodeEntriesNeeded)
 {
 	// we should begin with a clean slate
@@ -813,7 +963,7 @@ void FSerializedShaderArchive::CreateAsChunkFrom(const FSerializedShaderArchive&
 	for (int32 IdxSM = 0, NumSMs = Parent.ShaderMapHashes.Num(); IdxSM < NumSMs; ++IdxSM)
 	{
 		const FSHAHash& ShaderMapHash = Parent.ShaderMapHashes[IdxSM];
-		const FShaderMapAssetPaths* Assets = Parent.ShaderCodeToAssets.Find(ShaderMapHash);
+		const FShaderMapAssetPaths* Assets = Parent.ShaderMapAssetAssociations.FindShaderMap(ShaderMapHash);
 		bool bIncludeSM = false;
 		if (UNLIKELY(Assets == nullptr))
 		{
@@ -837,7 +987,7 @@ void FSerializedShaderArchive::CreateAsChunkFrom(const FSerializedShaderArchive&
 		{
 			// add this shader map
 			int32 ShaderMapIndex = INDEX_NONE;
-			if (FindOrAddShaderMap(ShaderMapHash, ShaderMapIndex, Assets))
+			if (FindOrAddShaderMapEditor(ShaderMapHash, ShaderMapIndex, Assets, ECookShaderLibrarySource::CurrentCook))
 			{
 				// if we're in this scope, it means it's a new shadermap for the chunk and we need more information about it from the parent
 				int32 ParentShaderMapIndex = IdxSM;
@@ -863,9 +1013,7 @@ void FSerializedShaderArchive::CreateAsChunkFrom(const FSerializedShaderArchive&
 
 						// copy the entry as is
 						ShaderEntries[ShaderIndex] = Parent.ShaderEntries[ParentShaderIndex];
-#if WITH_EDITOR
 						ShaderTypes[ShaderIndex] = Parent.ShaderTypes[ParentShaderIndex];
-#endif
 					}
 					ShaderIndices[ShaderMapDescriptor.ShaderIndicesOffset + ShaderIdx] = ShaderIndex;
 				}
@@ -898,19 +1046,7 @@ void FSerializedShaderArchive::CollectStatsAndDebugInfo(FDebugStats& OutDebugSta
 	}
 	OutDebugStats.NumShaders = TotalShaders;
 	OutDebugStats.ShadersSize = TotalShaderSize;
-
-	// this is moderately expensive, consider moving to ExtendedStats?
-	{
-		TSet<FName> AllAssets;
-		for (TMap<FSHAHash, FShaderMapAssetPaths>::TConstIterator Iter(ShaderCodeToAssets); Iter; ++Iter)
-		{
-			for (const FName& AssetName : Iter.Value())
-			{
-				AllAssets.Add(AssetName);
-			}
-		}
-		OutDebugStats.NumAssets = AllAssets.Num();
-	}
+	OutDebugStats.NumAssets = ShaderMapAssetAssociations.ViewAssets().Num();
 
 	int64 ActuallySavedShaderSize = 0;
 	for (const FShaderCodeEntry& ShaderEntry : ShaderEntries)
diff --git a/Engine/Source/Runtime/RenderCore/Private/ShaderLibrary/ShaderCodeLibrary.cpp b/Engine/Source/Runtime/RenderCore/Private/ShaderLibrary/ShaderCodeLibrary.cpp
--- a/Engine/Source/Runtime/RenderCore/Private/ShaderLibrary/ShaderCodeLibrary.cpp
+++ b/Engine/Source/Runtime/RenderCore/Private/ShaderLibrary/ShaderCodeLibrary.cpp
@@ -1403,10 +1403,9 @@ public:
 		FEditorShaderCodeArchive* CodeArchive = EditorShaderCodeArchive[Platform];
 		checkf(CodeArchive, TEXT("EditorShaderCodeArchive for (EShaderPlatform)%d is null!"), (int32)Platform);
 
-		CodeArchive->AddShaderCode(CodeStats, Code, AssociatedAssets);
+		CodeArchive->AddShaderCode(CodeStats, Code, AssociatedAssets, ShaderCodeArchive::ECookShaderLibrarySource::CurrentCook);
 	}
 
-#if WITH_EDITOR
 	void CopyToCompactBinaryAndClear(FCbWriter& Writer, bool& bOutHasData, bool& bOutRanOutOfRoom, int64 MaxShaderSize)
 	{
 		bOutRanOutOfRoom = false;
@@ -1451,7 +1450,8 @@ public:
 					int64 MaxShaderSizeThisCall = RemainingSize;
 					int64 MaxShaderCount = -1;
 					bool bRanOutOfRoom;
-					CodeArchive->CopyToArchiveAndClear(TransferArchive, TransferCode, bRanOutOfRoom, MaxShaderSizeThisCall, MaxShaderCount);
+					CodeArchive->CopyToArchiveAndClear(TransferArchive, TransferCode, bRanOutOfRoom, MaxShaderSizeThisCall,
+						MaxShaderCount, ShaderCodeArchive::ECookShaderLibrarySource::CurrentCook);
 					bOutRanOutOfRoom |= bRanOutOfRoom;
 					if (bRanOutOfRoom && TransferCode.IsEmpty() && RemainingSize == MaxShaderSize)
 					{
@@ -1461,7 +1461,8 @@ public:
 						MaxShaderSizeThisCall = -1;
 						MaxShaderCount = 1;
 						TransferArchive.Empty();
-						CodeArchive->CopyToArchiveAndClear(TransferArchive, TransferCode, bRanOutOfRoom, MaxShaderSizeThisCall, MaxShaderCount);
+						CodeArchive->CopyToArchiveAndClear(TransferArchive, TransferCode, bRanOutOfRoom,
+							MaxShaderSizeThisCall, MaxShaderCount, ShaderCodeArchive::ECookShaderLibrarySource::CurrentCook);
 					}
 					if (!TransferArchive.IsEmpty())
 					{
@@ -1510,7 +1511,8 @@ public:
 				if (CodeArchive)
 				{
 					FShaderCodeStats& CodeStats = EditorShaderCodeStats[Platform];
-					bOk = CodeArchive->AppendFromCompactBinary(CodeArchiveField, CodeStats) && bOk;
+					bOk = CodeArchive->AppendFromCompactBinary(CodeArchiveField, CodeStats,
+						ShaderCodeArchive::ECookShaderLibrarySource::CurrentCook) && bOk;
 				}
 				else
 				{
@@ -1532,7 +1534,6 @@ public:
 		}
 		return bOk;
 	}
-#endif
 
 	void AddShaderStableKeyValue(EShaderPlatform InShaderPlatform, FStableShaderKeyAndValue& StableKeyValue)
 	{
@@ -1559,7 +1560,8 @@ public:
 			FEditorShaderCodeArchive* CodeArchive = EditorShaderCodeArchive[SPlatform];
 			if (CodeArchive)
 			{
-				CodeArchive->FinishPopulate(ShaderCodeDir, CodeStats);
+				CodeArchive->FinishPopulate(ShaderCodeDir, CodeStats,
+					ShaderCodeArchive::ECookShaderLibrarySource::PreviousIncremental);
 			}
 
 			FEditorShaderStableInfo* StableArchive = EditorShaderStableInfo[SPlatform];
@@ -1570,6 +1572,34 @@ public:
 		}
 	}
 
+	void UpdateOplogPackages(UE::Cook::Artifact::FUpdateOplogPackagesContext& Context, const TArray<FName>& ShaderFormats)
+	{
+		for (FName ShaderFormatName : ShaderFormats)
+		{
+			EShaderPlatform SPlatform = ShaderFormatToLegacyShaderPlatform(ShaderFormatName);
+
+			FEditorShaderCodeArchive* CodeArchive = EditorShaderCodeArchive[SPlatform];
+			if (CodeArchive)
+			{
+				CodeArchive->UpdateOplogPackages(Context);
+			}
+		}
+	}
+
+	void PruneShaderLibrary(const TArray<FName>& ShaderFormats)
+	{
+		for (FName ShaderFormatName : ShaderFormats)
+		{
+			EShaderPlatform SPlatform = ShaderFormatToLegacyShaderPlatform(ShaderFormatName);
+
+			FEditorShaderCodeArchive* CodeArchive = EditorShaderCodeArchive[SPlatform];
+			if (CodeArchive)
+			{
+				CodeArchive->Prune();
+			}
+		}
+	}
+
 	bool SaveShaderCode(const FString& ShaderCodeDir, const FString& MetaOutputDir, const TArray<FName>& ShaderFormats, TArray<FString>& OutSCLCSVPath)
 	{
 		bool bOk = ShaderFormats.Num() > 0;
@@ -1814,7 +1844,7 @@ public:
 			PlatformId++;
 		}
 	}
-#endif// WITH_EDITOR
+#endif // WITH_EDITOR
 
 	void PreloadPackageShaderMapsDelegate(TArrayView<const FSHAHash> ShaderMapHashes, FCoreDelegates::FAttachShaderReadRequestFunc AttachShaderReadRequestFunc)
 	{
@@ -2390,12 +2420,22 @@ void FShaderLibraryCooker::FinishPopulateShaderLibrary(const ITargetPlatform* Ta
 
 void FShaderLibraryCooker::UpdateOplogPackages(UE::Cook::Artifact::FUpdateOplogPackagesContext& Context)
 {
-
+	TArray<FName> ShaderFormats;
+	Context.GetTargetPlatform()->GetAllTargetedShaderFormats(ShaderFormats);
+	if (ShaderFormats.Num() > 0)
+	{
+		FShaderLibrariesCollection::Impl->UpdateOplogPackages(Context, ShaderFormats);
+	}
 }
 
 void FShaderLibraryCooker::PruneShaderLibrary(const ITargetPlatform* TargetPlatform, const FString& Name)
 {
-
+	TArray<FName> ShaderFormats;
+	TargetPlatform->GetAllTargetedShaderFormats(ShaderFormats);
+	if (ShaderFormats.Num() > 0)
+	{
+		FShaderLibrariesCollection::Impl->PruneShaderLibrary(ShaderFormats);
+	}
 }
 
 bool FShaderLibraryCooker::MergeShaderCodeArchive(const TArray<FString>& CookedMetadataDirs, const FString& OutputDir, TArray<FString>& OutWrittenFiles)
@@ -2423,7 +2463,8 @@ bool FShaderLibraryCooker::MergeShaderCodeArchive(const TArray<FString>& CookedM
 		{
 			if (ShaderCodeArchives.Contains(ByteCodeFile))
 			{
-				ShaderCodeArchives[ByteCodeFile].Key.AddShaderCodeLibraryFromDirectory(ShaderCodeDir, ShaderCodeArchives[ByteCodeFile].Value);
+				ShaderCodeArchives[ByteCodeFile].Key.AddShaderCodeLibraryFromDirectory(ShaderCodeDir,
+					ShaderCodeArchives[ByteCodeFile].Value, ShaderCodeArchive::ECookShaderLibrarySource::CurrentCook);
 			}
 			else
 			{
@@ -2434,7 +2475,8 @@ bool FShaderLibraryCooker::MergeShaderCodeArchive(const TArray<FString>& CookedM
 					FEditorShaderCodeArchive NewCodeArchive(ShaderFormat, bNeedsDeterministicOrder);
 					FShaderCodeStats NewCodeStats;
 					NewCodeArchive.OpenLibrary(FindShaderFormat.GetCaptureGroup(1));
-					NewCodeArchive.AddShaderCodeLibraryFromDirectory(ShaderCodeDir, NewCodeStats);
+					NewCodeArchive.AddShaderCodeLibraryFromDirectory(ShaderCodeDir, NewCodeStats,
+						ShaderCodeArchive::ECookShaderLibrarySource::CurrentCook);
 
 					ShaderCodeArchives.FindOrAdd(ByteCodeFile, { MoveTemp(NewCodeArchive), MoveTemp(NewCodeStats) });
 				}
diff --git a/Engine/Source/Runtime/RenderCore/Public/ShaderCodeArchive.h b/Engine/Source/Runtime/RenderCore/Public/ShaderCodeArchive.h
--- a/Engine/Source/Runtime/RenderCore/Public/ShaderCodeArchive.h
+++ b/Engine/Source/Runtime/RenderCore/Public/ShaderCodeArchive.h
@@ -84,10 +84,38 @@ struct FShaderCodeEntry
 #pragma pack(pop)
 #endif
 
+#if WITH_EDITOR
+namespace ShaderCodeArchive
+{
+
+/**
+ * Descriptor when loading a ShaderLibrary during cooking for which operation is loading the ShaderLibrary.
+ * When the operation is loading data from previous incremental cooks, we have more strict validation requirements
+ * and we ignore data for stale Assets that have been recooked.
+ * Smaller values are newer sources and override larger values which come from older (more stale) sources.
+ */
+enum class ECookShaderLibrarySource : uint8
+{
+	/** The data is provided by compilation that executed during the current cook; this is the most up-to-date. */
+	CurrentCook = 0,
+	/** The current cook is a DLC cook and the data being provided came from the basegame cook. */
+	PatchBase,
+	/** The data was written to the ShaderLibrary artifacts in a previous incremental cook. */
+	PreviousIncremental,
+	Count,
+};
+ 
+}
+
+#endif // WITH_EDITOR
+
 // Portion of shader code archive that's serialize to disk
 class FSerializedShaderArchive
 {
 public:
+#if WITH_EDITOR
+	using ECookShaderLibrarySource = ShaderCodeArchive::ECookShaderLibrarySource;
+#endif
 #if USE_MMAPPED_SHADERARCHIVE
 	template<typename T> using TArrayType = TArrayView<T>;
 	// Ensure use of the const Get functions when mmapping is used.
@@ -128,7 +156,58 @@ public:
 	FHashTable ShaderHashTable;
 
 #if WITH_EDITOR
-	/** Mapping from shadermap hashes to an array of asset names - this is used for on-disk storage as it is shorter. */
+	/**
+	 * Bi-directional mapping from shadermap hashes to an array of asset names.
+	 * Assets are identified by LongPackageName.
+	 * ShaderMaps are identified by their SHAHash.
+	 * This structure is used to identify the ShaderMaps required in the projection of this library for a pakfile chunk
+	 * based on assets in the chunk, and to remove stale ShaderMaps and their Shaders from this library when the assets
+	 * are recooked or are pruned from the incremental cook oplog.
+	 */
+	class FShaderMapAssetAssociations
+	{
+	public:
+		struct FAssociatedAssetData
+		{
+			TArray<FSHAHash> ShaderMaps;
+			ECookShaderLibrarySource LatestSource = ECookShaderLibrarySource::PreviousIncremental;
+			// These reference flags are transient, intialized to false, and are read/written only during end of cook.
+			/** Set to true if the Asset is needed either for incremental cook or for staging, and should therefore not be pruned. */
+			bool bReferencedByOplog = false;
+			/** Set to true if the Asset is needed for runtime and staging, and should therefore be added to the staging artifacts. */
+			bool bReferencedByStaging = false;
+
+			/** this->LatestSource = NewestOf(LatestSource, InCookSource). */
+			void MergeSource(ECookShaderLibrarySource InCookSource);
+		};
+		bool IsEmpty() const;
+		SIZE_T GetAllocatedSize() const;
+		TMap<FName, FAssociatedAssetData>& ViewAssets();
+		const TMap<FName, FAssociatedAssetData>& ViewAssets() const;
+		const TMap<FSHAHash, FShaderMapAssetPaths>& ViewShaderMaps() const;
+
+		FAssociatedAssetData* FindAsset(FName AssetName);
+		const FAssociatedAssetData* FindAsset(FName AssetName) const;
+		FShaderMapAssetPaths* FindShaderMap(const FSHAHash& ShaderMap);
+		const FShaderMapAssetPaths* FindShaderMap(const FSHAHash& ShaderMap) const;
+
+		void Empty();
+		RENDERCORE_API void ReserveAssets(int32 NumAssets);
+		RENDERCORE_API FAssociatedAssetData& FindOrAddAsset(FName Asset, const FSHAHash& ShaderMap);
+		RENDERCORE_API void RemoveAsset(FName Asset);
+		RENDERCORE_API void SortForSaving();
+
+	private:
+		TMap<FSHAHash, FShaderMapAssetPaths> ShaderMapToAssets;
+		TMap<FName, FAssociatedAssetData> AssetToShaderMaps;
+	};
+private:
+	FShaderMapAssetAssociations ShaderMapAssetAssociations;
+public:
+	FShaderMapAssetAssociations& GetShaderMapAssetAssociations();
+	const FShaderMapAssetAssociations& GetShaderMapAssetAssociations() const;
+
+	UE_DEPRECATED(5.8, "No longer referenced. Use GetShaderMapAssetAssociations() instead.")
 	TMap<FSHAHash, FShaderMapAssetPaths> ShaderCodeToAssets;
 
 	struct FShaderTypeHashes
@@ -186,11 +265,15 @@ public:
 		/** Compressed (individually) size of all shaders of a given frequency. */
 		uint64 CompressedSizePerFrequency[SF_NumFrequencies];
 	};
-#endif
+#endif // WITH_EDITOR
 
 	FSerializedShaderArchive()
 	{
 	}
+	FSerializedShaderArchive(FSerializedShaderArchive&&);
+	FSerializedShaderArchive(const FSerializedShaderArchive&);
+	FSerializedShaderArchive& operator=(FSerializedShaderArchive&&);
+	FSerializedShaderArchive& operator=(const FSerializedShaderArchive&);
 
 	int64 GetAllocatedSize() const
 	{
@@ -210,7 +293,7 @@ public:
 			ShaderIndices.GetAllocatedSize()
 #endif
 #if WITH_EDITOR
-			+ ShaderCodeToAssets.GetAllocatedSize()
+			+ ShaderMapAssetAssociations.GetAllocatedSize()
 			+ ShaderTypes.GetAllocatedSize()
 #endif // WITH_EDITOR
 			;
@@ -239,7 +322,7 @@ public:
 #endif
 		ShaderMapHashTable.Clear();
 #if WITH_EDITOR
-		ShaderCodeToAssets.Empty();
+		ShaderMapAssetAssociations.Empty();
 #endif
 	}
 
@@ -257,7 +340,7 @@ public:
 	{
 		return ShaderMapEntries.IsEmpty() && ShaderEntries.IsEmpty() && PreloadEntries.IsEmpty()
 #if WITH_EDITOR
-			&& ShaderCodeToAssets.IsEmpty()
+			&& ShaderMapAssetAssociations.IsEmpty()
 			&& ShaderTypes.IsEmpty()
 #endif
 			;
@@ -270,18 +353,25 @@ public:
 
 #if !USE_MMAPPED_SHADERARCHIVE
 	RENDERCORE_API bool FindOrAddShader(const FSHAHash& Hash, int32& OutIndex);
+	UE_DEPRECATED(5.8, "Call FindOrAddShaderMapNonEditor or FindOrAddShaderMapEditor.")
 	RENDERCORE_API bool FindOrAddShaderMap(const FSHAHash& Hash, int32& OutIndex, const FShaderMapAssetPaths* AssociatedAssets);
+#if WITH_EDITOR
+	RENDERCORE_API bool FindOrAddShaderMapEditor(const FSHAHash& Hash, int32& OutIndex,
+		const FShaderMapAssetPaths* AssociatedAssets, ECookShaderLibrarySource CookSource);
+#else
+	RENDERCORE_API bool FindOrAddShaderMapNonEditor(const FSHAHash& Hash, int32& OutIndex);
+#endif
 	RENDERCORE_API void RemoveLastAddedShader();
 	RENDERCORE_API void Finalize();
-#endif
+#endif // !USE_MMAPPED_SHADERARCHIVE
 
 	RENDERCORE_API void DecompressShader(int32 Index, const TArray<FSharedBuffer>& ShaderCode, TArray<uint8>& OutDecompressedShader) const;
 
 	RENDERCORE_API void Serialize(FArchive& Ar);
 #if WITH_EDITOR
 	RENDERCORE_API void SaveAssetInfo(FArchive& Ar);
-	RENDERCORE_API bool LoadAssetInfo(const FString& Filename);
-	RENDERCORE_API bool LoadAssetInfo(FArchive* Ar);
+	RENDERCORE_API bool LoadAssetInfo(const FString& Filename, ECookShaderLibrarySource CookSource = ECookShaderLibrarySource::CurrentCook);
+	RENDERCORE_API bool LoadAssetInfo(FArchive* Ar, ECookShaderLibrarySource CookSource = ECookShaderLibrarySource::CurrentCook);
 	RENDERCORE_API void CreateAsChunkFrom(const FSerializedShaderArchive& Parent, const TSet<FName>& PackagesInChunk, TArray<int32>& OutShaderCodeEntriesNeeded);
 	RENDERCORE_API void CollectStatsAndDebugInfo(FDebugStats& OutDebugStats, FExtendedDebugStats* OutExtendedDebugStats);
 	RENDERCORE_API void DumpContentsInPlaintext(FString& OutText) const;
@@ -294,6 +384,11 @@ public:
 		Ref.Serialize(Ar);
 		return Ar;
 	}
+
+private:
+#if !USE_MMAPPED_SHADERARCHIVE
+	bool FindOrAddShaderMapNonEditorInternal(const FSHAHash& Hash, int32& OutIndex);
+#endif
 };
 
 // run-time only debugging facility
@@ -800,3 +895,98 @@ private:
 	/** debug visualizer - in Shipping compiles out to an empty struct with no-op functions */
 	FShaderUsageVisualizer DebugVisualizer;
 };
+
+
+///////////////////////////////////////////////////////
+// Inline implementations
+///////////////////////////////////////////////////////
+
+
+PRAGMA_DISABLE_DEPRECATION_WARNINGS
+inline FSerializedShaderArchive::FSerializedShaderArchive(FSerializedShaderArchive&&) = default;
+inline FSerializedShaderArchive::FSerializedShaderArchive(const FSerializedShaderArchive&) = default;
+inline FSerializedShaderArchive& FSerializedShaderArchive::operator=(FSerializedShaderArchive&&) = default;
+inline FSerializedShaderArchive& FSerializedShaderArchive::operator=(const FSerializedShaderArchive&) = default;
+PRAGMA_ENABLE_DEPRECATION_WARNINGS
+
+#if WITH_EDITOR
+
+inline FSerializedShaderArchive::FShaderMapAssetAssociations&
+FSerializedShaderArchive::GetShaderMapAssetAssociations()
+{
+	return ShaderMapAssetAssociations;
+}
+
+inline const FSerializedShaderArchive::FShaderMapAssetAssociations&
+FSerializedShaderArchive::GetShaderMapAssetAssociations() const
+{
+	return ShaderMapAssetAssociations;
+}
+
+inline bool FSerializedShaderArchive::FShaderMapAssetAssociations::IsEmpty() const
+{
+	return ShaderMapToAssets.IsEmpty();
+}
+
+inline SIZE_T FSerializedShaderArchive::FShaderMapAssetAssociations::GetAllocatedSize() const
+{
+	return ShaderMapToAssets.GetAllocatedSize() + AssetToShaderMaps.GetAllocatedSize();
+}
+
+inline TMap<FName, FSerializedShaderArchive::FShaderMapAssetAssociations::FAssociatedAssetData>&
+FSerializedShaderArchive::FShaderMapAssetAssociations::ViewAssets()
+{
+	return AssetToShaderMaps;
+}
+
+inline const TMap<FName, FSerializedShaderArchive::FShaderMapAssetAssociations::FAssociatedAssetData>&
+FSerializedShaderArchive::FShaderMapAssetAssociations::ViewAssets() const
+{
+	return AssetToShaderMaps;
+}
+
+inline const TMap<FSHAHash, FShaderMapAssetPaths>&
+FSerializedShaderArchive::FShaderMapAssetAssociations::ViewShaderMaps() const
+{
+	return ShaderMapToAssets;
+}
+
+inline FSerializedShaderArchive::FShaderMapAssetAssociations::FAssociatedAssetData*
+FSerializedShaderArchive::FShaderMapAssetAssociations::FindAsset(FName AssetName)
+{
+	return AssetToShaderMaps.Find(AssetName);
+}
+
+inline const FSerializedShaderArchive::FShaderMapAssetAssociations::FAssociatedAssetData*
+FSerializedShaderArchive::FShaderMapAssetAssociations::FindAsset(FName AssetName) const
+{
+	return AssetToShaderMaps.Find(AssetName);
+
+}
+
+inline FShaderMapAssetPaths*
+FSerializedShaderArchive::FShaderMapAssetAssociations::FindShaderMap(const FSHAHash& ShaderMap)
+{
+	return ShaderMapToAssets.Find(ShaderMap);
+}
+
+inline const FShaderMapAssetPaths*
+FSerializedShaderArchive::FShaderMapAssetAssociations::FindShaderMap(const FSHAHash& ShaderMap) const
+{
+	return ShaderMapToAssets.Find(ShaderMap);
+}
+
+inline void FSerializedShaderArchive::FShaderMapAssetAssociations::Empty()
+{
+	ShaderMapToAssets.Empty();
+	AssetToShaderMaps.Empty();
+}
+
+inline void FSerializedShaderArchive::FShaderMapAssetAssociations::FAssociatedAssetData::MergeSource(
+	ShaderCodeArchive::ECookShaderLibrarySource InCookSource)
+{
+	using namespace ShaderCodeArchive;
+	LatestSource = (ECookShaderLibrarySource)FMath::Min((uint8)LatestSource, (uint8)InCookSource);
+}
+
+#endif // WITH_EDITOR
\ No newline at end of file
diff --git a/Engine/Source/Runtime/RenderCore/Public/ShaderCodeLibrary.h b/Engine/Source/Runtime/RenderCore/Public/ShaderCodeLibrary.h
--- a/Engine/Source/Runtime/RenderCore/Public/ShaderCodeLibrary.h
+++ b/Engine/Source/Runtime/RenderCore/Public/ShaderCodeLibrary.h
@@ -354,7 +354,7 @@ struct FShaderLibraryCooker
 	static RENDERCORE_API void UpdateOplogPackages(UE::Cook::Artifact::FUpdateOplogPackagesContext& Context);
 
 	/**
-	 * Called by the cooker to ask the library to prune the non-longer referenced shaders from its incremental cook library,
+	 * Called by the cooker to ask the library to prune the no-longer referenced shaders from its incremental cook library,
 	 * called only after UpdateOplogPackages has been called on every platform.
 	 */
 	static RENDERCORE_API void PruneShaderLibrary(const ITargetPlatform* TargetPlatform, const FString& Name);
