Shader IncrementalCook: Temporarily downgrade error that shaders were submitted with no associated assets from error to display. It is occuring on the GlobalShaderMap when cooking EngineTest.
#rnx

#Synced-CL 48939402

diff --git a/Engine/Source/Runtime/RenderCore/Private/ShaderLibrary/ShaderCodeArchive.cpp b/Engine/Source/Runtime/RenderCore/Private/ShaderLibrary/ShaderCodeArchive.cpp
--- a/Engine/Source/Runtime/RenderCore/Private/ShaderLibrary/ShaderCodeArchive.cpp
+++ b/Engine/Source/Runtime/RenderCore/Private/ShaderLibrary/ShaderCodeArchive.cpp
@@ -178,7 +178,7 @@ bool FSerializedShaderArchive::FindOrAddShaderMapEditor(const FSHAHash& Hash, in
 	const FShaderMapAssetPaths* NormalizedAssets = AssociatedAssets;
 	if (!AssociatedAssets || AssociatedAssets->IsEmpty())
 	{
-		UE_LOG(LogShaderLibrary, Error,
+		UE_LOG(LogShaderLibrary, Display,
 			TEXT("FindOrAddShaderMapEditor must be called with non-empty AssociatedAssets using the ShaderMap. ")
 			TEXT("ShaderMaps without assets will be pruned. Pass NAME_None for ShaderMaps that are readded on every cook. ")
 			TEXT("Hash %s is being assigned to NAME_None and will be pruned next cook."),
