ShaderLibrary IncrementalCook: Speculative fix for warning during incremental cook:
"LogShaderLibrary: Warning: Shadermap %sis not associated with any asset. Including it in every chunk"
This is probably caused by the new pruning code blocking the registration of shadermap to asset for the stale shaders from previous cook. Disable that code until we finish the implementation that prunes the shadermaps.
#rnx
#Synced-CL 48939402

diff --git a/Engine/Source/Runtime/RenderCore/Private/ShaderLibrary/EditorShaderCodeArchive.cpp b/Engine/Source/Runtime/RenderCore/Private/ShaderLibrary/EditorShaderCodeArchive.cpp
--- a/Engine/Source/Runtime/RenderCore/Private/ShaderLibrary/EditorShaderCodeArchive.cpp
+++ b/Engine/Source/Runtime/RenderCore/Private/ShaderLibrary/EditorShaderCodeArchive.cpp
@@ -778,8 +778,7 @@ void FEditorShaderCodeArchive::Prune()
 {
 	using FShaderMapAssetAssociations = FSerializedShaderArchive::FShaderMapAssetAssociations;
 
-	constexpr bool bImplemented = false;
-	if (!bImplemented)
+	if (!UE::ShaderLibrary::Private::bIncrementalCookPruningEnabled)
 	{
 		return;
 	}
diff --git a/Engine/Source/Runtime/RenderCore/Private/ShaderLibrary/ShaderCodeArchive.cpp b/Engine/Source/Runtime/RenderCore/Private/ShaderLibrary/ShaderCodeArchive.cpp
--- a/Engine/Source/Runtime/RenderCore/Private/ShaderLibrary/ShaderCodeArchive.cpp
+++ b/Engine/Source/Runtime/RenderCore/Private/ShaderLibrary/ShaderCodeArchive.cpp
@@ -27,6 +27,7 @@
 #include "Misc/StringBuilder.h"
 #include "Serialization/CompactBinarySerialization.h"
 #include "Serialization/CompactBinaryWriter.h"
+#include "ShaderLibrary/ShaderCodeLibraryUtilities.h"
 #include "Templates/Greater.h"
 #endif
 
@@ -192,21 +193,30 @@ bool FSerializedShaderArchive::FindOrAddShaderMapEditor(const FSHAHash& Hash, in
 		FShaderMapAssetAssociations::FAssociatedAssetData* AssetData = ShaderMapAssetAssociations.FindAsset(Asset);
 		if (AssetData)
 		{
-			// When loading ShaderArchive data from a previous cook, if we have recooked the Asset and
-			// already have data from the current cook, do not add the stale references.
-			if (CookSource >= ECookShaderLibrarySource::PreviousIncremental
-				&& AssetData->LatestSource < ECookShaderLibrarySource::PreviousIncremental)
+			if (UE::ShaderLibrary::Private::bIncrementalCookPruningEnabled)
 			{
-				continue;
-			}
-			// When recooking an asset, the first time we add a ShaderMap for it remove any Assets added from a
-			// previous cook.
-			if (CookSource == ECookShaderLibrarySource::CurrentCook
-				&& AssetData->LatestSource >= ECookShaderLibrarySource::PreviousIncremental)
-			{
-				ShaderMapAssetAssociations.RemoveAsset(Asset);
-				// Add the current ShaderMap to ShaderMapAssetAssociations.
-				AssetData = &ShaderMapAssetAssociations.FindOrAddAsset(Asset, Hash);
+				// When loading ShaderArchive data from a previous cook, if we have recooked the Asset and
+				// already have data from the current cook, do not add the stale references.
+				if (CookSource >= ECookShaderLibrarySource::PreviousIncremental
+					&& AssetData->LatestSource < ECookShaderLibrarySource::PreviousIncremental)
+				{
+					continue;
+				}
+				// When recooking an asset, the first time we add a ShaderMap for it remove any Assets added from a
+				// previous cook.
+				if (CookSource == ECookShaderLibrarySource::CurrentCook
+					&& AssetData->LatestSource >= ECookShaderLibrarySource::PreviousIncremental)
+				{
+					ShaderMapAssetAssociations.RemoveAsset(Asset);
+					// Add the current ShaderMap to ShaderMapAssetAssociations.
+					AssetData = &ShaderMapAssetAssociations.FindOrAddAsset(Asset, Hash);
+				}
+				else
+				{
+					// We keep both the previous and current ShaderMap in ShaderMapAssetAssociations, so add the current
+					// ShaderMap to the existing ShaderMaps.
+					ShaderMapAssetAssociations.FindOrAddAsset(Asset, Hash);
+				}
 			}
 			else
 			{
diff --git a/Engine/Source/Runtime/RenderCore/Private/ShaderLibrary/ShaderCodeLibraryUtilities.h b/Engine/Source/Runtime/RenderCore/Private/ShaderLibrary/ShaderCodeLibraryUtilities.h
--- a/Engine/Source/Runtime/RenderCore/Private/ShaderLibrary/ShaderCodeLibraryUtilities.h
+++ b/Engine/Source/Runtime/RenderCore/Private/ShaderLibrary/ShaderCodeLibraryUtilities.h
@@ -15,6 +15,9 @@ class ICookArtifactReader;
 
 namespace UE::ShaderLibrary::Private
 {
+	// Compile time switch to enable the new feature of pruning stale shader data from incremental cooks.
+	// Remove this switch once the feature is stable enough not to need an off switch.
+	constexpr bool bIncrementalCookPruningEnabled = false;
 
 	// Shader library constants:
 	extern const uint32 GShaderCodeArchiveVersion;
