Shader cooking: Fix missing AssociatedAssets in FGlobalShaderMapSection. An empty list of AssociatedAssets causes the cook error "FindOrAddShaderMapEditor must be called with non-empty AssociatedAssets using the ShaderMap".
Repro case for the error without this fix: "UnrealEditor.exe EngineTest -run=Cook  -TargetPlatform=<PLATFORM> -Unversioned -compressed -DDC-All-MissTypes=GlobalShaderMap"

Shader IncrementalCook: Restore severity to error of the log message that shaders were submitted with no associated assets. The occurence of the error on GlobalShaderMaps has been fixed.

#rnx
#Synced-CL 49029036

diff --git a/Engine/Source/Runtime/RenderCore/Private/ShaderLibrary/ShaderCodeArchive.cpp b/Engine/Source/Runtime/RenderCore/Private/ShaderLibrary/ShaderCodeArchive.cpp
--- a/Engine/Source/Runtime/RenderCore/Private/ShaderLibrary/ShaderCodeArchive.cpp
+++ b/Engine/Source/Runtime/RenderCore/Private/ShaderLibrary/ShaderCodeArchive.cpp
@@ -179,7 +179,7 @@ bool FSerializedShaderArchive::FindOrAddShaderMapEditor(const FSHAHash& Hash, in
 	const FShaderMapAssetPaths* NormalizedAssets = AssociatedAssets;
 	if (!AssociatedAssets || AssociatedAssets->IsEmpty())
 	{
-		UE_LOG(LogShaderLibrary, Display,
+		UE_LOG(LogShaderLibrary, Error,
 			TEXT("FindOrAddShaderMapEditor must be called with non-empty AssociatedAssets using the ShaderMap. ")
 			TEXT("ShaderMaps without assets will be pruned. Pass NAME_None for ShaderMaps that are readded on every cook. ")
 			TEXT("Hash %s is being assigned to NAME_None and will be pruned next cook."),
diff --git a/Engine/Source/Runtime/RenderCore/Public/GlobalShader.h b/Engine/Source/Runtime/RenderCore/Public/GlobalShader.h
--- a/Engine/Source/Runtime/RenderCore/Public/GlobalShader.h
+++ b/Engine/Source/Runtime/RenderCore/Public/GlobalShader.h
@@ -161,6 +161,7 @@ private:
 	}
 
 	inline FGlobalShaderMapSection(EShaderPlatform InPlatform, const FHashedName& InHashedSourceFilename)
+		: FGlobalShaderMapSection()
 	{
 		AssignContent(new FGlobalShaderMapContent(InPlatform, InHashedSourceFilename));
 	}
