[DumpGPU] Unified Resource Filtering

Reworks how passes and resource filtering is done. Both are now unified under one command-line parameter to allow even more granular filtering. This improves rendering performance and lowers disk space used. A new Cvar is added to control whether the user wants the log file to be written to disk.

#rnx
#tests win64 editor : datagen all added sequences
#Synced-CL 49037196

diff --git a/Engine/Plugins/MovieScene/MovieRenderPipeline/Source/MovieRenderPipelineCore/Private/Graph/MovieGraphDefaultRenderer.cpp b/Engine/Plugins/MovieScene/MovieRenderPipeline/Source/MovieRenderPipelineCore/Private/Graph/MovieGraphDefaultRenderer.cpp
--- a/Engine/Plugins/MovieScene/MovieRenderPipeline/Source/MovieRenderPipelineCore/Private/Graph/MovieGraphDefaultRenderer.cpp
+++ b/Engine/Plugins/MovieScene/MovieRenderPipeline/Source/MovieRenderPipelineCore/Private/Graph/MovieGraphDefaultRenderer.cpp
@@ -356,7 +356,7 @@ void UMovieGraphDefaultRenderer::Render(const FMovieGraphTimeStepData& InTimeSte
 		if (bIsRendering && bWantsDumpGpuDebugCapture && bIsDumpGPUCorrectFrame)
 		{
 			const FString FinalPath = UE::MovieGraph::Capture::PrepareCapturePath(GetOwningGraph(), InTimeStepData.EvaluatedConfig, DebugSettings->DumpGpuCaptureFilePathFormat, TEXT(""), true);
-			const FString Cmd = FString::Printf(TEXT("-noexplorer -path=\"%s\" -pass_path_filters=\"%s\" -resource_name_filters=\"%s\""), *FinalPath, *DebugSettings->DumpGpuCapturePassPathFilters, *DebugSettings->DumpGpuCaptureResourcesFilters);
+			const FString Cmd = FString::Printf(TEXT("-noexplorer -path=\"%s\" -unified_filter=\"%s\""), *FinalPath, *DebugSettings->DumpGpuCaptureUnifiedPassResourceFilters);
 			const FString Context = FString::Printf(TEXT("MRQ Frame: %d"), InTimeStepData.RootFrameNumber.Value);
 			ScopedDumpGPUCapture = MakeUnique<UE::RenderCore::DumpGPU::FDumpScope>(true, *Cmd, *Context);
 		}
diff --git a/Engine/Plugins/MovieScene/MovieRenderPipeline/Source/MovieRenderPipelineCore/Public/Graph/Nodes/MovieGraphDebugNode.h b/Engine/Plugins/MovieScene/MovieRenderPipeline/Source/MovieRenderPipelineCore/Public/Graph/Nodes/MovieGraphDebugNode.h
--- a/Engine/Plugins/MovieScene/MovieRenderPipeline/Source/MovieRenderPipelineCore/Public/Graph/Nodes/MovieGraphDebugNode.h
+++ b/Engine/Plugins/MovieScene/MovieRenderPipeline/Source/MovieRenderPipelineCore/Public/Graph/Nodes/MovieGraphDebugNode.h
@@ -51,10 +51,7 @@ public:
 	uint8 bOverride_DumpGpuCaptureFilePathFormat : 1;
 
 	UPROPERTY(EditAnywhere, BlueprintReadWrite, Category = Overrides, meta = (InlineEditConditionToggle))
-	uint8 bOverride_DumpGpuCapturePassPathFilters : 1;
-
-	UPROPERTY(EditAnywhere, BlueprintReadWrite, Category = Overrides, meta = (InlineEditConditionToggle))
-	uint8 bOverride_DumpGpuCaptureResourcesFilters : 1;
+	uint8 bOverride_DumpGpuCaptureUnifiedPassResourceFilters : 1;
 
 	/** 
 	* If true, automatically trigger RenderDoc to capture rendering information. RenderDoc plugin must be enabled, 
@@ -108,17 +105,11 @@ public:
 
 	/**
 	* If bCaptureFramesWithDumpGpu is true, wildcard enabled filters to apply to the path of the passes to dump.
-	* Support multiple filters separated by "|". For example : *PostProcessing/DOF*|*PostProcessing/MotionBlur*
-	*/
-	UPROPERTY(EditAnywhere, BlueprintReadWrite, Category = "DumpGPU", meta = (EditCondition = "bOverride_DumpGpuCapturePassPathFilters"))
-	FString DumpGpuCapturePassPathFilters = TEXT("*");
-
-	/**
-	* If bCaptureFramesWithDumpGpu is true, wildcard enabled filters to apply to the resources to dump.
-	* Support multiple filters separated by "|". For example : MegaLights.ResolvedDiffuseLighting|MegaLights.ResolvedSpecularLighting
+	* Syntax: Separate pass names to resources with "@". Filter multiple passes:resources by separating them with "|" 
+	* Example: "[REDACTED_EMAIL]|[REDACTED_EMAIL]"
 	*/
-	UPROPERTY(EditAnywhere, BlueprintReadWrite, Category = "DumpGPU", meta = (EditCondition = "bOverride_DumpGpuCaptureResourcesFilters"))
-	FString DumpGpuCaptureResourcesFilters = TEXT("*");
+	UPROPERTY(EditAnywhere, BlueprintReadWrite, Category = "DumpGPU", meta = (EditCondition = "bOverride_DumpGpuCaptureUnifiedPassResourceFilters"))
+	FString DumpGpuCaptureUnifiedPassResourceFilters = TEXT("*");
 };
 
 #undef UE_API
diff --git a/Engine/Source/Runtime/RenderCore/Private/DumpGPU.cpp b/Engine/Source/Runtime/RenderCore/Private/DumpGPU.cpp
--- a/Engine/Source/Runtime/RenderCore/Private/DumpGPU.cpp
+++ b/Engine/Source/Runtime/RenderCore/Private/DumpGPU.cpp
@@ -178,6 +178,13 @@ static TAutoConsoleVariable<FString> GDumpGPUVisualizeResource(
 	TEXT("Name of RDG output resource to automatically open in the dump viewer."),
 	ECVF_Default);
 
+static TAutoConsoleVariable<int32> GDumpGPUEnableLogWrite(
+	TEXT("r.DumpGPU.EnableLogWrite"), 1,
+	TEXT("Enables writing the log file to disk.\n")
+	TEXT(" 0: Does not write log file to disk\n")
+	TEXT(" 1: Logging writes are enabled (default)"),
+	ECVF_Default);
+
 } // namespace
 
 #endif // RDG_DUMP_RESOURCES
@@ -229,6 +236,14 @@ const FName GDumpGPUBufferFenceName("DumpGPU.BufferFence");
 
 }
 
+struct FUnifiedPassResourceFilter
+{
+	FWildcardString PassPattern;
+	FWildcardString ResourcePattern;
+
+	FUnifiedPassResourceFilter(const FString& InPassPattern, const FString& InResourcePattern) : PassPattern(InPassPattern), ResourcePattern(InResourcePattern) {}
+};
+
 class FRDGResourceDumpContext
 {
 public:
@@ -245,8 +260,7 @@ public:
 	int32 DumpedFrameId = 0;
 	int32 FrameCount = 1;
 	FName UploadResourceCompressionName;
-	TArray<FString> PassPathFilters;
-	TArray<FString> ResourceNameFilters;
+	TArray<FUnifiedPassResourceFilter> UnifiedPassResourceFilter;
 	FString DumpingDirectoryPath;
 	FString Context;
 	FDateTime Time;
@@ -1496,7 +1510,7 @@ public:
 		const FRDGTextureSRVDesc& SubresourceRangeDesc,
 		bool bIsOutputResource)
 	{
-		if (!IsDumpingResource(SubresourceRangeDesc.Texture->Name))
+		if (!IsDumpingResource(SubresourceRangeDesc.Texture->Name, Pass))
 		{
 			return;
 		}
@@ -1582,7 +1596,7 @@ public:
 		FRDGBuffer* Buffer,
 		bool bIsOutputResource)
 	{
-		if (!IsDumpingResource(Buffer->Name))
+		if (!IsDumpingResource(Buffer->Name, Pass))
 		{
 			return;
 		}
@@ -1956,13 +1970,13 @@ public:
 		}
 #endif // RDG_EVENTS
 
-		if (!PassPathFilters.IsEmpty())
+		if (!UnifiedPassResourceFilter.IsEmpty())
 		{
 			bool bPathPassAllowed = false;
 
-			for (const FString& PassPathFilter : PassPathFilters)
+			for (const FUnifiedPassResourceFilter& Filter : UnifiedPassResourceFilter)
 			{
-				FWildcardString WildcardPassPathFilter(PassPathFilter);
+				FWildcardString WildcardPassPathFilter(Filter.PassPattern);
 				bPathPassAllowed |= WildcardPassPathFilter.IsMatch(*PassPath);
 			}
 			bDumpPass &= bPathPassAllowed;
@@ -1971,23 +1985,22 @@ public:
 		return bDumpPass;
 	}
 
-	bool IsDumpingResource(const TCHAR* ResourceName) const
+	bool IsDumpingResource(const TCHAR* ResourceName, const FRDGPass* Pass) const
 	{
-		if (ResourceNameFilters.IsEmpty())
+		if (UnifiedPassResourceFilter.IsEmpty())
 		{
 			return true;
 		}
 
-		for (const FString& ResourceFilter : ResourceNameFilters)
+		FString PassName = Pass->GetEventName().GetTCHAR(); 
+		for (const FUnifiedPassResourceFilter& Filter : UnifiedPassResourceFilter)
 		{
-			FWildcardString WildcardResourceFilter(ResourceFilter);
-			if (WildcardResourceFilter.IsMatch(ResourceName))
+			if (Filter.PassPattern.IsMatch(PassName) && Filter.ResourcePattern.IsMatch(ResourceName))
 			{
 				return true;
 			}
 		}
-
-
+        
 		return false;
 	}
 
@@ -2095,14 +2108,25 @@ FString FRDGBuilder::BeginResourceDump(const TCHAR* Cmd, const TCHAR* Context)
 		NewResourceDumpContext->DumpingDirectoryPath = Parameters[TEXT("path")];
 	}
 
-	if (Parameters.Contains(TEXT("pass_path_filters")))
+	if (Parameters.Contains(TEXT("unified_filter")))
 	{
-		Parameters[TEXT("pass_path_filters")].ParseIntoArray(NewResourceDumpContext->PassPathFilters, TEXT("|"));
-	}
+		// Example: "Temporal*@GBufferA|BasePass@SceneDepth*"
+		TArray<FString> RawFilters;
+		Parameters[TEXT("unified_filter")].ParseIntoArray(RawFilters, TEXT("|"));
 
-	if (Parameters.Contains(TEXT("resource_name_filters")))
-	{
-		Parameters[TEXT("resource_name_filters")].ParseIntoArray(NewResourceDumpContext->ResourceNameFilters, TEXT("|"));
+		for (const FString& RawFilter : RawFilters)
+		{
+			FString PassStr, ResourceStr;
+			// Split Resource and Passes into two variables in the same object
+			if (RawFilter.Split(TEXT("@"), &PassStr, &ResourceStr))
+			{
+				NewResourceDumpContext->UnifiedPassResourceFilter.Emplace(PassStr, ResourceStr);
+			}
+			else
+			{
+				UE_LOG(LogDumpGPU, Warning, TEXT("Invalid unified filter format: %s. Expected Pass@Resource"), *RawFilter);
+			}
+		}
 	}
 
 	NewResourceDumpContext->bStream = GDumpGPUStream.GetValueOnGameThread() != 0;
@@ -2396,7 +2420,10 @@ void FRDGResourceDumpContext::Finish()
 			GLog->FlushThreadedLogs();
 			GLog->Flush();
 		}
-		FGenericCrashContext::DumpLog(DumpingDirectoryPath / FRDGResourceDumpContext::kBaseDir);
+		if (GDumpGPUEnableLogWrite.GetValueOnGameThread())
+		{
+			FGenericCrashContext::DumpLog(DumpingDirectoryPath / FRDGResourceDumpContext::kBaseDir);
+		}
 	}
 
 	// Update the dump status to OK and removes subscription to crashes
