Shader library: fix a rare deadlock when mounting chunks.

- Replace write lock on the library stack with a read lock when mounting chunks, since we are not modifying the stack itself (i.e. not adding/removing new named libs). When there are few worker threads, a write lock in that place causes an oversubscription deadlock.

- Calculate present chunks from the component data to avoid potential hazards as the previous structure might have been getting accidental protection from the overly broad lock.

#Synced-CL 48999368

diff --git a/Engine/Source/Runtime/RenderCore/Private/ShaderLibrary/ShaderCodeLibrary.cpp b/Engine/Source/Runtime/RenderCore/Private/ShaderLibrary/ShaderCodeLibrary.cpp
--- a/Engine/Source/Runtime/RenderCore/Private/ShaderLibrary/ShaderCodeLibrary.cpp
+++ b/Engine/Source/Runtime/RenderCore/Private/ShaderLibrary/ShaderCodeLibrary.cpp
@@ -151,11 +151,6 @@ namespace UE::ShaderLibrary::Private
 		/** Base directory for chunk 0 */
 		FString BaseDirectory;
 
-		/** Parts of the library corresponding to particular chunk Ids that we have found for this library.
-		 *  This is used so we don't try to open a library for the chunk more than once
-		 */
-		TSet<int32> PresentChunks;
-
 		/** Guards access to components*/
 		FRWLock ComponentsMutex;
 
@@ -174,22 +169,39 @@ namespace UE::ShaderLibrary::Private
 			return Components.Num();
 		}
 
+		/** 
+		 * Returns which chunk Ids are present in this library. Monolithic libraries will return empty set
+		 */
+		TSet<int32> GetPresentChunks()
+		{
+			FRWScopeLock ReadLock(ComponentsMutex, SLT_ReadOnly);
+			TSet<int32> PresentChunks;
+
+			for (TUniquePtr<FShaderLibraryInstance>& Component : Components)
+			{
+				int32 ChunkId = Component->GetChunkId();
+				if (LIKELY(ChunkId != INDEX_NONE))
+				{
+					PresentChunks.Add(ChunkId);
+				}
+			}
+
+			return MoveTemp(PresentChunks);
+		}
+
 		void OnPakFileMounted(const FMountedPakFileInfo& MountInfo, const FString& Directory)
 		{
-			if (!PresentChunks.Contains(MountInfo.ChunkId))
+			if (!GetPresentChunks().Contains(MountInfo.ChunkId))
 			{
 				FString ChunkLibraryName = GetShaderLibraryNameForChunk(LogicalName, MountInfo.ChunkId);
 
 				// Parts of shader library might be in the UFS (e.g. .metallibs), hence we need to look for them in the
 				// appropriate directory.
-				if (OpenShaderCode(Directory, ChunkLibraryName))
-				{
-					PresentChunks.Add(MountInfo.ChunkId);
-				}
+				OpenShaderCode(Directory, ChunkLibraryName, MountInfo.ChunkId);
 			}
 		}
 
-		bool OpenShaderCode(const FString& ShaderCodeDir, FString const& Library);
+		bool OpenShaderCode(const FString& ShaderCodeDir, FString const& Library, int32 ChunkId);
 		FShaderLibraryInstance* FindShaderLibraryForShaderMap(const FSHAHash& Hash, int32& OutShaderMapIndex);
 		FShaderLibraryInstance* FindShaderLibraryForShader(const FSHAHash& Hash, int32& OutShaderIndex);
 		uint32 GetShaderCount();
@@ -948,14 +960,14 @@ public:
 			bool bOpenedAsMonolithic = false;
 
 			// if we're able to open the library by name, it's not chunked
-			if (Library->OpenShaderCode(Directory, Name))
+			if (Library->OpenShaderCode(Directory, Name, INDEX_NONE))
 			{
 				bResult = true;
 				bOpenedAsMonolithic = true;
 			}
 			else if (!bMonolithicOnly) // attempt to open a chunked library
 			{
-				TSet<int32> PrevComponentSet = Library->PresentChunks;
+				TSet<int32> PrevComponentSet = Library->GetPresentChunks();
 				{
 					FScopeLock KnownPakFilesLocker(&FMountedPakFileInfo::KnownPakFilesAccessLock);
 					for (TSet<FMountedPakFileInfo>::TConstIterator Iter(FMountedPakFileInfo::KnownPakFiles); Iter; ++Iter)
@@ -963,7 +975,7 @@ public:
 						Library->OnPakFileMounted(*Iter, Directory);
 					}
 				}
-				NewComponentIDs = Library->PresentChunks.Difference(PrevComponentSet);
+				NewComponentIDs = Library->GetPresentChunks().Difference(PrevComponentSet);
 				bResult = !NewComponentIDs.IsEmpty();
 
 #if UE_SHADERLIB_SUPPORT_CHUNK_DISCOVERY
@@ -1005,7 +1017,7 @@ public:
 								}
 							}
 
-							NewComponentIDs = Library->PresentChunks.Difference(PrevComponentSet);
+							NewComponentIDs = Library->GetPresentChunks().Difference(PrevComponentSet);
 							bResult = !NewComponentIDs.IsEmpty();
 						}
 						else
@@ -1148,12 +1160,12 @@ public:
 			LLM_SCOPE(ELLMTag::Shaders);
 			TArray<TUniqueFunction<void()>> ShaderLibraryStateChanges;
 			{
-				FRWScopeLock WriteLock(NamedLibrariesMutex, SLT_Write);
+				FRWScopeLock ReadLock(NamedLibrariesMutex, SLT_ReadOnly);
 				for (TTuple<FString, TUniquePtr<UE::ShaderLibrary::Private::FNamedShaderLibrary>>& NamedLibraryPair : NamedLibrariesStack)
 				{
-					TSet<int32> PrevComponentSet = NamedLibraryPair.Value->PresentChunks;
+					TSet<int32> PrevComponentSet = NamedLibraryPair.Value->GetPresentChunks();
 					NamedLibraryPair.Value->OnPakFileMounted(MountInfo, NamedLibraryPair.Value->BaseDirectory);
-					for (int32 ComponentID : NamedLibraryPair.Value->PresentChunks.Difference(PrevComponentSet))
+					for (int32 ComponentID : NamedLibraryPair.Value->GetPresentChunks().Difference(PrevComponentSet))
 					{
 						// Defer these for outside of the lock, ShaderPipelineCache may want to inspect shader library.
 						ShaderLibraryStateChanges.Add([ShaderPlatform = NamedLibraryPair.Value->ShaderPlatform, LogicalName = NamedLibraryPair.Value->LogicalName, ComponentID]() {
@@ -2661,7 +2673,7 @@ void FShaderCodeLibrary::OpenPluginShaderLibrary(IPlugin& Plugin, bool bMonolith
 // FNamedShaderLibrary methods
 
 // At runtime, open shader code collection for specified shader platform. Returns true if new code was opened
-bool UE::ShaderLibrary::Private::FNamedShaderLibrary::OpenShaderCode(const FString& ShaderCodeDir, FString const& Library)
+bool UE::ShaderLibrary::Private::FNamedShaderLibrary::OpenShaderCode(const FString& ShaderCodeDir, FString const& Library, int32 ChunkId)
 {
 	LLM_SCOPE(ELLMTag::Shaders);
 	// check if any of the components has this content
@@ -2677,7 +2689,7 @@ bool UE::ShaderLibrary::Private::FNamedShaderLibrary::OpenShaderCode(const FStri
 		}
 	}
 
-	FShaderLibraryInstance* LibraryInstance = FShaderLibraryInstance::Create(ShaderPlatform, ShaderCodeDir, Library);
+	FShaderLibraryInstance* LibraryInstance = FShaderLibraryInstance::Create(ShaderPlatform, ShaderCodeDir, Library, ChunkId);
 	if (LibraryInstance == nullptr)
 	{
 		UE_LOG(LogShaderLibrary, Verbose, TEXT("Cooked Context: No Shared Shader Library for: %s and native library not supported."), *Library);
diff --git a/Engine/Source/Runtime/RenderCore/Private/ShaderLibrary/ShaderLibraryInstance.cpp b/Engine/Source/Runtime/RenderCore/Private/ShaderLibrary/ShaderLibraryInstance.cpp
--- a/Engine/Source/Runtime/RenderCore/Private/ShaderLibrary/ShaderLibraryInstance.cpp
+++ b/Engine/Source/Runtime/RenderCore/Private/ShaderLibrary/ShaderLibraryInstance.cpp
@@ -17,7 +17,7 @@
 #include "UObject/NameTypes.h"
 
 
-FShaderLibraryInstance* FShaderLibraryInstance::Create(EShaderPlatform InShaderPlatform, const FString& ShaderCodeDir, FString const& InLibraryName)
+FShaderLibraryInstance* FShaderLibraryInstance::Create(EShaderPlatform InShaderPlatform, const FString& ShaderCodeDir, FString const& InLibraryName, int32 InChunkId)
 {
 	using namespace UE::ShaderLibrary::Private;
 
@@ -79,6 +79,7 @@ FShaderLibraryInstance* FShaderLibraryInstance::Create(EShaderPlatform InShaderP
 	FShaderLibraryInstance* Instance = new FShaderLibraryInstance();
 	Instance->Library = Library;
 	Instance->ShaderCodeDirectory = ShaderCodeDirectory;
+	Instance->ChunkId = InChunkId;
 
 	const int32 NumResources = Library->GetNumShaderMaps();
 	Instance->Resources.AddZeroed(NumResources);
diff --git a/Engine/Source/Runtime/RenderCore/Private/ShaderLibrary/ShaderLibraryInstance.h b/Engine/Source/Runtime/RenderCore/Private/ShaderLibrary/ShaderLibraryInstance.h
--- a/Engine/Source/Runtime/RenderCore/Private/ShaderLibrary/ShaderLibraryInstance.h
+++ b/Engine/Source/Runtime/RenderCore/Private/ShaderLibrary/ShaderLibraryInstance.h
@@ -19,7 +19,7 @@
 class FShaderLibraryInstance
 {
 public:
-	static FShaderLibraryInstance* Create(EShaderPlatform InShaderPlatform, const FString& ShaderCodeDir, FString const& InLibraryName);
+	static FShaderLibraryInstance* Create(EShaderPlatform InShaderPlatform, const FString& ShaderCodeDir, FString const& InLibraryName, int32 InChunkId);
 
 	~FShaderLibraryInstance();
 
@@ -57,6 +57,9 @@ public:
 
 	bool HasContentFrom(const FString& ShaderCodeDir, FString const& InLibraryName);
 
+	/** Returns chunk id for this instance, if it is a part of a chunked named library. INDEX_NONE is returned for a monolithic lib. */
+	int32 GetChunkId() const { return ChunkId; }
+
 public:
 	FRHIShaderLibraryRef Library;
 
@@ -93,4 +96,7 @@ private:
 
 	/** Folder the library was created from (doesn't matter for IoStore-based libraries) */
 	FString ShaderCodeDirectory;
+
+	/** Chunk ID of this library, if any (INDEX_NONE if it is monolithic) */
+	int32 ChunkId;
 };
