[ShaderLibrary] Incremental cooking: Store cookcache in separate files not read by staging.

Changes from Matt.Peter's initial shelf:
- Removed `ESaveToDiskTarget::Count`; Was only used in one loop with forth-and-back type casting.
- Retain const members in `FEditorShaderCodeArchive`; Only `Prune()` needed the move-assignment operator, which is now implemented as private function `AssignFrom()`.
- Removed `ECookShaderLibrarySource::Count`; Introduced in CL 48964500.
- Removed `GetSaveToDiskTargetFromFilenameToken()`.
- Added documentation to `InitializeReferenceTracking()`.
- Added `FEditorShaderCodeArchive::ForEachAssetReferencingShaderType()` to suppress assertion in `DumpShaderTypeStats()` for shader types that are *not* referenced by any assets with their staging flag set.
- Made regex group for "_CookCache" token a non-capture group and added *.ucacheshaderbytecode.

#rnx
#Synced-CL 49105494

diff --git a/Engine/Source/Editor/UnrealEd/Classes/CookOnTheSide/CookOnTheFlyServer.h b/Engine/Source/Editor/UnrealEd/Classes/CookOnTheSide/CookOnTheFlyServer.h
--- a/Engine/Source/Editor/UnrealEd/Classes/CookOnTheSide/CookOnTheFlyServer.h
+++ b/Engine/Source/Editor/UnrealEd/Classes/CookOnTheSide/CookOnTheFlyServer.h
@@ -1103,7 +1103,7 @@ private:
 	 * Finishes collection of data that should be in the named code library (e.g. load data from previous
 	 * incremental cook)
 	 */
-	void FinishPopulateShaderLibrary(const ITargetPlatform* TargetPlatform, FString const& Name);
+	void FinishPopulateShaderLibrary(const ITargetPlatform* TargetPlatform, FString const& Name, bool bInitializeReferenceTracking);
 	/**
 	 * Prune the named ShaderLibrary of any data no longer referenced by assets in the current cook or
 	 * IncrementalCook oplog.
diff --git a/Engine/Source/Editor/UnrealEd/Private/CookOnTheFlyServer.cpp b/Engine/Source/Editor/UnrealEd/Private/CookOnTheFlyServer.cpp
--- a/Engine/Source/Editor/UnrealEd/Private/CookOnTheFlyServer.cpp
+++ b/Engine/Source/Editor/UnrealEd/Private/CookOnTheFlyServer.cpp
@@ -9412,6 +9412,15 @@ void UCookOnTheFlyServer::CookByTheBookFinishedInternal()
 	// call UpdateOplogPackages to report which assets are still referenced, before pruning any no-longer-referenced
 	// shaders from any platforms. We do it in this order because platforms can share shaderformats, so we need to
 	// know if any platform is referencing shaderformat X before we can prune it.
+	if (bCacheShaderLibraries)
+	{
+		for (const ITargetPlatform* TargetPlatform : PlatformManager->GetSessionPlatforms())
+		{
+			UE_SCOPED_HIERARCHICAL_COOKTIMER(FinishPopulateShaderLibrary);
+			FinishPopulateShaderLibrary(TargetPlatform, LibraryName, true /* bInitializeReferenceTracking */);
+		}
+	}
+
 	int32 PlatformIndex = 0;
 	for (const ITargetPlatform* TargetPlatform : PlatformManager->GetSessionPlatforms())
 	{
@@ -9425,12 +9434,6 @@ void UCookOnTheFlyServer::CookByTheBookFinishedInternal()
 		FPlatformData* PlatformData = PlatformManager->GetPlatformData(TargetPlatform);
 		FAssetRegistryGenerator& Generator = *PlatformData->RegistryGenerator;
 
-		if (bCacheShaderLibraries)
-		{
-			UE_SCOPED_HIERARCHICAL_COOKTIMER(FinishPopulateShaderLibrary);
-			FinishPopulateShaderLibrary(TargetPlatform, LibraryName);
-		}
-
 		TArray<FPackageData*>& CookedPackageDatas = EndCookPlatformData.CookedPackageDatas;
 		TArray<FPackageData*>& IgnorePackageDatas = EndCookPlatformData.IgnorePackageDatas;
 
diff --git a/Engine/Source/Editor/UnrealEd/Private/Cooker/ShaderLibraryCooking.cpp b/Engine/Source/Editor/UnrealEd/Private/Cooker/ShaderLibraryCooking.cpp
--- a/Engine/Source/Editor/UnrealEd/Private/Cooker/ShaderLibraryCooking.cpp
+++ b/Engine/Source/Editor/UnrealEd/Private/Cooker/ShaderLibraryCooking.cpp
@@ -568,7 +568,7 @@ void UCookOnTheFlyServer::SaveAndCloseGlobalShaderLibrary()
 		// Save shader code map - cleaning directories is deliberately a separate loop here as we open the cache once per shader platform and we don't assume that they can't be shared across target platforms.
 		for (const ITargetPlatform* TargetPlatform : PlatformManager->GetSessionPlatforms())
 		{
-			FinishPopulateShaderLibrary(TargetPlatform, GlobalShaderLibName);
+			FinishPopulateShaderLibrary(TargetPlatform, GlobalShaderLibName, false /* bInitializeReferenceTracking */);
 			SaveShaderLibrary(TargetPlatform, GlobalShaderLibName);
 		}
 
@@ -587,13 +587,13 @@ void UCookOnTheFlyServer::GetShaderLibraryPaths(const ITargetPlatform* TargetPla
 	OutMetaDataPath = ConvertToFullSandboxPath(*RootMetaDataPath, true, *TargetPlatform->PlatformName());
 }
 
-void UCookOnTheFlyServer::FinishPopulateShaderLibrary(const ITargetPlatform* TargetPlatform, const FString& Name)
+void UCookOnTheFlyServer::FinishPopulateShaderLibrary(const ITargetPlatform* TargetPlatform, const FString& Name, bool bInitializeReferenceTracking)
 {
 	FString ShaderCodeDir;
 	FString MetaDataPath;
 	GetShaderLibraryPaths(TargetPlatform, ShaderCodeDir, MetaDataPath);
 
-	FShaderLibraryCooker::FinishPopulateShaderLibrary(TargetPlatform, Name, ShaderCodeDir, MetaDataPath);
+	FShaderLibraryCooker::FinishPopulateShaderLibrary(TargetPlatform, Name, ShaderCodeDir, MetaDataPath, bInitializeReferenceTracking);
 }
 
 void UCookOnTheFlyServer::PruneShaderLibrary(const ITargetPlatform* TargetPlatform, const FString& Name)
diff --git a/Engine/Source/Runtime/RenderCore/Private/ShaderLibrary/EditorShaderCodeArchive.cpp b/Engine/Source/Runtime/RenderCore/Private/ShaderLibrary/EditorShaderCodeArchive.cpp
--- a/Engine/Source/Runtime/RenderCore/Private/ShaderLibrary/EditorShaderCodeArchive.cpp
+++ b/Engine/Source/Runtime/RenderCore/Private/ShaderLibrary/EditorShaderCodeArchive.cpp
@@ -400,29 +400,72 @@ bool FEditorShaderCodeArchive::AppendFromArchive(
 	return bOk;
 }
 
-FEditorShaderCodeArchive* FEditorShaderCodeArchive::CreateChunk(int ChunkId, const TSet<FName>& PackagesInChunk)
+TUniquePtr<FEditorShaderCodeArchive> FEditorShaderCodeArchive::CreateArchiveFromFilteredAssets(const FString& NewLibraryName, FSerializedShaderArchive::FAssetFilterFunctionRef ShouldKeepAsset)
 {
-	checkf(!bHasCopiedAndCleared, TEXT("It is not valid to call CreateChunk on an FEditorShaderCodeArchive that has sent its shaders to another process by calling CopyAndClear."));
+	checkf(!bHasCopiedAndCleared, TEXT("It is not valid to call CreateArchiveFromFilteredAssets on an FEditorShaderCodeArchive that has sent its shaders to another process by calling CopyAndClear."));
 
-	FEditorShaderCodeArchive* NewChunk = new FEditorShaderCodeArchive(FormatName, bNeedsDeterministicOrder);
-	NewChunk->OpenLibrary(UE::ShaderLibrary::Private::GetShaderLibraryNameForChunk(LibraryName, ChunkId));
+	TUniquePtr<FEditorShaderCodeArchive> NewArchive(new FEditorShaderCodeArchive(FormatName, bNeedsDeterministicOrder));
+	NewArchive->OpenLibrary(NewLibraryName);
 
 	TArray<int32> ShaderCodeEntriesNeeded;	// this array is filled with the indices from the existing ShaderCode that will need to be taken
-	NewChunk->SerializedShaders.CreateAsChunkFrom(SerializedShaders, PackagesInChunk, ShaderCodeEntriesNeeded);
+	NewArchive->SerializedShaders.CreateFromFilteredAssets(SerializedShaders, ShaderCodeEntriesNeeded, ShouldKeepAsset);
+
 	// extra integrity check
-	checkf(ShaderCodeEntriesNeeded.Num() == NewChunk->SerializedShaders.ShaderHashes.Num(), TEXT("FSerializedShaderArchive for the new chunk did not create a valid shader code mapping"));
-	checkf(ShaderCodeEntriesNeeded.Num() == NewChunk->SerializedShaders.ShaderEntries.Num(), TEXT("FSerializedShaderArchive for the new chunk did not create a valid shader code mapping"));
+	checkf(ShaderCodeEntriesNeeded.Num() == NewArchive->SerializedShaders.ShaderHashes.Num(),
+		TEXT("FSerializedShaderArchive for the new chunk did not create a valid shader code mapping"));
+	checkf(ShaderCodeEntriesNeeded.Num() == NewArchive->SerializedShaders.ShaderEntries.Num(),
+		TEXT("FSerializedShaderArchive for the new chunk did not create a valid shader code mapping"));
 
 	// copy the shader code
-	NewChunk->ShaderCode.Empty();
+	NewArchive->ShaderCode.Empty();
 	for (int32 NewArchiveIdx = 0, NumIndices = ShaderCodeEntriesNeeded.Num(); NewArchiveIdx < NumIndices; ++NewArchiveIdx)
 	{
 		FSharedBuffer& SourceShaderCodeEntry = ShaderCode[ShaderCodeEntriesNeeded[NewArchiveIdx]];
 		check(SourceShaderCodeEntry.GetSize() > 0);
-		NewChunk->ShaderCode.Add(SourceShaderCodeEntry);
+		NewArchive->ShaderCode.Add(SourceShaderCodeEntry);
 	}
 
-	return NewChunk;
+	return NewArchive;
+}
+
+TUniquePtr<FEditorShaderCodeArchive> FEditorShaderCodeArchive::CreateArchiveFromAssetsReferencedByStaging()
+{
+	return CreateArchiveFromFilteredAssets(
+		LibraryName,
+		[](FName AssetName, const FSerializedShaderArchive::FShaderMapAssetAssociations::FAssociatedAssetData& AssetData)
+		{
+			return AssetData.bReferencedByStaging;
+		});
+}
+
+TUniquePtr<FEditorShaderCodeArchive> FEditorShaderCodeArchive::CreateArchiveFromAssetsReferencedByOplog()
+{
+	return CreateArchiveFromFilteredAssets(
+		LibraryName,
+		[](FName AssetName, const FSerializedShaderArchive::FShaderMapAssetAssociations::FAssociatedAssetData& AssetData)
+		{
+			return AssetData.bReferencedByOplog;
+		});
+}
+
+TUniquePtr<FEditorShaderCodeArchive> FEditorShaderCodeArchive::CreateArchiveFromAssetsOnlyReferencedByOplog()
+{
+	return CreateArchiveFromFilteredAssets(
+		LibraryName,
+		[](FName AssetName, const FSerializedShaderArchive::FShaderMapAssetAssociations::FAssociatedAssetData& AssetData)
+		{
+			return AssetData.bReferencedByOplog && !AssetData.bReferencedByStaging;
+		});
+}
+
+TUniquePtr<FEditorShaderCodeArchive> FEditorShaderCodeArchive::CreateChunk(int ChunkId, const TSet<FName>& PackagesInChunk)
+{
+	return CreateArchiveFromFilteredAssets(
+		UE::ShaderLibrary::Private::GetShaderLibraryNameForChunk(LibraryName, ChunkId),
+		[&PackagesInChunk](FName AssetName, const FSerializedShaderArchive::FShaderMapAssetAssociations::FAssociatedAssetData& AssetData)
+		{
+			return PackagesInChunk.Contains(AssetName);
+		});
 }
 
 int32 FEditorShaderCodeArchive::AddShaderCode(int32 OtherShaderMapIndex, const FEditorShaderCodeArchive& OtherArchive,
@@ -602,90 +645,99 @@ FEditorShaderCodeArchive::EShaderCodeReadStatus FEditorShaderCodeArchive::ReadSh
 FEditorShaderCodeArchive::EShaderCodeReadStatus FEditorShaderCodeArchive::AddShaderCodeLibraryByName(
 	const FString& BaseDir, const FString& InLibraryName, FShaderCodeStats& CodeStats, ECookShaderLibrarySource CookSource)
 {
-	FString CodeArchiveFileName = GetCodeArchiveFilename(BaseDir, InLibraryName, FormatName);
-	if (FArchive* PrevCookedAr = CreateShaderFileReader(*CodeArchiveFileName))
+	using namespace UE::ShaderLibrary;
+
+	bool bFoundAny = false;
+	for (ESaveToDiskTarget Target : { ESaveToDiskTarget::Staging, ESaveToDiskTarget::CookCache })
 	{
-		EShaderCodeReadStatus Status;
+		const FString CodeArchiveFileName = GetCodeArchiveFilename(BaseDir, InLibraryName, FormatName, Target);
+		TUniquePtr<FArchive> PrevCookedAr(CreateShaderFileReader(*CodeArchiveFileName));
+		if (!PrevCookedAr)
+		{
+			continue;
+		}
+		bFoundAny = true;
+
 		uint32 Version = 0;
 		*PrevCookedAr << Version;
 
-		if (Version == GShaderCodeArchiveVersion)
+		if (Version != GShaderCodeArchiveVersion)
 		{
-			FSerializedShaderArchive PrevCookedShaders;
+			return EShaderCodeReadStatus::VersionMismatch;
+		}
 
-			*PrevCookedAr << PrevCookedShaders;
+		FSerializedShaderArchive PrevCookedShaders;
 
-			// check if it also contains the asset info file
-			FString AssetInfoFileName = GetShaderAssetInfoFilename(BaseDir, InLibraryName, FormatName);
-			TUniquePtr<FArchive> ShaderAssetInfoReader(CreateShaderFileReader(*AssetInfoFileName));
-			if (PrevCookedShaders.LoadAssetInfo(ShaderAssetInfoReader.Get(), CookSource))
-			{
-				UE_LOG(LogShaderLibrary, Display, TEXT("Loaded asset info %s for the shader library %s: %d entries"),
-					*AssetInfoFileName,
-					*CodeArchiveFileName,
-					PrevCookedShaders.GetShaderMapAssetAssociations().ViewShaderMaps().Num()
-				);
-			}
-			else
-			{
-				FString ErrorMessage = FString::Printf(
-					TEXT("Could not find or load asset info %s for the shader library %s"),
-					*AssetInfoFileName, *CodeArchiveFileName);
-				if (CookSource == ECookShaderLibrarySource::PreviousIncremental)
-				{
-					// In incremental cooks, we require the AssetInfo so we know which shadermaps to keep, lack of it
-					// should report a cook error.
-					UE_LOG(LogShaderLibrary, Error, TEXT("%s"), *ErrorMessage);
-				}
-				else
-				{
-					UE_LOG(LogShaderLibrary, Warning, TEXT("%s"), *ErrorMessage);
-				}
-			}
+		*PrevCookedAr << PrevCookedShaders;
 
-			// check if it also contains the shader type info file
-			FString TypeInfoFileName = GetShaderTypeInfoFilename(BaseDir, InLibraryName, FormatName);
-			TUniquePtr<FArchive> ShaderTypeInfoReader(CreateShaderFileReader(*TypeInfoFileName));
-			if (ShaderTypeInfoReader)
+		// check if it also contains the asset info file
+		FString AssetInfoFileName = GetShaderAssetInfoFilename(BaseDir, InLibraryName, FormatName, Target);
+		TUniquePtr<FArchive> ShaderAssetInfoReader(CreateShaderFileReader(*AssetInfoFileName));
+		if (PrevCookedShaders.LoadAssetInfo(ShaderAssetInfoReader.Get(), CookSource))
+		{
+			UE_LOG(LogShaderLibrary, Display, TEXT("Loaded asset info %s for the shader library %s: %d entries"),
+				*AssetInfoFileName,
+				*CodeArchiveFileName,
+				PrevCookedShaders.GetShaderMapAssetAssociations().ViewShaderMaps().Num()
+			);
+		}
+		else
+		{
+			FString ErrorMessage = FString::Printf(
+				TEXT("Could not find or load asset info %s for the shader library %s"),
+				*AssetInfoFileName, *CodeArchiveFileName);
+			if (CookSource == ECookShaderLibrarySource::PreviousIncremental)
 			{
-				*ShaderTypeInfoReader << PrevCookedShaders.ShaderTypes;
-				check(PrevCookedShaders.ShaderTypes.Num() == PrevCookedShaders.ShaderEntries.Num());
-				UE_LOG(LogShaderLibrary, Display, TEXT("Loaded shader type info %s for the shader library %s: %d entries"),
-					*TypeInfoFileName,
-					*CodeArchiveFileName,
-					PrevCookedShaders.ShaderTypes.Num()
-				);
+				// In incremental cooks, we require the AssetInfo so we know which shadermaps to keep, lack of it
+				// should report a cook error.
+				UE_LOG(LogShaderLibrary, Error, TEXT("%s"), *ErrorMessage);
 			}
 			else
 			{
-				UE_LOG(LogShaderLibrary, Warning, TEXT("Could not load shader type info %s for the shader library %s"),
-					*TypeInfoFileName,
-					*CodeArchiveFileName
-				);
+				UE_LOG(LogShaderLibrary, Warning, TEXT("%s"), *ErrorMessage);
 			}
+		}
 
-			// Setup current shader library information and accumulate the reads we will need to do. Then read them.
-			TArray<FShaderCodeReadInfo> ReadInfo;
-			for (int32 PrevShaderMapIndex = 0; PrevShaderMapIndex < PrevCookedShaders.ShaderMapEntries.Num(); ++PrevShaderMapIndex)
-			{
-				PrepareAddShaderCode(CodeStats, PrevShaderMapIndex, PrevCookedShaders, ReadInfo, CookSource);
-			}
-			Status = ReadShaderCode(*PrevCookedAr, ReadInfo);
+		// check if it also contains the shader type info file
+		FString TypeInfoFileName = GetShaderTypeInfoFilename(BaseDir, InLibraryName, FormatName, Target);
+		TUniquePtr<FArchive> ShaderTypeInfoReader(CreateShaderFileReader(*TypeInfoFileName));
+		if (ShaderTypeInfoReader)
+		{
+			*ShaderTypeInfoReader << PrevCookedShaders.ShaderTypes;
+			check(PrevCookedShaders.ShaderTypes.Num() == PrevCookedShaders.ShaderEntries.Num());
+			UE_LOG(LogShaderLibrary, Display, TEXT("Loaded shader type info %s for the shader library %s: %d entries"),
+				*TypeInfoFileName,
+				*CodeArchiveFileName,
+				PrevCookedShaders.ShaderTypes.Num()
+			);
 		}
 		else
 		{
-			Status = EShaderCodeReadStatus::VersionMismatch;
+			UE_LOG(LogShaderLibrary, Warning, TEXT("Could not load shader type info %s for the shader library %s"),
+				*TypeInfoFileName,
+				*CodeArchiveFileName
+			);
 		}
 
-		PrevCookedAr->Close();
-		delete PrevCookedAr;
-		return Status;
+		// Setup current shader library information and accumulate the reads we will need to do. Then read them.
+		TArray<FShaderCodeReadInfo> ReadInfo;
+		for (int32 PrevShaderMapIndex = 0; PrevShaderMapIndex < PrevCookedShaders.ShaderMapEntries.Num(); ++PrevShaderMapIndex)
+		{
+			PrepareAddShaderCode(CodeStats, PrevShaderMapIndex, PrevCookedShaders, ReadInfo, CookSource);
+		}
+		EShaderCodeReadStatus Status = ReadShaderCode(*PrevCookedAr, ReadInfo);
+		if (Status != EShaderCodeReadStatus::OK)
+		{
+			return Status;
+		}
 	}
-	return EShaderCodeReadStatus::FileNotFound;
+	return bFoundAny ? EShaderCodeReadStatus::OK : EShaderCodeReadStatus::FileNotFound;
 }
 
 bool FEditorShaderCodeArchive::LoadExistingShaderCodeLibrary(FString const& MetaDataDir, ECookShaderLibrarySource CookSource)
 {
+	using namespace UE::ShaderLibrary;
+
 	FString BaseDir = MetaDataDir / TEXT("ShaderLibrarySource");
 	FShaderCodeStats CodeStats;
 	EShaderCodeReadStatus Status = AddShaderCodeLibraryByName(BaseDir, LibraryName, CodeStats, CookSource);
@@ -693,17 +745,20 @@ bool FEditorShaderCodeArchive::LoadExistingShaderCodeLibrary(FString const& Meta
 	{
 	case EShaderCodeReadStatus::FileNotFound:
 	{
-		UE_LOG(LogShaderLibrary, Error, TEXT("Failed to load shader code library from %s; file not found"), *GetCodeArchiveFilename(BaseDir, LibraryName, FormatName));
+		UE_LOG(LogShaderLibrary, Error, TEXT("Failed to load shader code library from %s; file not found"),
+			*GetCodeArchiveFilename(BaseDir, LibraryName, FormatName, ESaveToDiskTarget::Staging));
 		return false;
 	}
 	case EShaderCodeReadStatus::MalformedArchive:
 	{
-		UE_LOG(LogShaderLibrary, Error, TEXT("Failed to deserialized code from library %s"), *GetCodeArchiveFilename(BaseDir, LibraryName, FormatName));
+		UE_LOG(LogShaderLibrary, Error, TEXT("Failed to deserialized code from library %s"),
+			*GetCodeArchiveFilename(BaseDir, LibraryName, FormatName, ESaveToDiskTarget::Staging));
 		return false;
 	}
 	case EShaderCodeReadStatus::VersionMismatch:
 	{
-		UE_LOG(LogShaderLibrary, Error, TEXT("Failed to deserialize shader code from %s because the archive format is incompatible with the current version %u"), *GetCodeArchiveFilename(BaseDir, LibraryName, FormatName), GShaderCodeArchiveVersion);
+		UE_LOG(LogShaderLibrary, Error, TEXT("Failed to deserialize shader code from %s because the archive format is incompatible with the current version %u"),
+			*GetCodeArchiveFilename(BaseDir, LibraryName, FormatName, ESaveToDiskTarget::Staging), GShaderCodeArchiveVersion);
 		return false;
 	}
 	case EShaderCodeReadStatus::OK:
@@ -725,6 +780,7 @@ void FEditorShaderCodeArchive::AddExistingShaderCodeLibrary(FString const& Outpu
 	TArray<FString> ShaderFiles;
 	ShaderFindFiles(ShaderFiles, *ShaderIntermediateLocation, *ShaderExtension);
 	ShaderFindFiles(ShaderFiles, *OutputDir, *ShaderExtension);
+	ShaderFindFiles(ShaderFiles, *OutputDir, *CookCacheShaderExtension);
 
 	for (const FString& ShaderFileName : ShaderFiles)
 	{
@@ -754,6 +810,16 @@ void FEditorShaderCodeArchive::AddExistingShaderCodeLibrary(FString const& Outpu
 	}
 }
 
+void FEditorShaderCodeArchive::InitializeReferenceTracking()
+{
+	for (TPair<FName, FSerializedShaderArchive::FShaderMapAssetAssociations::FAssociatedAssetData>& Pair
+		: SerializedShaders.GetShaderMapAssetAssociations().ViewAssets())
+	{
+		Pair.Value.bReferencedByOplog = false;
+		Pair.Value.bReferencedByStaging = false;
+	}
+}
+
 void FEditorShaderCodeArchive::UpdateOplogPackages(UE::Cook::Artifact::FUpdateOplogPackagesContext& Context)
 {
 	using FShaderMapAssetAssociations = FSerializedShaderArchive::FShaderMapAssetAssociations;
@@ -783,37 +849,28 @@ void FEditorShaderCodeArchive::Prune()
 		return;
 	}
 
-	FShaderMapAssetAssociations& AssetMap = SerializedShaders.GetShaderMapAssetAssociations();
-	TArray<FName> RemoveAssets;
-	for (const TPair<FName, FShaderMapAssetAssociations::FAssociatedAssetData>& AssetPair : AssetMap.ViewAssets())
-	{
-		FName PackageName = AssetPair.Key;
-		if (!AssetPair.Value.bReferencedByOplog)
-		{
-			RemoveAssets.Add(PackageName);
-		}
-	}
-
-	TSet<FSHAHash> RemoveShaderMaps;
-	for (FSHAHash& ShaderMapHash : SerializedShaders.ShaderMapHashes)
-	{
-		if (!AssetMap.FindShaderMap(ShaderMapHash))
-		{
-			RemoveShaderMaps.Add(ShaderMapHash);
-		}
-	}
-
-	// TODO: Implement RemoveShaderMaps.
-	// TODO: In the Library write functions, move the ShaderMaps that do not have any bReferencedByStaging Assets into
-	// a separate pair of files: IncrementalCookArchive.uincrementalcookshaderbytecode, IncrementalCookAssetInfo.json that
-	// are not used by staging. Note that we cannot use .ushaderbytecode as the extension, because staging searches
-	// for *.ushaderbytecode.
-	// TODO: In FinishPopulate, read the IncrementalCookArchive files in addition to the .ushaderbytecode files.
+	TUniquePtr<FEditorShaderCodeArchive> PrunedArchive = this->CreateArchiveFromAssetsReferencedByOplog();
+	AssignFrom(*PrunedArchive);
+}
 
-	// FSerializedShaderArchive::RemoveShaderMaps(RemoveShaderMaps);
+void FEditorShaderCodeArchive::AssignFrom(FEditorShaderCodeArchive& InArchive)
+{
+	checkf(this->FormatName == InArchive.FormatName,
+		TEXT("Cannot assign shader code archive due to `FormatName` mismatch: Expected '%s', but got '%s'"), *this->FormatName.ToString(), *InArchive.FormatName.ToString());
+	checkf(this->Format == InArchive.Format,
+		TEXT("Cannot assign shader code archive due to `Format` mismatch: Expected %p, but got %p"), this->Format, InArchive.Format);
+	checkf(this->bNeedsDeterministicOrder == InArchive.bNeedsDeterministicOrder,
+		TEXT("Cannot assign shader code archive due to `bNeedsDeterministicOrder` mismatch: Expected %s, but got %s"), this->bNeedsDeterministicOrder ? TEXT("true") : TEXT("false"), InArchive.bNeedsDeterministicOrder ? TEXT("true") : TEXT("false"));
+
+	this->LibraryName = MoveTemp(InArchive.LibraryName);
+	this->SerializedShaders = MoveTemp(InArchive.SerializedShaders);
+	this->ShaderCode = MoveTemp(InArchive.ShaderCode);
+	this->ShaderMapsToCopy = MoveTemp(InArchive.ShaderMapsToCopy);
+	this->bHasCopiedAndCleared = MoveTemp(InArchive.bHasCopiedAndCleared);
 }
 
-bool FEditorShaderCodeArchive::SaveToDisk(const FString& OutputDir, const FString& MetaOutputDir, bool bSaveOnlyAssetInfo, TArray<FString>* OutputFilenames)
+bool FEditorShaderCodeArchive::SaveToDisk(const FString& OutputDir, const FString& MetaOutputDir,
+	ESaveToDiskTarget Target, TArray<FString>* OutputFilenames)
 {
 	check(LibraryName.Len() > 0);
 	checkf(!bHasCopiedAndCleared, TEXT("It is not valid to call SaveToDisk on an FEditorShaderCodeArchive that has sent its shaders to another process by calling CopyAndClear."));
@@ -842,11 +899,12 @@ bool FEditorShaderCodeArchive::SaveToDisk(const FString& OutputDir, const FStrin
 	{
 		// Write to a intermediate file
 		const FString DebugInfoPath = GetShaderCodeLibraryDebugInfoPath();
-		FString IntermediateFormatPath = GetShaderCodeFilename(DebugInfoPath / FormatName.ToString(), LibraryName, FormatName);
-		FString AssetInfoIntermediatePath = GetShaderAssetInfoFilename(DebugInfoPath / FormatName.ToString(), LibraryName, FormatName);
-		FString TypeInfoIntermediatePath = GetShaderTypeInfoFilename(DebugInfoPath / FormatName.ToString(), LibraryName, FormatName);
+		FString IntermediateFormatPath = GetShaderCodeFilename(DebugInfoPath / FormatName.ToString(), LibraryName, FormatName, Target);
+		FString AssetInfoIntermediatePath = GetShaderAssetInfoFilename(DebugInfoPath / FormatName.ToString(), LibraryName, FormatName, Target);
+		FString TypeInfoIntermediatePath = GetShaderTypeInfoFilename(DebugInfoPath / FormatName.ToString(), LibraryName, FormatName, Target);
 
-		// save the actual shader code code
+		// save the ShaderMaps, Shaders, and the Shaders' byte code
+		constexpr bool bSaveOnlyAssetInfo = false; // A currently unused setting that we might want to bring back.
 		if (!bSaveOnlyAssetInfo)
 		{
 			FArchive* FileWriter = IFileManager::Get().CreateFileWriter(*IntermediateFormatPath, FILEWRITE_NoFail);
@@ -871,13 +929,13 @@ bool FEditorShaderCodeArchive::SaveToDisk(const FString& OutputDir, const FStrin
 				delete FileWriter;
 
 				// Copy to output location - support for iterative native library cooking
-				if (!CopyFile(GetCodeArchiveFilename(OutputDir, LibraryName, FormatName), IntermediateFormatPath, OutputFilenames))
-				{
-					bSuccess = false;
-				}
-				else if (MetaOutputDir.Len())
+				bool bOutputDirSuccess = CopyFile(GetCodeArchiveFilename(OutputDir, LibraryName, FormatName, Target),
+					IntermediateFormatPath, OutputFilenames);
+				bSuccess = bSuccess && bOutputDirSuccess;
+				if (bOutputDirSuccess && MetaOutputDir.Len() && Target == ESaveToDiskTarget::Staging)
 				{
-					if (!CopyFile(GetCodeArchiveFilename(MetaOutputDir / TEXT("../ShaderLibrarySource"), LibraryName, FormatName), IntermediateFormatPath, nullptr))
+					if (!CopyFile(GetCodeArchiveFilename(MetaOutputDir / TEXT("../ShaderLibrarySource"), LibraryName, FormatName, Target),
+						IntermediateFormatPath, nullptr))
 					{
 						bSuccess = false;
 					}
@@ -892,13 +950,13 @@ bool FEditorShaderCodeArchive::SaveToDisk(const FString& OutputDir, const FStrin
 				delete ShaderTypeFileWriter;
 
 				// Copy to output location - support for iterative native library cooking
-				if (!CopyFile(GetShaderTypeInfoFilename(OutputDir, LibraryName, FormatName), TypeInfoIntermediatePath, OutputFilenames))
-				{
-					bSuccess = false;
-				}
-				else if (MetaOutputDir.Len())
+				bool bOutputDirSuccess = CopyFile(GetShaderTypeInfoFilename(OutputDir, LibraryName, FormatName, Target),
+					TypeInfoIntermediatePath, OutputFilenames);
+				bSuccess = bSuccess && bOutputDirSuccess;
+				if (bOutputDirSuccess && MetaOutputDir.Len() && Target == ESaveToDiskTarget::Staging)
 				{
-					if (!CopyFile(GetShaderTypeInfoFilename(MetaOutputDir / TEXT("../ShaderLibrarySource"), LibraryName, FormatName), TypeInfoIntermediatePath, nullptr))
+					if (!CopyFile(GetShaderTypeInfoFilename(MetaOutputDir / TEXT("../ShaderLibrarySource"), LibraryName, FormatName, Target),
+						TypeInfoIntermediatePath, nullptr))
 					{
 						bSuccess = false;
 					}
@@ -915,14 +973,14 @@ bool FEditorShaderCodeArchive::SaveToDisk(const FString& OutputDir, const FStrin
 				AssetInfoWriter->Close();
 				delete AssetInfoWriter;
 
-				if (!CopyFile(GetShaderAssetInfoFilename(OutputDir, LibraryName, FormatName), AssetInfoIntermediatePath, nullptr))
-				{
-					bSuccess = false;
-				}
-				else if (MetaOutputDir.Len())
+				bool bOutputDirSuccess = CopyFile(GetShaderAssetInfoFilename(OutputDir, LibraryName, FormatName, Target),
+					AssetInfoIntermediatePath, nullptr);
+				bSuccess = bSuccess && bOutputDirSuccess;
+				if (bOutputDirSuccess && MetaOutputDir.Len() && Target == ESaveToDiskTarget::Staging)
 				{
 					// copy asset info as well for debugging
-					if (!CopyFile(GetShaderAssetInfoFilename(MetaOutputDir / TEXT("../ShaderLibrarySource"), LibraryName, FormatName), AssetInfoIntermediatePath, nullptr))
+					if (!CopyFile(GetShaderAssetInfoFilename(MetaOutputDir / TEXT("../ShaderLibrarySource"), LibraryName, FormatName, Target),
+						AssetInfoIntermediatePath, nullptr))
 					{
 						bSuccess = false;
 					}
@@ -956,7 +1014,8 @@ bool FEditorShaderCodeArchive::PackageNativeShaderLibrary(const FString& ShaderC
 		// Delete Shader code library / pipelines as we now have native versions
 		// [RCL] 2020-12-02 FIXME: check if this doesn't ruin iterative cooking during Launch On
 		{
-			FString OutputFilePath = GetCodeArchiveFilename(ShaderCodeDir, LibraryName, FormatName);
+			FString OutputFilePath = GetCodeArchiveFilename(ShaderCodeDir, LibraryName, FormatName,
+				ESaveToDiskTarget::Staging);
 			IFileManager::Get().Delete(*OutputFilePath);
 		}
 		{
@@ -1158,4 +1217,62 @@ void FEditorShaderCodeArchive::DumpStatsAndDebugInfo()
 	UE_LOG(LogShaderLibrary, Display, TEXT("================="));
 }
 
+void FEditorShaderCodeArchive::ForEachAssetReferencingShaderType(uint64 InShaderTypeHash, FSerializedShaderArchive::FAssetFilterFunctionRef AssetReferenceCallback) const
+{
+	check(SerializedShaders.ShaderTypes.Num() == SerializedShaders.ShaderHashes.Num());
+	const int32 NumShaders = SerializedShaders.ShaderTypes.Num();
+
+	// Create temporary array to quickly map from shader index to all its shadermap hashes
+	TArray<TArray<FSHAHash>> ShaderIndexToShaderMapHashes;
+	ShaderIndexToShaderMapHashes.SetNum(NumShaders);
+
+	for (int32 ShaderMapIndex = 0; ShaderMapIndex < SerializedShaders.ShaderMapHashes.Num(); ++ShaderMapIndex)
+	{
+		// Put current shadermap's hash into the temporary array for all shaders it belongs to
+		const FSHAHash& ShaderMapHash = SerializedShaders.ShaderMapHashes[ShaderMapIndex];
+		const FShaderMapEntry& ShaderMapEntry = SerializedShaders.ShaderMapEntries[ShaderMapIndex];
+		for (uint32 ShaderIter = 0; ShaderIter < ShaderMapEntry.NumShaders; ++ShaderIter)
+		{
+			const int32 ShaderIndex = SerializedShaders.ShaderIndices[ShaderMapEntry.ShaderIndicesOffset + ShaderIter];
+			ShaderIndexToShaderMapHashes[ShaderIndex].Add(ShaderMapHash);
+		}
+	}
+
+	// Find shader type hash in all of serialized shaders
+	const FSerializedShaderArchive::FShaderMapAssetAssociations& ShaderMapAssetAssoc = SerializedShaders.GetShaderMapAssetAssociations();
+	for (int32 ShaderIndex = 0; ShaderIndex < NumShaders; ++ShaderIndex)
+	{
+		if (!SerializedShaders.ShaderTypes[ShaderIndex].Data.Contains(InShaderTypeHash))
+		{
+			continue;
+		}
+
+		for (const FSHAHash& ShaderMapHash : ShaderIndexToShaderMapHashes[ShaderIndex])
+		{
+			const FShaderMapAssetPaths* ShaderMapAssetNames = ShaderMapAssetAssoc.FindShaderMap(ShaderMapHash);
+			if (!ShaderMapAssetNames)
+			{
+				continue;
+			}
+
+			// Invoke input callback for each asset for this shadermap
+			for (const FName& AssetName : *ShaderMapAssetNames)
+			{
+				const FSerializedShaderArchive::FShaderMapAssetAssociations::FAssociatedAssetData* AssetData = ShaderMapAssetAssoc.FindAsset(AssetName);
+				if (!AssetData)
+				{
+					continue;
+				}
+
+				// Invoke the callback on the current asset data
+				if (!AssetReferenceCallback(AssetName, *AssetData))
+				{
+					// Stop iterating
+					return;
+				}
+			}
+		}
+	}
+}
+
 #endif // WITH_EDITOR
diff --git a/Engine/Source/Runtime/RenderCore/Private/ShaderLibrary/EditorShaderCodeArchive.h b/Engine/Source/Runtime/RenderCore/Private/ShaderLibrary/EditorShaderCodeArchive.h
--- a/Engine/Source/Runtime/RenderCore/Private/ShaderLibrary/EditorShaderCodeArchive.h
+++ b/Engine/Source/Runtime/RenderCore/Private/ShaderLibrary/EditorShaderCodeArchive.h
@@ -16,6 +16,8 @@
 #include "Misc/SecureHash.h"
 #include "Serialization/Archive.h"
 #include "ShaderCodeArchive.h"
+#include "ShaderLibrary/ShaderCodeLibraryUtilities.h" // ESaveToDiskTarget. Maybe put it in a separate header to avoid the need to include ShaderCodeLibraryUtilities.h?
+#include "Templates/Function.h"
 #include "UObject/NameTypes.h"
 
 
@@ -34,6 +36,7 @@ struct FShaderCodeStats
 	};
 
 	int32 NumShaderMaps = 0;
+	/** Maps from a shader type hash (see FShaderMapResourceCode::FShaderEditorOnlyDataEntry::ShaderTypeHashes) to its statistics. */
 	TMap<uint64, FPerTypeStats> ShaderTypeStats;
 };
 
@@ -44,6 +47,7 @@ struct FShaderCodeStats
 class FEditorShaderCodeArchive
 {
 	using ECookShaderLibrarySource = ShaderCodeArchive::ECookShaderLibrarySource;
+	using ESaveToDiskTarget = UE::ShaderLibrary::Private::ESaveToDiskTarget;
 
 	const FName FormatName;
 	FString LibraryName;
@@ -117,7 +121,7 @@ public:
 		FShaderCodeStats& CodeStats, ECookShaderLibrarySource CookSource);
 
 	/** Produces another archive that contains the code only for these assets */
-	FEditorShaderCodeArchive* CreateChunk(int ChunkId, const TSet<FName>& PackagesInChunk);
+	TUniquePtr<FEditorShaderCodeArchive> CreateChunk(int ChunkId, const TSet<FName>& PackagesInChunk);
 
 	/** Adds entries for the specified shader map from another archive with all its individual shaders and their code to this archive. */
 	int32 AddShaderCode(int32 OtherShaderMapIndex, const FEditorShaderCodeArchive& OtherArchive,
@@ -137,6 +141,9 @@ public:
 		AddExistingShaderCodeLibrary(OutputDir, CodeStats, CookSource);
 	}
 
+	/** Resets bReferencedByOplog and bReferencedByStaging to false for all shadermap-asset associations. */
+	void InitializeReferenceTracking();
+
 	/** Called by the cooker to report which assets are referenced in the current cook or in incrementalcook oplog. */
 	void UpdateOplogPackages(UE::Cook::Artifact::FUpdateOplogPackagesContext& Context);
 
@@ -147,7 +154,8 @@ public:
 	 * Shader entries and all their code to *.ushaderbytecode file,
 	 * Asset-shader associations to *.assetinfo.json file,
 	 * and shader type information to *.stinfo file. */
-	bool SaveToDisk(const FString& OutputDir, const FString& MetaOutputDir, bool bSaveOnlyAssetInfo = false, TArray<FString>* OutputFilenames = nullptr);
+	bool SaveToDisk(const FString& OutputDir, const FString& MetaOutputDir,
+		ESaveToDiskTarget Target = ESaveToDiskTarget::Staging, TArray<FString>* OutputFilenames = nullptr);
 
 	/** Creates a native shader library for this archive, returns the output filenames,
 	 * and deletes the shader code and pipeline files as they are no longer needed. */
@@ -156,11 +164,22 @@ public:
 	/** Prints the shader library statistics to the log. */
 	void DumpStatsAndDebugInfo();
 
+	/** Invokes the specified callback for each asset that references the specified shader type hash.
+	 * This can be slow and should only be used for debugging code or fallback conditions for assertions.
+	 * The callback must return true to continue iterating. */
+	void ForEachAssetReferencingShaderType(uint64 InShaderTypeHash, FSerializedShaderArchive::FAssetFilterFunctionRef AssetReferenceCallback) const;
+
 public:
 	/** Loads all existing shader code libraries from the new meta data directory and patches their entries into this archive. */
 	static bool CreatePatchLibrary(FName FormatName, FString const& LibraryName, TArray<FString> const& OldMetaDataDirs, FString const& NewMetaDataDir, FString const& OutDir, bool bNativeFormat, bool bNeedsDeterministicOrder);
 
+	TUniquePtr<FEditorShaderCodeArchive> CreateArchiveFromAssetsReferencedByStaging();
+	TUniquePtr<FEditorShaderCodeArchive> CreateArchiveFromAssetsReferencedByOplog();
+	TUniquePtr<FEditorShaderCodeArchive> CreateArchiveFromAssetsOnlyReferencedByOplog();
+
 private:
+	TUniquePtr<FEditorShaderCodeArchive> CreateArchiveFromFilteredAssets(const FString& NewLibraryName, FSerializedShaderArchive::FAssetFilterFunctionRef ShouldKeepAsset);
+
 	// Calls AddShaderCode() for all shader maps from the new library unless they were already added in one of the old libraries.
 	void MakePatchLibrary(TArray<FEditorShaderCodeArchive*> const& OldLibraries, FEditorShaderCodeArchive const& NewLibrary);
 
@@ -194,6 +213,10 @@ private:
 	void MarkShaderMapDirty(int32 ShaderMapIndex);
 	void MarkAllShaderMapsDirty();
 
+	// Wrapper to move all data from InArchive into this archive.
+	// Constant members such as Format and such must be equal since they must never change.
+	void AssignFrom(FEditorShaderCodeArchive& InArchive);
+
 };
 
 #endif // WITH_EDITOR
diff --git a/Engine/Source/Runtime/RenderCore/Private/ShaderLibrary/ShaderCodeArchive.cpp b/Engine/Source/Runtime/RenderCore/Private/ShaderLibrary/ShaderCodeArchive.cpp
--- a/Engine/Source/Runtime/RenderCore/Private/ShaderLibrary/ShaderCodeArchive.cpp
+++ b/Engine/Source/Runtime/RenderCore/Private/ShaderLibrary/ShaderCodeArchive.cpp
@@ -173,22 +173,21 @@ bool FSerializedShaderArchive::FindOrAddShaderMap(const FSHAHash& Hash, int32& O
 bool FSerializedShaderArchive::FindOrAddShaderMapEditor(const FSHAHash& Hash, int32& OutIndex,
 	const FShaderMapAssetPaths* AssociatedAssets, ECookShaderLibrarySource CookSource)
 {
-	bool bResult = FindOrAddShaderMapNonEditorInternal(Hash, OutIndex);
+	const bool bResult = FindOrAddShaderMapNonEditorInternal(Hash, OutIndex);
 
-	TOptional<FShaderMapAssetPaths> Buffer;
-	const FShaderMapAssetPaths* NormalizedAssets = AssociatedAssets;
+	FShaderMapAssetPaths AssociatedAssetsProxy;
 	if (!AssociatedAssets || AssociatedAssets->IsEmpty())
 	{
 		UE_LOG(LogShaderLibrary, Error,
 			TEXT("FindOrAddShaderMapEditor must be called with non-empty AssociatedAssets using the ShaderMap. ")
-			TEXT("ShaderMaps without assets will be pruned. Pass NAME_None for ShaderMaps that are readded on every cook. ")
+			TEXT("ShaderMaps without assets will be pruned. Pass NAME_None for ShaderMaps that are re-added on every cook. ")
 			TEXT("Hash %s is being assigned to NAME_None and will be pruned next cook."),
 			*Hash.ToString());
-		Buffer.Emplace();
-		Buffer->Add(NAME_None);
-		NormalizedAssets = Buffer.GetPtrOrNull();
+		AssociatedAssetsProxy.Add(NAME_None);
+		AssociatedAssets = &AssociatedAssetsProxy;
 	}
-	for (FName Asset : *NormalizedAssets)
+
+	for (FName Asset : *AssociatedAssets)
 	{
 		FShaderMapAssetAssociations::FAssociatedAssetData* AssetData = ShaderMapAssetAssociations.FindAsset(Asset);
 		if (AssetData)
@@ -391,9 +390,10 @@ bool LoadFromCompactBinary(FCbFieldView Field, FSerializedShaderArchive& OutArch
 			ShaderCodeArchive::ECookShaderLibrarySource LatestSource
 				= ShaderCodeArchive::ECookShaderLibrarySource::CurrentCook;
 			bOk = LoadFromCompactBinary(*ArrayIter++, AssetName) && bOk;
-			uint8 LatestSourceInt;
+			uint8 LatestSourceInt = 0xFF;
 			if (LoadFromCompactBinary(*ArrayIter++, LatestSourceInt)
-				&& LatestSourceInt < (uint8) ShaderCodeArchive::ECookShaderLibrarySource::Count)
+				&& LatestSourceInt >= (uint8)ShaderCodeArchive::ECookShaderLibrarySource::CurrentCook
+				&& LatestSourceInt <= (uint8)ShaderCodeArchive::ECookShaderLibrarySource::PreviousIncremental)
 			{
 				LatestSource = (ShaderCodeArchive::ECookShaderLibrarySource)LatestSourceInt;
 			}
@@ -964,10 +964,19 @@ void FSerializedShaderArchive::FShaderMapAssetAssociations::SortForSaving()
 }
 
 void FSerializedShaderArchive::CreateAsChunkFrom(const FSerializedShaderArchive& Parent, const TSet<FName>& PackagesInChunk, TArray<int32>& OutShaderCodeEntriesNeeded)
+{
+	CreateFromFilteredAssets(Parent, OutShaderCodeEntriesNeeded,
+		[&PackagesInChunk](FName AssetName, const FShaderMapAssetAssociations::FAssociatedAssetData& AssetData)
+		{
+			return PackagesInChunk.Contains(AssetName);
+		});
+}
+
+void FSerializedShaderArchive::CreateFromFilteredAssets(const FSerializedShaderArchive& Parent, TArray<int32>& OutShaderCodeEntriesNeeded, FAssetFilterFunctionRef ShouldKeepAsset)
 {
 	// we should begin with a clean slate
 	checkf(ShaderMapHashes.Num() == 0 && ShaderHashes.Num() == 0 && ShaderMapEntries.Num() == 0 && ShaderEntries.Num() == 0 && PreloadEntries.Num() == 0 && ShaderIndices.Num() == 0,
-		TEXT("Expecting a new, uninitialized FSerializedShaderArchive instance for creating a chunk."));
+		TEXT("Expecting a new, uninitialized FSerializedShaderArchive instance for CreateFromFilteredAssets."));
 
 	// go through parent's shadermap hashes in the order of their addition
 	for (int32 IdxSM = 0, NumSMs = Parent.ShaderMapHashes.Num(); IdxSM < NumSMs; ++IdxSM)
@@ -977,15 +986,30 @@ void FSerializedShaderArchive::CreateAsChunkFrom(const FSerializedShaderArchive&
 		bool bIncludeSM = false;
 		if (UNLIKELY(Assets == nullptr))
 		{
-			UE_LOG(LogShaderLibrary, Warning, TEXT("Shadermap %s is not associated with any asset. Including it in every chunk"), *ShaderMapHash.ToString());
-			bIncludeSM = true;
+			if (UE::ShaderLibrary::Private::bIncrementalCookPruningEnabled)
+			{
+				// ShaderMaps with no Assets occur when we load a ShaderLibrary from a previous cook, and find that
+				// one of the Assets it recorded has recompiled in the current cook, and it recompiled with different ShaderMaps
+				// than were previously recorded. The old ShaderMaps are added, but their association with the Asset is removed.
+				// These unreferenced ShaderMaps should be pruned from any CreateFromFilteredAssets call.
+				bIncludeSM = false;
+			}
+			else
+			{
+				// In the behavior before IncrementalCookPruning was implemented, it should never be possible for a ShaderMap
+				// to be lacking an asset. Log a warning and conservatively keep the ShaderMap in every CreateFromFilteredAssets call.
+				UE_LOG(LogShaderLibrary, Warning, TEXT("Shadermap %s is not associated with any asset. Including it in CreateFromFilteredAssets"), *ShaderMapHash.ToString());
+				bIncludeSM = true;
+			}
 		}
 		else
 		{
 			// if any asset is in the chunk, include
-			for (const FName& Asset : *Assets)
+			for (const FName& AssetName : *Assets)
 			{
-				if (PackagesInChunk.Contains(Asset))
+				const FShaderMapAssetAssociations::FAssociatedAssetData* AssetData = Parent.ShaderMapAssetAssociations.FindAsset(AssetName);
+				check(AssetData); // FShaderMapAssetAssociations guarantees that the Assets using a Shader are present in its list of assets.
+				if (ShouldKeepAsset(AssetName, *AssetData))
 				{
 					bIncludeSM = true;
 					break;
diff --git a/Engine/Source/Runtime/RenderCore/Private/ShaderLibrary/ShaderCodeLibrary.cpp b/Engine/Source/Runtime/RenderCore/Private/ShaderLibrary/ShaderCodeLibrary.cpp
--- a/Engine/Source/Runtime/RenderCore/Private/ShaderLibrary/ShaderCodeLibrary.cpp
+++ b/Engine/Source/Runtime/RenderCore/Private/ShaderLibrary/ShaderCodeLibrary.cpp
@@ -1002,7 +1002,7 @@ public:
 							for (const FString& Filename : UshaderbytecodeFiles)
 							{
 								const TCHAR* ChunkSubstring = TEXT("_Chunk");
-								const int kChunkSubstringSize = 6; // stlren(ChunkSubstring)
+								const int kChunkSubstringSize = 6; // strlen(ChunkSubstring)
 								int32 ChunkSuffix = Filename.Find(ChunkSubstring, ESearchCase::IgnoreCase, ESearchDir::FromEnd);
 								if (ChunkSuffix != INDEX_NONE && ChunkSuffix + kChunkSubstringSize < Filename.Len())
 								{
@@ -1561,7 +1561,8 @@ public:
 		StableArchive->AddShader(StableKeyValue, FEditorShaderStableInfo::EMergeRule::OverwriteUnmodifiedWarnModified);
 	}
  
-	void FinishPopulateShaderCode(const FString& ShaderCodeDir, const FString& MetaOutputDir, const TArray<FName>& ShaderFormats)
+	void FinishPopulateShaderCode(const FString& ShaderCodeDir, const FString& MetaOutputDir, const TArray<FName>& ShaderFormats,
+		bool bInitializeReferenceTracking)
 	{
 		for (FName ShaderFormatName : ShaderFormats)
 		{
@@ -1574,6 +1575,10 @@ public:
 			{
 				CodeArchive->FinishPopulate(ShaderCodeDir, CodeStats,
 					ShaderCodeArchive::ECookShaderLibrarySource::PreviousIncremental);
+				if (bInitializeReferenceTracking)
+				{
+					CodeArchive->InitializeReferenceTracking();
+				}
 			}
 
 			FEditorShaderStableInfo* StableArchive = EditorShaderStableInfo[SPlatform];
@@ -1614,7 +1619,8 @@ public:
 
 	bool SaveShaderCode(const FString& ShaderCodeDir, const FString& MetaOutputDir, const TArray<FName>& ShaderFormats, TArray<FString>& OutSCLCSVPath)
 	{
-		bool bOk = ShaderFormats.Num() > 0;
+		using namespace UE::ShaderLibrary::Private;
+		bool bAllOk = ShaderFormats.Num() > 0;
 
 		FScopeLock ScopeLock(&ShaderCodeCS);
 
@@ -1630,25 +1636,32 @@ public:
 				// Still save the full asset info for debugging
 				if (ChunksSaved[SPlatform].Num() == 0)
 				{
-					// always save shaders in our format even if the platform will use native one. This is needed for iterative cooks (Launch On et al)
-					// to reload previously cooked shaders
-					bOk = CodeArchive->SaveToDisk(ShaderCodeDir, MetaOutputDir) && bOk;
-
-					bool bShouldWriteInNativeFormat = bOk && bNativeFormat && CodeArchive->GetFormat()->SupportsShaderArchives();
-					if (bShouldWriteInNativeFormat)
+					TUniquePtr<FEditorShaderCodeArchive> ReferencedAssetsArchive(CodeArchive->CreateArchiveFromAssetsReferencedByStaging());
+					if (!ReferencedAssetsArchive->IsEmpty())
 					{
-						bOk = CodeArchive->PackageNativeShaderLibrary(ShaderCodeDir) && bOk;
-					}
+						// always save shaders in our format even if the platform will use native one. This is needed for iterative cooks (Launch On et al)
+						// to reload previously cooked shaders
+						bool bOk = ReferencedAssetsArchive->SaveToDisk(ShaderCodeDir, MetaOutputDir,
+							ESaveToDiskTarget::Staging);
 
-					if (bOk)
-					{
-						CodeArchive->DumpStatsAndDebugInfo();
+						bool bShouldWriteInNativeFormat = bOk && bNativeFormat && ReferencedAssetsArchive->GetFormat()->SupportsShaderArchives();
+						if (bShouldWriteInNativeFormat)
+						{
+							bOk = ReferencedAssetsArchive->PackageNativeShaderLibrary(ShaderCodeDir) && bOk;
+						}
+
+						if (bOk)
+						{
+							ReferencedAssetsArchive->DumpStatsAndDebugInfo();
+						}
+						bAllOk = bAllOk && bOk;
 					}
 				}
-				else
+				// save UnreferencedAssets for IncrementalCook
+				TUniquePtr<FEditorShaderCodeArchive> UnreferencedAssetsArchive(CodeArchive->CreateArchiveFromAssetsOnlyReferencedByOplog());
+				if (!UnreferencedAssetsArchive->IsEmpty())
 				{
-					// save asset info only, for debugging
-					bOk = CodeArchive->SaveToDisk(ShaderCodeDir, MetaOutputDir, true) && bOk;
+					bAllOk = UnreferencedAssetsArchive->SaveToDisk(ShaderCodeDir, MetaOutputDir, ESaveToDiskTarget::CookCache) && bAllOk;
 				}
 			}
 			FEditorShaderStableInfo* StableArchive = EditorShaderStableInfo[SPlatform];
@@ -1656,7 +1669,7 @@ public:
 			{
 				// Stable shader info is not saved per-chunk (it is not needed at runtime), so save it always
 				FString SCLCSVPath;
-				bOk &= StableArchive->SaveToDisk(MetaOutputDir, SCLCSVPath);
+				bAllOk = StableArchive->SaveToDisk(MetaOutputDir, SCLCSVPath) && bAllOk;
 				
 				// Only add output files if they were actually written to disk (if there were no shaders in the library it is not a failure).
 				if (!SCLCSVPath.IsEmpty())
@@ -1666,12 +1679,14 @@ public:
 			}
 		}
 
-		return bOk;
+		return bAllOk;
 	}
 
 	bool SaveShaderCodeChunk(int32 ChunkId, const TSet<FName>& InPackagesInChunk, const TArray<FName>& ShaderFormats,
 		const FString& SandboxDestinationPath, const FString& SandboxMetadataPath, TArray<FString>& OutChunkFilenames)
 	{
+		using namespace UE::ShaderLibrary::Private;
+
 		bool bOk = ShaderFormats.Num() > 0;
 
 		FScopeLock ScopeLock(&ShaderCodeCS);
@@ -1693,7 +1708,7 @@ public:
 				break;
 			}
 
-			FEditorShaderCodeArchive* PerChunkArchive = OrginalCodeArchive->CreateChunk(ChunkId, InPackagesInChunk);
+			TUniquePtr<FEditorShaderCodeArchive> PerChunkArchive = OrginalCodeArchive->CreateChunk(ChunkId, InPackagesInChunk);
 			if (!PerChunkArchive)
 			{
 				bOk = false;
@@ -1706,7 +1721,8 @@ public:
 
 				// always save shaders in our format even if the platform will use native one. This is needed for iterative cooks (Launch On et al)
 				// to reload previously cooked shaders
-				bOk = PerChunkArchive->SaveToDisk(SandboxDestinationPath, SandboxMetadataPath, false, &OutChunkFilenames) && bOk;
+				bOk = PerChunkArchive->SaveToDisk(SandboxDestinationPath, SandboxMetadataPath,
+					ESaveToDiskTarget::Staging, &OutChunkFilenames) && bOk;
 
 				bool bShouldWriteInNativeFormat = bOk && bNativeFormat && PerChunkArchive->GetFormat()->SupportsShaderArchives();
 				if (bShouldWriteInNativeFormat)
@@ -1783,9 +1799,32 @@ public:
 
 	void DumpShaderTypeStats(const FString& DebugInfoDir, const FString& MetaDataDir)
 	{
-		int32 PlatformId = 0;
-		for (const FShaderCodeStats& PlatformCodeStats : EditorShaderCodeStats)
+		// Returns true if the specified shader type is referenced by any asset for staging
+		auto IsShaderTypeReferencedByStaging = [this](int32 PlatformId, uint64 InShaderTypeHash) -> bool
+			{
+				bool bOutIsReferencedByStaging = false;
+				if (FEditorShaderCodeArchive* ShaderArchive = this->EditorShaderCodeArchive[PlatformId])
+				{
+					ShaderArchive->ForEachAssetReferencingShaderType(
+						InShaderTypeHash,
+						[&bOutIsReferencedByStaging](FName AssetName, const FSerializedShaderArchive::FShaderMapAssetAssociations::FAssociatedAssetData& AssetData) -> bool
+						{
+							if (AssetData.bReferencedByStaging)
+							{
+								// Mark shader type as referenced and stop iterating through further assets
+								bOutIsReferencedByStaging = true;
+								return false;
+							}
+							return true;
+						}
+					);
+				}
+				return bOutIsReferencedByStaging;
+			};
+
+		for (int32 PlatformId = 0; PlatformId < EShaderPlatform::SP_NumPlatforms; ++PlatformId)
 		{
+			const FShaderCodeStats& PlatformCodeStats = EditorShaderCodeStats[PlatformId];
 			if (PlatformCodeStats.NumShaderMaps > 0)
 			{
 				const FName PlatformFName = FGenericDataDrivenShaderPlatformInfo::GetName((EShaderPlatform)PlatformId);
@@ -1794,10 +1833,18 @@ public:
 				TMap<const FShaderType*, FShaderCodeStats::FPerTypeStats> CountersByTypeName;
 				for (const TPair<uint64, FShaderCodeStats::FPerTypeStats>& Pair : PlatformCodeStats.ShaderTypeStats)
 				{
-					const FShaderType* ShaderType = FindShaderTypeByName(FHashedName(Pair.Key));
-					checkf(ShaderType != nullptr, TEXT("Shader library contains a reference to an FShaderType with hash %llu which does not exist.\n"
-												       "If it has been removed or renamed, please bump CookIncrementalVersion to force a rebuild"), Pair.Key);
-					CountersByTypeName.Add(ShaderType, Pair.Value);
+					if (const FShaderType* ShaderType = FindShaderTypeByName(FHashedName(Pair.Key)))
+					{
+						CountersByTypeName.Add(ShaderType, Pair.Value);
+					}
+					else
+					{
+						// This error needs to be suppressed for data used by ShaderMaps or Shaders
+						// that are not referenced by any assets that have bReferencedByStaging.
+						checkf(!IsShaderTypeReferencedByStaging(PlatformId, Pair.Key),
+							TEXT("Shader library contains a reference to an FShaderType with hash %llu which does not exist.\n")
+							TEXT("If it has been removed or renamed, please bump CookIncrementalVersion to force a rebuild"), Pair.Key);
+					}
 				}
 				
 				// sort first by the type of shadertype, then by the name (so all global, material, etc. shaders are grouped together in the CSV)
@@ -1853,7 +1900,6 @@ public:
 					UE_LOG(LogShaderLibrary, Warning, TEXT("Failed to copy from '%s' to '%s'"), *Filename, *MetaPath);
 				}
 			}
-			PlatformId++;
 		}
 	}
 #endif // WITH_EDITOR
@@ -2413,7 +2459,7 @@ bool FShaderLibraryCooker::CreatePatchLibrary(TArray<FString> const& OldMetaData
 }
 
 void FShaderLibraryCooker::FinishPopulateShaderLibrary(const ITargetPlatform* TargetPlatform, FString const& Name,
-	FString const& SandboxDestinationPath, FString const& SandboxMetadataPath)
+	FString const& SandboxDestinationPath, FString const& SandboxMetadataPath, bool bInitializeReferenceTracking)
 {
 	const FString& ShaderCodeDir = SandboxDestinationPath;
 	const FString& MetaDataPath = SandboxMetadataPath;
@@ -2426,7 +2472,7 @@ void FShaderLibraryCooker::FinishPopulateShaderLibrary(const ITargetPlatform* Ta
 	TargetPlatform->GetAllTargetedShaderFormats(ShaderFormats);
 	if (ShaderFormats.Num() > 0)
 	{
-		FShaderLibrariesCollection::Impl->FinishPopulateShaderCode(ShaderCodeDir, MetaDataPath, ShaderFormats);
+		FShaderLibrariesCollection::Impl->FinishPopulateShaderCode(ShaderCodeDir, MetaDataPath, ShaderFormats, bInitializeReferenceTracking);
 	}
 }
 
@@ -2452,6 +2498,8 @@ void FShaderLibraryCooker::PruneShaderLibrary(const ITargetPlatform* TargetPlatf
 
 bool FShaderLibraryCooker::MergeShaderCodeArchive(const TArray<FString>& CookedMetadataDirs, const FString& OutputDir, TArray<FString>& OutWrittenFiles)
 {
+	using namespace UE::ShaderLibrary::Private;
+
 	OutWrittenFiles.Empty();
 
 	bool bNeedsDeterministicOrder = false;
@@ -2529,7 +2577,8 @@ bool FShaderLibraryCooker::MergeShaderCodeArchive(const TArray<FString>& CookedM
 	for (TPair<FString, TPair<FEditorShaderCodeArchive, FShaderCodeStats>>& ShaderCodeArchivePair : ShaderCodeArchives)
 	{
 		TArray<FString> CreatedFiles;
-		if (ShaderCodeArchivePair.Value.Key.SaveToDisk(OutShaderCodeDir, TEXT(""), false, &CreatedFiles))
+		if (ShaderCodeArchivePair.Value.Key.SaveToDisk(OutShaderCodeDir, TEXT(""),
+			ESaveToDiskTarget::Staging, &CreatedFiles))
 		{
 			OutWrittenFiles.Append(MoveTemp(CreatedFiles));
 		}
@@ -2539,7 +2588,8 @@ bool FShaderLibraryCooker::MergeShaderCodeArchive(const TArray<FString>& CookedM
 			bSuccess = false;
 		}
 		CreatedFiles.Reset();
-		if (ShaderCodeArchivePair.Value.Key.SaveToDisk(OutContentDir, TEXT(""), false, &CreatedFiles))
+		if (ShaderCodeArchivePair.Value.Key.SaveToDisk(OutContentDir, TEXT(""),
+			ESaveToDiskTarget::Staging, &CreatedFiles))
 		{
 			OutWrittenFiles.Append(MoveTemp(CreatedFiles));
 		}
diff --git a/Engine/Source/Runtime/RenderCore/Private/ShaderLibrary/ShaderCodeLibraryUtilities.cpp b/Engine/Source/Runtime/RenderCore/Private/ShaderLibrary/ShaderCodeLibraryUtilities.cpp
--- a/Engine/Source/Runtime/RenderCore/Private/ShaderLibrary/ShaderCodeLibraryUtilities.cpp
+++ b/Engine/Source/Runtime/RenderCore/Private/ShaderLibrary/ShaderCodeLibraryUtilities.cpp
@@ -24,13 +24,17 @@ namespace UE::ShaderLibrary::Private
 	const uint32 GShaderPipelineArchiveVersion = 1;
 
 	const FString ShaderExtension = ".ushaderbytecode";
+	const FString CookCacheShaderExtension = ".ucacheshaderbytecode";
 	const FString ShaderAssetInfoExtension = ".assetinfo.json";
 	const FString ShaderTypeInfoExtension = ".stinfo";
 	const FString StableExtension = ".shk";
 	const FString PipelineExtension = ".ushaderpipelines";
 
-	const FString ShaderCodePatternStr = "^ShaderCode\\-(\\w*)\\-([\\w\\-]*)\\.ushaderbytecode$";
-	const FString ShaderArchivePatternStr = "^ShaderArchive\\-(\\w*)\\-([\\w\\-]*)\\.ushaderbytecode$";
+	static const FString CookCacheSuffix = "_CookCache";
+
+	// ShaderCode(_CookCache)-LibraryName-Platform.u(cache)shaderbytecode
+	const FString ShaderCodePatternStr = "^ShaderCode(?:_CookCache)?\\-(\\w*)\\-([\\w\\-]*)\\.(?:ushaderbytecode|ucacheshaderbytecode)$";
+	const FString ShaderArchivePatternStr = "^ShaderArchive(?:_CookCache)?\\-(\\w*)\\-([\\w\\-]*)\\.(?:ushaderbytecode|ucacheshaderbytecode)$";
 	const FString StableInfoPatternStr = "^ShaderStableInfo\\-(\\w*)\\-([\\w\\-]*)\\.shk$";
 	const FString ShaderFileForChunkPatternStr = "Shader.*_Chunk[0-9]+-.*\\.(ushaderbytecode|assetinfo.json|stinfo|ushaderpipelines)";
 
@@ -86,9 +90,18 @@ namespace UE::ShaderLibrary::Private
 		return BaseDebugInfoPath / TEXT("Shaders");
 	}
 
-	FString GetCodeArchiveFilename(const FString& BaseDir, const FString& LibraryName, FName Platform)
+	FString GetCodeArchiveFilename(const FString& BaseDir, const FString& LibraryName, FName Platform, ESaveToDiskTarget Target)
 	{
-		return BaseDir / FString::Printf(TEXT("ShaderArchive-%s-"), *LibraryName) + Platform.ToString() + ShaderExtension;
+		switch (Target)
+		{
+		case ESaveToDiskTarget::Staging:
+			return BaseDir / FString::Printf(TEXT("ShaderArchive-%s-"), *LibraryName) + Platform.ToString() + ShaderExtension;
+		case ESaveToDiskTarget::CookCache:
+			return BaseDir / FString::Printf(TEXT("ShaderArchive%s-%s-"), *CookCacheSuffix, *LibraryName) + Platform.ToString() + CookCacheShaderExtension;
+		default:
+			checkNoEntry();
+			return FString();
+		}
 	}
 
 	FString GetStableInfoArchiveFilename(const FString& BaseDir, const FString& LibraryName, FName Platform)
@@ -101,19 +114,46 @@ namespace UE::ShaderLibrary::Private
 		return BaseDir / FString::Printf(TEXT("ShaderArchive-%s-"), *LibraryName) + Platform.ToString() + PipelineExtension;
 	}
 
-	FString GetShaderCodeFilename(const FString& BaseDir, const FString& LibraryName, FName Platform)
+	FString GetShaderCodeFilename(const FString& BaseDir, const FString& LibraryName, FName Platform, ESaveToDiskTarget Target)
 	{
-		return BaseDir / FString::Printf(TEXT("ShaderCode-%s-"), *LibraryName) + Platform.ToString() + ShaderExtension;
+		switch (Target)
+		{
+		case ESaveToDiskTarget::Staging:
+			return BaseDir / FString::Printf(TEXT("ShaderCode-%s-"), *LibraryName) + Platform.ToString() + ShaderExtension;
+		case ESaveToDiskTarget::CookCache:
+			return BaseDir / FString::Printf(TEXT("ShaderCode%s-%s-"), *CookCacheSuffix, *LibraryName) + Platform.ToString() + CookCacheShaderExtension;
+		default:
+			checkNoEntry();
+			return FString();
+		}
 	}
 
-	FString GetShaderAssetInfoFilename(const FString& BaseDir, const FString& LibraryName, FName Platform)
+	FString GetShaderAssetInfoFilename(const FString& BaseDir, const FString& LibraryName, FName Platform, ESaveToDiskTarget Target)
 	{
-		return BaseDir / FString::Printf(TEXT("ShaderAssetInfo-%s-"), *LibraryName) + Platform.ToString() + ShaderAssetInfoExtension;
+		switch (Target)
+		{
+		case ESaveToDiskTarget::Staging:
+			return BaseDir / FString::Printf(TEXT("ShaderAssetInfo-%s-"), *LibraryName) + Platform.ToString() + ShaderAssetInfoExtension;
+		case ESaveToDiskTarget::CookCache:
+			return BaseDir / FString::Printf(TEXT("ShaderAssetInfo%s-%s-"), *CookCacheSuffix, *LibraryName) + Platform.ToString() + ShaderAssetInfoExtension;
+		default:
+			checkNoEntry();
+			return FString();
+		}
 	}
 
-	FString GetShaderTypeInfoFilename(const FString& BaseDir, const FString& LibraryName, FName Platform)
+	FString GetShaderTypeInfoFilename(const FString& BaseDir, const FString& LibraryName, FName Platform, ESaveToDiskTarget Target)
 	{
-		return BaseDir / FString::Printf(TEXT("ShaderTypeInfo-%s-"), *LibraryName) + Platform.ToString() + ShaderTypeInfoExtension;
+		switch (Target)
+		{
+		case ESaveToDiskTarget::Staging:
+			return BaseDir / FString::Printf(TEXT("ShaderTypeInfo-%s-"), *LibraryName) + Platform.ToString() + ShaderTypeInfoExtension;
+		case ESaveToDiskTarget::CookCache:
+			return BaseDir / FString::Printf(TEXT("ShaderTypeInfo%s-%s-"), *CookCacheSuffix, *LibraryName) + Platform.ToString() + ShaderTypeInfoExtension;
+		default:
+			checkNoEntry();
+			return FString();
+		}
 	}
 
 	FString GetShaderDebugFolder(const FString& BaseDir, const FString& LibraryName, FName Platform)
diff --git a/Engine/Source/Runtime/RenderCore/Private/ShaderLibrary/ShaderCodeLibraryUtilities.h b/Engine/Source/Runtime/RenderCore/Private/ShaderLibrary/ShaderCodeLibraryUtilities.h
--- a/Engine/Source/Runtime/RenderCore/Private/ShaderLibrary/ShaderCodeLibraryUtilities.h
+++ b/Engine/Source/Runtime/RenderCore/Private/ShaderLibrary/ShaderCodeLibraryUtilities.h
@@ -15,15 +15,22 @@ class ICookArtifactReader;
 
 namespace UE::ShaderLibrary::Private
 {
+	enum class ESaveToDiskTarget : uint8
+	{
+		Staging = 0,
+		CookCache,
+	};
+
 	// Compile time switch to enable the new feature of pruning stale shader data from incremental cooks.
 	// Remove this switch once the feature is stable enough not to need an off switch.
-	constexpr bool bIncrementalCookPruningEnabled = false;
+	constexpr bool bIncrementalCookPruningEnabled = true;
 
 	// Shader library constants:
 	extern const uint32 GShaderCodeArchiveVersion;
 	extern const uint32 GShaderPipelineArchiveVersion;
 
 	extern const FString ShaderExtension;
+	extern const FString CookCacheShaderExtension;
 	extern const FString ShaderAssetInfoExtension;
 	extern const FString ShaderTypeInfoExtension;
 	extern const FString StableExtension;
@@ -48,12 +55,12 @@ namespace UE::ShaderLibrary::Private
 
 	// Shader library utility functions:
 	FString GetShaderCodeLibraryDebugInfoPath();
-	FString GetCodeArchiveFilename(const FString& BaseDir, const FString& LibraryName, FName Platform);
+	FString GetCodeArchiveFilename(const FString& BaseDir, const FString& LibraryName, FName Platform, ESaveToDiskTarget Target);
 	FString GetStableInfoArchiveFilename(const FString& BaseDir, const FString& LibraryName, FName Platform);
 	FString GetPipelinesArchiveFilename(const FString& BaseDir, const FString& LibraryName, FName Platform);
-	FString GetShaderCodeFilename(const FString& BaseDir, const FString& LibraryName, FName Platform);
-	FString GetShaderAssetInfoFilename(const FString& BaseDir, const FString& LibraryName, FName Platform);
-	FString GetShaderTypeInfoFilename(const FString& BaseDir, const FString& LibraryName, FName Platform);
+	FString GetShaderCodeFilename(const FString& BaseDir, const FString& LibraryName, FName Platform, ESaveToDiskTarget Target);
+	FString GetShaderAssetInfoFilename(const FString& BaseDir, const FString& LibraryName, FName Platform, ESaveToDiskTarget Target);
+	FString GetShaderTypeInfoFilename(const FString& BaseDir, const FString& LibraryName, FName Platform, ESaveToDiskTarget Target);
 	FString GetShaderDebugFolder(const FString& BaseDir, const FString& LibraryName, FName Platform);
 
 	/** Helper function shared between the cooker and runtime */
diff --git a/Engine/Source/Runtime/RenderCore/Private/ShaderLibrary/ShaderLibraryInstance.cpp b/Engine/Source/Runtime/RenderCore/Private/ShaderLibrary/ShaderLibraryInstance.cpp
--- a/Engine/Source/Runtime/RenderCore/Private/ShaderLibrary/ShaderLibraryInstance.cpp
+++ b/Engine/Source/Runtime/RenderCore/Private/ShaderLibrary/ShaderLibraryInstance.cpp
@@ -45,7 +45,7 @@ FShaderLibraryInstance* FShaderLibraryInstance::Create(EShaderPlatform InShaderP
 		const FName ShaderFormatName = LegacyShaderPlatformToShaderFormat(InShaderPlatform);
 		FString ShaderFormatAndPlatform = ShaderFormatName.ToString() + TEXT("-") + PlatformName.ToString();
 
-		const FString DestFilePath = GetCodeArchiveFilename(ShaderCodeDir, InLibraryName, FName(ShaderFormatAndPlatform));
+		const FString DestFilePath = GetCodeArchiveFilename(ShaderCodeDir, InLibraryName, FName(ShaderFormatAndPlatform), ESaveToDiskTarget::Staging);
 		TUniquePtr<FArchive> Ar(CreateShaderFileReader(*DestFilePath));
 		if (Ar)
 		{
diff --git a/Engine/Source/Runtime/RenderCore/Public/ShaderCodeArchive.h b/Engine/Source/Runtime/RenderCore/Public/ShaderCodeArchive.h
--- a/Engine/Source/Runtime/RenderCore/Public/ShaderCodeArchive.h
+++ b/Engine/Source/Runtime/RenderCore/Public/ShaderCodeArchive.h
@@ -25,6 +25,10 @@
 #include "Templates/RefCounting.h"
 #include "UObject/NameTypes.h"
 
+#if WITH_EDITOR
+#include "Templates/Function.h"
+#endif
+
 #if WITH_EDITOR
 class FCbFieldView;
 class FCbWriter;
@@ -98,15 +102,15 @@ enum class ECookShaderLibrarySource : uint8
 {
 	/** The data is provided by compilation that executed during the current cook; this is the most up-to-date. */
 	CurrentCook = 0,
+
 	/** The current cook is a DLC cook and the data being provided came from the basegame cook. */
 	PatchBase,
+
 	/** The data was written to the ShaderLibrary artifacts in a previous incremental cook. */
 	PreviousIncremental,
-	Count,
 };
- 
-}
 
+} // namespace ShaderCodeArchive
 #endif // WITH_EDITOR
 
 // Portion of shader code archive that's serialize to disk
@@ -171,11 +175,13 @@ public:
 		{
 			TArray<FSHAHash> ShaderMaps;
 			ECookShaderLibrarySource LatestSource = ECookShaderLibrarySource::PreviousIncremental;
-			// These reference flags are transient, intialized to false, and are read/written only during end of cook.
+			// These reference flags are transient. By default they are initialized to true, for shaderlibraries that do not use
+			// reference tracking such as the global shader library. When tracking references, we initialize them to false.
+			// These flags are read/written only during end of cook.
 			/** Set to true if the Asset is needed either for incremental cook or for staging, and should therefore not be pruned. */
-			bool bReferencedByOplog = false;
+			bool bReferencedByOplog = true;
 			/** Set to true if the Asset is needed for runtime and staging, and should therefore be added to the staging artifacts. */
-			bool bReferencedByStaging = false;
+			bool bReferencedByStaging = true;
 
 			/** this->LatestSource = NewestOf(LatestSource, InCookSource). */
 			void MergeSource(ECookShaderLibrarySource InCookSource);
@@ -201,8 +207,10 @@ public:
 		TMap<FSHAHash, FShaderMapAssetPaths> ShaderMapToAssets;
 		TMap<FName, FAssociatedAssetData> AssetToShaderMaps;
 	};
+
 private:
 	FShaderMapAssetAssociations ShaderMapAssetAssociations;
+
 public:
 	FShaderMapAssetAssociations& GetShaderMapAssetAssociations();
 	const FShaderMapAssetAssociations& GetShaderMapAssetAssociations() const;
@@ -212,6 +220,7 @@ public:
 
 	struct FShaderTypeHashes
 	{
+		/** Most of the time, this will only contain one element. But some shaders are deduplicated to two or more shader types. */
 		TArray<uint64> Data;
 
 		friend FArchive& operator<<(FArchive& Ar, FShaderTypeHashes& Ref)
@@ -372,7 +381,11 @@ public:
 	RENDERCORE_API void SaveAssetInfo(FArchive& Ar);
 	RENDERCORE_API bool LoadAssetInfo(const FString& Filename, ECookShaderLibrarySource CookSource = ECookShaderLibrarySource::CurrentCook);
 	RENDERCORE_API bool LoadAssetInfo(FArchive* Ar, ECookShaderLibrarySource CookSource = ECookShaderLibrarySource::CurrentCook);
+	UE_DEPRECATED(5.8, "Use CreateFromFilteredAssets instead")
 	RENDERCORE_API void CreateAsChunkFrom(const FSerializedShaderArchive& Parent, const TSet<FName>& PackagesInChunk, TArray<int32>& OutShaderCodeEntriesNeeded);
+
+	using FAssetFilterFunctionRef = TFunctionRef<bool(FName AssetName, const FShaderMapAssetAssociations::FAssociatedAssetData& AssetData)>;
+	RENDERCORE_API void CreateFromFilteredAssets(const FSerializedShaderArchive& Parent, TArray<int32>& OutShaderCodeEntriesNeeded, FAssetFilterFunctionRef ShouldKeepAsset);
 	RENDERCORE_API void CollectStatsAndDebugInfo(FDebugStats& OutDebugStats, FExtendedDebugStats* OutExtendedDebugStats);
 	RENDERCORE_API void DumpContentsInPlaintext(FString& OutText) const;
 	RENDERCORE_API friend FCbWriter& operator<<(FCbWriter& Writer, const FSerializedShaderArchive& Archive);
diff --git a/Engine/Source/Runtime/RenderCore/Public/ShaderCodeLibrary.h b/Engine/Source/Runtime/RenderCore/Public/ShaderCodeLibrary.h
--- a/Engine/Source/Runtime/RenderCore/Public/ShaderCodeLibrary.h
+++ b/Engine/Source/Runtime/RenderCore/Public/ShaderCodeLibrary.h
@@ -348,7 +348,7 @@ struct FShaderLibraryCooker
 
 	/** Finishes collection of data that should be in the named code library. This includes loading data from a previous iterative cook. */
 	static RENDERCORE_API void FinishPopulateShaderLibrary(const ITargetPlatform* TargetPlatform, FString const& Name, FString const& SandboxDestinationPath,
-		FString const& SandboxMetadataPath);
+		FString const& SandboxMetadataPath, bool bInitializeReferenceTracking = false);
 
 	/** Called by the cooker to report which assets are referenced in the current cook or in incrementalcook oplog. */
 	static RENDERCORE_API void UpdateOplogPackages(UE::Cook::Artifact::FUpdateOplogPackagesContext& Context);
