Remove deprecation flag from IPersistentViewUniformBufferExtension.

#Synced-CL 49186836

diff --git a/Engine/Source/Runtime/RenderCore/Public/RendererInterface.h b/Engine/Source/Runtime/RenderCore/Public/RendererInterface.h
--- a/Engine/Source/Runtime/RenderCore/Public/RendererInterface.h
+++ b/Engine/Source/Runtime/RenderCore/Public/RendererInterface.h
@@ -883,7 +883,6 @@ public:
 	/**	Start prefetching streaming data for Nanite resource that will soon be used for rendering. TODO: Implement callback mechanism */
 	virtual void PrefetchNaniteResource(const Nanite::FResources* Resource, uint32 NumFramesUntilRender) = 0;
 
-	UE_DEPRECATED(5.4, "IPersistentViewUniformBufferExtension will be removed in a future version.")
 	virtual void RegisterPersistentViewUniformBufferExtension(IPersistentViewUniformBufferExtension* Extension) = 0;
 
 	/**
