[HWRT] Fix crash when using platform preview in Editor but project has ray tracing shaders disabled on Windows.

#Synced-CL 49330151

diff --git a/Engine/Source/Runtime/RenderCore/Private/RenderUtils.cpp b/Engine/Source/Runtime/RenderCore/Private/RenderUtils.cpp
--- a/Engine/Source/Runtime/RenderCore/Private/RenderUtils.cpp
+++ b/Engine/Source/Runtime/RenderCore/Private/RenderUtils.cpp
@@ -902,9 +902,7 @@ RENDERCORE_API void RenderUtilsInit()
 					{
 						uint32 ShaderPlatformIndex = static_cast<uint32>(ShaderPlatform);
 						GRayTracingPlatformMask[ShaderPlatformIndex] = true;
-						GRayTracingShadersPlatformMask[ShaderPlatformIndex] = FDataDrivenShaderPlatformInfo::GetIsPreviewPlatform(ShaderPlatform)
-							? true
-							: bRayTracingShadersEnabled;
+						GRayTracingShadersPlatformMask[ShaderPlatformIndex] = bRayTracingShadersEnabled;
 					}
 				}
 			}
