[ShaderLibrary] Exclude shader maps that are re-added on every cook from reference tracking.

Problem:
Since CL 49276028, all shader maps that are not associated by a referenced asset during cooking are pruned. Although, the documentation states at multiple places that shader maps can be associated with a special name ("NAME_None"), if they are added on every cook and therefore do not get excluded from the cook.

Solution:
Exclude shader maps that are re-added on every cook from reference tracking.

#tests Verified shader resources present in packaged game (NNETest).

#rnx

#Synced-CL 49589669

diff --git a/Engine/Source/Runtime/RenderCore/Private/ShaderLibrary/EditorShaderCodeArchive.cpp b/Engine/Source/Runtime/RenderCore/Private/ShaderLibrary/EditorShaderCodeArchive.cpp
--- a/Engine/Source/Runtime/RenderCore/Private/ShaderLibrary/EditorShaderCodeArchive.cpp
+++ b/Engine/Source/Runtime/RenderCore/Private/ShaderLibrary/EditorShaderCodeArchive.cpp
@@ -831,11 +831,22 @@ void FEditorShaderCodeArchive::UpdateOplogPackages(UE::Cook::Artifact::FUpdateOp
 	for (TPair<FName, FShaderMapAssetAssociations::FAssociatedAssetData>& Pair : AssetMap.ViewAssets())
 	{
 		FName PackageName = Pair.Key;
-		const FOplogPackageData* OplogData = OplogPackages.Find(PackageName);
-		if (OplogData)
+		if (PackageName.IsNone())
 		{
+			// NAME_None is used for shader maps that are re-added on every incremental cook. It will not be found in
+			// the list of PackageNames passed in by the cooker, so mark it as referenced by staging and oplog as a
+			// special case so that we keep these globally added shaders for staging.
 			Pair.Value.bReferencedByOplog = true;
-			Pair.Value.bReferencedByStaging |= OplogData->IsReferenced();
+			Pair.Value.bReferencedByStaging = true;
+		}
+		else
+		{
+			const FOplogPackageData* OplogData = OplogPackages.Find(PackageName);
+			if (OplogData)
+			{
+				Pair.Value.bReferencedByOplog = true;
+				Pair.Value.bReferencedByStaging |= OplogData->IsReferenced();
+			}
 		}
 	}
 }
diff --git a/Engine/Source/Runtime/RenderCore/Private/ShaderLibrary/ShaderCodeArchive.cpp b/Engine/Source/Runtime/RenderCore/Private/ShaderLibrary/ShaderCodeArchive.cpp
--- a/Engine/Source/Runtime/RenderCore/Private/ShaderLibrary/ShaderCodeArchive.cpp
+++ b/Engine/Source/Runtime/RenderCore/Private/ShaderLibrary/ShaderCodeArchive.cpp
@@ -228,6 +228,15 @@ bool FSerializedShaderArchive::FindOrAddShaderMapEditor(const FSHAHash& Hash, in
 			// and we report the newest.
 			AssetData->MergeSource(CookSource);
 		}
+		else if (Asset.IsNone() && UE::ShaderLibrary::Private::bIncrementalCookPruningEnabled
+			&& CookSource >= ECookShaderLibrarySource::PreviousIncremental)
+		{
+			// ShaderMaps that are associated with asset name NAME_None have a contract that they are never loaded from
+			// incremental and instead are always added in every cook. Add support to enforce this contract in the case
+			// that some shaders previously existed for NAME_None but in the current cook they were not added. We need
+			// to not add the stale references in that case.
+			continue;
+		}
 		else
 		{
 			// Add the new Asset to ShaderMapAssetAssociations with the current ShaderMap.
