Fix check when using FRayTracingGeometryInitializer as SourceGeometry (bTemplate=true) for a skeletal mesh with offline data.

#Synced-CL 49337795

diff --git a/Engine/Source/Runtime/RenderCore/Private/RayTracingGeometry.cpp b/Engine/Source/Runtime/RenderCore/Private/RayTracingGeometry.cpp
--- a/Engine/Source/Runtime/RenderCore/Private/RayTracingGeometry.cpp
+++ b/Engine/Source/Runtime/RenderCore/Private/RayTracingGeometry.cpp
@@ -38,7 +38,7 @@ bool FRayTracingGeometry::HasValidInitializer() const
 		if (!Segment.VertexBuffer)
 		{
 			bAllSegmentsAreValid = false;
-			checkf(Initializer.OfflineData == nullptr, TEXT("RayTracingGeometry (%s) should not have OfflineData"), *Initializer.DebugName.ToString());
+			checkf(Initializer.bTemplate || Initializer.OfflineData == nullptr, TEXT("RayTracingGeometry (%s) should not have OfflineData"), *Initializer.DebugName.ToString());
 			break;
 		}
 	}
