[Apple] Metallib stable ordering and per chunk deduplication

Metallib contents are now stable sorted and deduped on the library level, not library chunk level.
This gives us:
- ordering for smaller future patches
- metallibs size reduction

+ MetalShaderPlatform - In SPIRV path, remove header comments from CRC to improve duplication detection by [REDACTED_PERSON]
which also cuts the size.

This invalidates _all_ Metal shaders to force the metallibs rebuilding.

Same change was initially implemented in 39.xx and shipped live -- 48872703. This CL is a manual merge for Main as touched code was refactored there.

[RN]

diff --git a/Engine/Shaders/Public/Platform/Metal/MetalCommon.ush b/Engine/Shaders/Public/Platform/Metal/MetalCommon.ush
--- a/Engine/Shaders/Public/Platform/Metal/MetalCommon.ush
+++ b/Engine/Shaders/Public/Platform/Metal/MetalCommon.ush
@@ -7,7 +7,7 @@
 #pragma once
 
 // Update this GUID to invalidate shader recompilation for only Metal shaders
-#pragma message("UESHADERMETADATA_VERSION 83AC7099-FE8A-45C2-B6F5-94E8DF69E3A3")
+#pragma message("UESHADERMETADATA_VERSION 5402461A-0CAA-4A0B-844D-8DE806520E9A")
 
 #if VERTEXSHADER && HULLSHADER
 #ifndef TESSELLATIONSHADER
diff --git a/Engine/Source/Developer/Apple/MetalShaderFormat/Private/MetalCompileShaderSPIRV.cpp b/Engine/Source/Developer/Apple/MetalShaderFormat/Private/MetalCompileShaderSPIRV.cpp
--- a/Engine/Source/Developer/Apple/MetalShaderFormat/Private/MetalCompileShaderSPIRV.cpp
+++ b/Engine/Source/Developer/Apple/MetalShaderFormat/Private/MetalCompileShaderSPIRV.cpp
@@ -1082,27 +1082,6 @@ void FMetalCompileShaderSPIRV::DoCompileMetalShader(
 		}
 		else
 		{
-			if (FCStringAnsi::Strstr(ResultsTargetDataAsString.c_str(), "spvBufferSizeConstants"))
-			{
-				CCHeaderWriter.WriteSideTable(TEXT("spvBufferSizeConstants"), SideTableIndex);
-			}
-
-			CCHeaderWriter.WriteSourceInfo(*Input.VirtualSourceFilePath, *Input.EntryPointName);
-			CCHeaderWriter.WriteCompilerInfo();
-
-			FString MetaData = CCHeaderWriter.ToString();
-			MetaData += RTString;
-			MetaData += TEXT("\n\n");
-			if (ALNString.Len())
-			{
-				MetaData += TEXT("// Attributes: ");
-				MetaData += ALNString;
-				MetaData += TEXT("\n\n");
-			}
-			
-			MetalSource = TCHAR_TO_UTF8(*MetaData);
-			MetalSource += ResultsTargetDataAsString;
-			
 			if (Options.bEnableFMAPass)
 			{
 				std::string FMADefine = std::string("\n"
@@ -1119,17 +1098,39 @@ void FMetalCompileShaderSPIRV::DoCompileMetalShader(
 										);
 				
 				std::string IncludeString = "using namespace metal;";
-				size_t IncludePos = MetalSource.find(IncludeString);
+				size_t IncludePos = ResultsTargetDataAsString.find(IncludeString);
 				if (IncludePos != std::string::npos)
-					MetalSource.replace(IncludePos, IncludeString.length(), FMADefine);
+					ResultsTargetDataAsString.replace(IncludePos, IncludeString.length(), FMADefine);
 			}
 
-			CRCLen = MetalSource.length();
-			CRC = FCrc::MemCrc_DEPRECATED(MetalSource.c_str(), CRCLen);
+			// We calculate the CRC of the metal source code before adding comments
+			CRCLen = ResultsTargetDataAsString.length();
+			CRC = FCrc::MemCrc_DEPRECATED(ResultsTargetDataAsString.c_str(), CRCLen);
 			
 			ANSICHAR MainCRC[25];
 			int32 NewLen = FCStringAnsi::Snprintf(MainCRC, 25, "Main_%0.8x_%0.8x(", CRCLen, CRC);
+
+			if (FCStringAnsi::Strstr(ResultsTargetDataAsString.c_str(), "spvBufferSizeConstants"))
+			{
+				CCHeaderWriter.WriteSideTable(TEXT("spvBufferSizeConstants"), SideTableIndex);
+			}
+
+			CCHeaderWriter.WriteSourceInfo(*Input.VirtualSourceFilePath, *Input.EntryPointName);
+			CCHeaderWriter.WriteCompilerInfo();
+
+			FString MetaData = CCHeaderWriter.ToString();
+			MetaData += RTString;
+			MetaData += TEXT("\n\n");
+			if (ALNString.Len())
+			{
+				MetaData += TEXT("// Attributes: ");
+				MetaData += ALNString;
+				MetaData += TEXT("\n\n");
+			}
 			
+			MetalSource = TCHAR_TO_UTF8(*MetaData);
+			MetalSource += ResultsTargetDataAsString;
+
 			std::string MainEntryPoint = EntryPointNameAnsi + "(";
 			size_t Pos;
 			do
diff --git a/Engine/Source/Developer/Apple/MetalShaderFormat/Private/MetalShaderCompiler.cpp b/Engine/Source/Developer/Apple/MetalShaderFormat/Private/MetalShaderCompiler.cpp
--- a/Engine/Source/Developer/Apple/MetalShaderFormat/Private/MetalShaderCompiler.cpp
+++ b/Engine/Source/Developer/Apple/MetalShaderFormat/Private/MetalShaderCompiler.cpp
@@ -5,6 +5,7 @@
 #if UE_ENABLE_INCLUDE_ORDER_DEPRECATED_IN_5_4
 #include "CoreMinimal.h"
 #endif
+#include "Containers/ArrayView.h"
 #include "HAL/PlatformFileManager.h"
 #include "MetalShaderFormat.h"
 #include "MetalShaderResources.h"
@@ -1267,7 +1268,7 @@ uint64 AppendShader_Metal(FString const& WorkingDir, const FSHAHash& Hash, TArra
 				
 				// Object code segment
 				FString ObjFilename = WorkingDir / FString::Printf(TEXT("Main_%0.8x_%0.8x.o"), Header.SourceLen, Header.SourceCRC);
-				
+
 				bool const bHasObjectData = (ObjectCodeDataSize > 0) || IFileManager::Get().FileExists(*ObjFilename);
 				if (bHasObjectData)
 				{
@@ -1330,7 +1331,7 @@ uint64 AppendShader_Metal(FString const& WorkingDir, const FSHAHash& Hash, TArra
 	return Id;
 }
 
-bool FinalizeLibrary_Metal(FName const& Format, FString const& WorkingDir, FString const& LibraryPath, TSet<uint64> const& Shaders, class FString const& DebugOutputDir, bool bStripLibs)
+bool FinalizeLibrary_Metal(FName const& Format, FString const& WorkingDir, FString const& LibraryPath, TArrayView<const uint64> Shaders, class FString const& DebugOutputDir, bool bStripLibs)
 {
 	bool bOK = false;
 
diff --git a/Engine/Source/Developer/Apple/MetalShaderFormat/Private/MetalShaderFormat.cpp b/Engine/Source/Developer/Apple/MetalShaderFormat/Private/MetalShaderFormat.cpp
--- a/Engine/Source/Developer/Apple/MetalShaderFormat/Private/MetalShaderFormat.cpp
+++ b/Engine/Source/Developer/Apple/MetalShaderFormat/Private/MetalShaderFormat.cpp
@@ -21,6 +21,8 @@
 #include "FileUtilities/ZipArchiveWriter.h"
 #include "MetalShaderCompiler.h"
 #include "DataDrivenShaderPlatformInfo.h"
+#include "Containers/ArrayView.h"
+#include "Templates/Sorting.h"
 
 DEFINE_LOG_CATEGORY(LogMetalCompilerSetup)
 DEFINE_LOG_CATEGORY(LogMetalShaderCompiler)
@@ -52,7 +54,7 @@ extern void OutputMetalDebugData(const FShaderCompilerInput& Input, const FShade
 
 extern bool StripShader_Metal(TArray<uint8>& Code, class FString const& DebugPath, bool const bNative);
 extern uint64 AppendShader_Metal(class FString const& ArchivePath, const FSHAHash& Hash, TArray<uint8>& Code);
-extern bool FinalizeLibrary_Metal(class FName const& Format, class FString const& ArchivePath, class FString const& LibraryPath, TSet<uint64> const& Shaders, class FString const& DebugOutputDir, bool bStripLibs);
+extern bool FinalizeLibrary_Metal(class FName const& Format, class FString const& ArchivePath, class FString const& LibraryPath, TArrayView<const uint64> Shaders, class FString const& DebugOutputDir, bool bStripLibs);
 
 /** Version for shader format, this becomes part of the DDC key. */
 static const FGuid UE_SHADER_METAL_VER = FGuid("B0DC25EF-C34D-437A-94E5-5E5146AF1B9A");
@@ -173,7 +175,8 @@ public:
 		return CanCompileBinaryShaders();
 	}
 
-	virtual bool CreateShaderArchive(FString const& LibraryName,
+	virtual bool CreateShaderArchive(
+		FString const& LibraryName,
 		FName ShaderFormatAndShaderPlatformName,
 		const FString& WorkingDirectory,
 		const FString& OutputDir,
@@ -183,7 +186,7 @@ public:
 		TArray<FString>* OutputFiles) const override final
 	{
 		int32 NumShadersPerLibrary = CVarNumShadersPerLibrary.GetValueOnAnyThread();
-		
+
 		check(LibraryName.Len() > 0);
 
 		TArray<FString> Components;
@@ -213,7 +216,8 @@ public:
 		TArray<uint8> StrippedShaderCode;
 		TArray<uint8> TempShaderCode;
 
-		TArray<TSet<uint64>> SubLibraries;
+		TArray<uint64> AllShaders;
+		AllShaders.SetNumUninitialized(SerializedShaders.GetNumShaders());
 
 		for (int32 ShaderIndex = 0; ShaderIndex < SerializedShaders.GetNumShaders(); ++ShaderIndex)
 		{
@@ -221,35 +225,91 @@ public:
 			StripShader_Metal(TempShaderCode, DebugOutputDir, true);
 
 			uint64 ShaderId = AppendShader_Metal(ArchivePath, SerializedShaders.ShaderHashes[ShaderIndex], TempShaderCode);
-			uint32 LibraryIndex = ShaderIndex / NumShadersPerLibrary;
+			AllShaders[ShaderIndex] = ShaderId;
 
-			if (ShaderId)
+			FShaderCodeEntry& ShaderEntry = SerializedShaders.ShaderEntries[ShaderIndex];
+			ShaderEntry.Size = TempShaderCode.Num();
+
+			StrippedShaderCode.Append(TempShaderCode);
+		}
+		// SerializedShaders archive is finalized later, when we sort out
+		// the duplicated entries.
+
+		// This is used later to lookup the shader code entry by the
+		// original shader index.
+		TArray<uint64> UnsortedAllShaders(AllShaders);
+
+		// Gives us stable ordering, while also clumping all
+		// duplicates together, if any.
+		RadixSort64(AllShaders.GetData(), AllShaders.Num());
+
+		//
+		// As we throw away duplicates, we need to re-calculate
+		// each shader index in the library. That index is serialized
+		// per shader code entry and is used during runtime loading.
+		//
+		// The loop below does both operations.
+		//
+
+		TMap<uint64, uint32> ShaderIdToIndexInLibrary;
+		ShaderIdToIndexInLibrary.Reserve(AllShaders.Num());
+
+		TArray<uint64> UniqueShaders;
+		UniqueShaders.Reserve(AllShaders.Num());
+		{
+			uint64 LastShader = 0;
+			for (uint64 Shader : AllShaders)
 			{
-				if (SubLibraries.Num() <= (int32)LibraryIndex)
+				if (Shader == LastShader)
 				{
-					SubLibraries.Add(TSet<uint64>());
+					continue;
 				}
-				SubLibraries[LibraryIndex].Add(ShaderId);
-			}
+				LastShader = Shader;
 
-			FShaderCodeEntry& ShaderEntry = SerializedShaders.ShaderEntries[ShaderIndex];
-			ShaderEntry.Size = TempShaderCode.Num();
-			ShaderEntry.UncompressedSize = TempShaderCode.Num();
+				ShaderIdToIndexInLibrary.Add(Shader, UniqueShaders.Num());
+				UniqueShaders.Add(Shader);
+			}
 
-			StrippedShaderCode.Append(TempShaderCode);
+			UE_LOG(LogShaders, Display, TEXT("Metallib %s has %d shaders (%d unique)"), *LibraryName, AllShaders.Num(), UniqueShaders.Num());
+			AllShaders.Empty();
 		}
 
+		for (int32 ShaderIndex = 0; ShaderIndex < SerializedShaders.GetNumShaders(); ++ShaderIndex)
+		{
+			uint64 ShaderId = UnsortedAllShaders[ShaderIndex];
+			FShaderCodeEntry& ShaderEntry = SerializedShaders.ShaderEntries[ShaderIndex];
+			ShaderEntry.SetLibraryIndex(ShaderIdToIndexInLibrary[ShaderId]);
+		}
 		SerializedShaders.Finalize();
 
+		TArrayView<const uint64> Workset = MakeArrayView(UniqueShaders.GetData(), UniqueShaders.Num());
+		TArray<TArrayView<const uint64>> SubLibraries;
+		
+		int32 CurrentSubLibrarySize = 0;
+		for (uint64 Shader : UniqueShaders)
+		{
+			++CurrentSubLibrarySize;
+			if (CurrentSubLibrarySize == NumShadersPerLibrary)
+			{
+				SubLibraries.Add(Workset.Left(CurrentSubLibrarySize));
+				Workset.RightChopInline(CurrentSubLibrarySize);
+				CurrentSubLibrarySize = 0;
+			}
+		}
+		if (!Workset.IsEmpty())
+		{
+			SubLibraries.Add(Workset);
+		}
+
 		bool bOK = false;
 		FString LibraryPlatformName = FString::Printf(TEXT("%s_%s"), *LibraryName, *ShaderFormatAndShaderPlatformName.GetPlainNameString());
 		LibraryPlatformName.ToLowerInline();
 		volatile int32 CompiledLibraries = 0;
-		TArray<FGraphEventRef> Tasks;
+		FGraphEventArray Tasks;
 
-		for (uint32 Index = 0; Index < (uint32)SubLibraries.Num(); Index++)
+		for (int32 Index = 0; Index < SubLibraries.Num(); ++Index)
 		{
-			TSet<uint64>& PartialShaders = SubLibraries[Index];
+			const TArrayView<const uint64>& PartialShaders = SubLibraries[Index];
 
 			FString LibraryPath = (OutputDir / LibraryPlatformName) + FString::Printf(TEXT(".%d"), Index) + FMetalCompilerToolchain::MetalLibraryExtension;
 			if (OutputFiles)
@@ -322,11 +382,7 @@ public:
 		Tasks.Add(DebugDataCompletionFence);
 #endif // WITH_ENGINE
 
-		// Wait for tasks
-		for (auto& Task : Tasks)
-		{
-			FTaskGraphInterface::Get().WaitUntilTaskCompletes(Task);
-		}
+		FTaskGraphInterface::Get().WaitUntilTasksComplete(Tasks);
 
 		if (CompiledLibraries == SubLibraries.Num())
 		{
diff --git a/Engine/Source/Developer/IoStoreUtilities/Private/IoStoreShaderLibraryProcessor.cpp b/Engine/Source/Developer/IoStoreUtilities/Private/IoStoreShaderLibraryProcessor.cpp
--- a/Engine/Source/Developer/IoStoreUtilities/Private/IoStoreShaderLibraryProcessor.cpp
+++ b/Engine/Source/Developer/IoStoreUtilities/Private/IoStoreShaderLibraryProcessor.cpp
@@ -177,19 +177,20 @@ public:
 	uint8* LoadAndDecompressShader(uint8* Destination, int32 ShaderIndex) const
 	{
 		const FShaderCodeEntry& IndividuallyCompressedShader = ShaderEntries[ShaderIndex];
+		uint32 UncompressedSize = IndividuallyCompressedShader.GetUncompressedSize();
 		{
 			// small shaders might be stored without compression, handle them here
-			if (IndividuallyCompressedShader.Size == IndividuallyCompressedShader.UncompressedSize)
+			if (IndividuallyCompressedShader.Size == UncompressedSize)
 			{
 				FMemory::Memcpy(Destination, GetShaderCodeDataStart() + IndividuallyCompressedShader.Offset, IndividuallyCompressedShader.Size);
 			}
 			else
 			{
 				// This function will crash if decompression fails.
-				ShaderCodeArchive::DecompressShaderWithOodle(Destination, IndividuallyCompressedShader.UncompressedSize, GetShaderCodeDataStart() + IndividuallyCompressedShader.Offset, IndividuallyCompressedShader.Size);
+				ShaderCodeArchive::DecompressShaderWithOodle(Destination, UncompressedSize, GetShaderCodeDataStart() + IndividuallyCompressedShader.Offset, IndividuallyCompressedShader.Size);
 			}
 		}
-		return Destination + IndividuallyCompressedShader.UncompressedSize;
+		return Destination + UncompressedSize;
 	}
 
 	// Returns a pointer to the start of shader code payload to read from.
@@ -344,7 +345,7 @@ static bool ConvertToIoStoreShaderLibrary(
 			const FShaderCodeEntry& IndividuallyCompressedShader = SerializedShaders.ShaderEntries[ShaderIndex];
 
 			TotalIndividuallyCompressedSize += IndividuallyCompressedShader.Size;
-			TotalUncompressedSizeViaShaders += IndividuallyCompressedShader.UncompressedSize;
+			TotalUncompressedSizeViaShaders += IndividuallyCompressedShader.GetUncompressedSize();
 		}
 	}
 
@@ -664,4 +665,3 @@ void FIoStoreShaderLibraryProcessor::AddAllReferencedShadersToContainerTarget(FC
 }
 
 } // namespace UE::IoStore::Private
-
diff --git a/Engine/Source/Editor/UnrealEd/Private/Commandlets/ShaderPipelineCacheToolsCommandlet.cpp b/Engine/Source/Editor/UnrealEd/Private/Commandlets/ShaderPipelineCacheToolsCommandlet.cpp
--- a/Engine/Source/Editor/UnrealEd/Private/Commandlets/ShaderPipelineCacheToolsCommandlet.cpp
+++ b/Engine/Source/Editor/UnrealEd/Private/Commandlets/ShaderPipelineCacheToolsCommandlet.cpp
@@ -752,15 +752,16 @@ int32 ExtractShaderFromUshaderbytecode(const FString& UshaderbytecodeName, const
 				if (PtrToShader)
 				{
 					// small shaders might be stored without compression, handle them here
-					if (Entry.Size == Entry.UncompressedSize)
+					if (!Entry.IsCompressed())
 					{
 						FileWriter->Serialize(PtrToShader, Entry.Size);
 					}
 					else
 					{
-						TUniquePtr<uint8[]> Buffer = MakeUnique<uint8[]>(Entry.UncompressedSize);
-						ShaderCodeArchive::DecompressShaderWithOodle(Buffer.Get(), Entry.UncompressedSize, PtrToShader, Entry.Size);
-						FileWriter->Serialize(Buffer.Get(), Entry.UncompressedSize);
+						uint32 UncompressedSize = Entry.GetUncompressedSize();
+						TUniquePtr<uint8[]> Buffer = MakeUnique<uint8[]>(UncompressedSize);
+						ShaderCodeArchive::DecompressShaderWithOodle(Buffer.Get(), UncompressedSize, PtrToShader, Entry.Size);
+						FileWriter->Serialize(Buffer.Get(), UncompressedSize);
 					}
 				}
 			}
diff --git a/Engine/Source/Runtime/Apple/MetalRHI/Private/MetalDynamicRHI_Shaders.cpp b/Engine/Source/Runtime/Apple/MetalRHI/Private/MetalDynamicRHI_Shaders.cpp
--- a/Engine/Source/Runtime/Apple/MetalRHI/Private/MetalDynamicRHI_Shaders.cpp
+++ b/Engine/Source/Runtime/Apple/MetalRHI/Private/MetalDynamicRHI_Shaders.cpp
@@ -202,8 +202,6 @@ FRHIShaderLibraryRef FMetalDynamicRHI::RHICreateShaderLibrary(EShaderPlatform Pl
         // Would be good to check the language version of the library with the archive format here.
         if (Header.Format == ShaderFormatName.GetPlainNameString())
         {
-            check(((SerializedShaders.GetNumShaders() + Header.NumShadersPerLibrary - 1) / Header.NumShadersPerLibrary) == Header.NumLibraries);
-
 			TArray<TUniquePtr<FMetalShaderLibrary::FLazyMetalLib>> LazyLibraries;
 			LazyLibraries.Empty(Header.NumLibraries);
 
diff --git a/Engine/Source/Runtime/Apple/MetalRHI/Private/Shaders/MetalShaderLibrary.cpp b/Engine/Source/Runtime/Apple/MetalRHI/Private/Shaders/MetalShaderLibrary.cpp
--- a/Engine/Source/Runtime/Apple/MetalRHI/Private/Shaders/MetalShaderLibrary.cpp
+++ b/Engine/Source/Runtime/Apple/MetalRHI/Private/Shaders/MetalShaderLibrary.cpp
@@ -148,10 +148,10 @@ TRefCountPtr<FRHIShader> FMetalShaderLibrary::CreateShader(int32 Index, bool bRe
 	const FShaderCodeEntry& ShaderEntry = SerializedShaders.GetShaderEntries()[Index];
 
 	// We don't handle compressed shaders here, since typically these are just tiny headers.
-	check(ShaderEntry.Size == ShaderEntry.UncompressedSize);
+	check(!ShaderEntry.IsCompressed());
 
 	const TArrayView<const uint8> Code = MakeArrayView(ShaderCode.GetData() + ShaderEntry.Offset, ShaderEntry.Size);
-	const int32 LibraryIndex = Index / Header.NumShadersPerLibrary;
+	const int32 LibraryIndex = ShaderEntry.GetLibraryIndex() / Header.NumShadersPerLibrary;
 
 	{
 		TUniquePtr<FLazyMetalLib>& LazyLibrary = LazyLibraries[LibraryIndex];
@@ -198,7 +198,7 @@ TRefCountPtr<FRHIShader> FMetalShaderLibrary::CreateShader(int32 Index, bool bRe
 	check(LazyLibraries[LibraryIndex]->Library.get() != nullptr);
 	
 	TRefCountPtr<FRHIShader> Shader;
-	switch (ShaderEntry.Frequency)
+	switch (ShaderEntry.GetFrequency())
 	{
 		case SF_Vertex:
 			Shader = CreateMetalShader<FMetalVertexShader>(Device, Code, LazyLibraries[LibraryIndex]->Library);
diff --git a/Engine/Source/Runtime/RenderCore/Private/ShaderLibrary/EditorShaderCodeArchive.cpp b/Engine/Source/Runtime/RenderCore/Private/ShaderLibrary/EditorShaderCodeArchive.cpp
--- a/Engine/Source/Runtime/RenderCore/Private/ShaderLibrary/EditorShaderCodeArchive.cpp
+++ b/Engine/Source/Runtime/RenderCore/Private/ShaderLibrary/EditorShaderCodeArchive.cpp
@@ -101,10 +101,10 @@ int32 FEditorShaderCodeArchive::AddShaderCode(FShaderCodeStats& CodeStats, const
 			{
 				const FShaderCodeResource& SourceShaderResource = Code->ShaderCodeResources[i];
 				FShaderCodeEntry& SerializedShaderEntry = SerializedShaders.ShaderEntries[ShaderIndex];
-				SerializedShaderEntry.Frequency = SourceShaderResource.GetFrequency();
+				SerializedShaderEntry.SetFrequency(SourceShaderResource.GetFrequency());
 				FSharedBuffer CodeBuffer = SourceShaderResource.GetCodeBuffer();
 				SerializedShaderEntry.Size = CodeBuffer.GetSize();
-				SerializedShaderEntry.UncompressedSize = SourceShaderResource.GetUncompressedSize();
+				SerializedShaderEntry.SetUncompressedSize(SourceShaderResource.GetUncompressedSize());
 				check(SerializedShaderEntry.Size > 0);
 				ShaderCode.Add(CodeBuffer);
 				check(ShaderCode.Num() == SerializedShaders.ShaderEntries.Num());
diff --git a/Engine/Source/Runtime/RenderCore/Private/ShaderLibrary/ShaderCodeArchive.cpp b/Engine/Source/Runtime/RenderCore/Private/ShaderLibrary/ShaderCodeArchive.cpp
--- a/Engine/Source/Runtime/RenderCore/Private/ShaderLibrary/ShaderCodeArchive.cpp
+++ b/Engine/Source/Runtime/RenderCore/Private/ShaderLibrary/ShaderCodeArchive.cpp
@@ -639,14 +639,16 @@ bool ShaderCodeArchive::CompressShaderWithOodle(uint8* OutCompressedShader, int6
 void FSerializedShaderArchive::DecompressShader(int32 Index, const TArray<FSharedBuffer>& ShaderCode, TArray<uint8>& OutDecompressedShader) const
 {
 	const FShaderCodeEntry& Entry = ShaderEntries[Index];
-	OutDecompressedShader.SetNum(Entry.UncompressedSize, EAllowShrinking::No);
-	if (Entry.Size == Entry.UncompressedSize)
+	if (!Entry.IsCompressed())
 	{
-		FMemory::Memcpy(OutDecompressedShader.GetData(), ShaderCode[Index].GetData(), Entry.UncompressedSize);
+		OutDecompressedShader.SetNum(Entry.Size, EAllowShrinking::No);
+		FMemory::Memcpy(OutDecompressedShader.GetData(), ShaderCode[Index].GetData(), Entry.Size);
 	}
 	else
 	{
-		ShaderCodeArchive::DecompressShaderWithOodle(OutDecompressedShader.GetData(), Entry.UncompressedSize, reinterpret_cast<const uint8*>(ShaderCode[Index].GetData()), Entry.Size);
+		uint32 UncompressedSize = Entry.GetUncompressedSize();
+		OutDecompressedShader.SetNum(UncompressedSize, EAllowShrinking::No);
+		ShaderCodeArchive::DecompressShaderWithOodle(OutDecompressedShader.GetData(), UncompressedSize, reinterpret_cast<const uint8*>(ShaderCode[Index].GetData()), Entry.Size);
 	}
 }
 
@@ -1146,12 +1148,13 @@ void FSerializedShaderArchive::CollectStatsAndDebugInfo(FDebugStats& OutDebugSta
 		FMemory::Memzero(OutExtendedDebugStats->CompressedSizePerFrequency);
 		for (const FShaderCodeEntry& ShaderEntry : ShaderEntries)
 		{
-			check(ShaderEntry.Frequency < UE_ARRAY_COUNT(OutExtendedDebugStats->NumShadersPerFrequency));
-			++OutExtendedDebugStats->NumShadersPerFrequency[ShaderEntry.Frequency];
-			check(ShaderEntry.Frequency < UE_ARRAY_COUNT(OutExtendedDebugStats->UncompressedSizePerFrequency));
-			OutExtendedDebugStats->UncompressedSizePerFrequency[ShaderEntry.Frequency] += ShaderEntry.UncompressedSize;
-			check(ShaderEntry.Frequency < UE_ARRAY_COUNT(OutExtendedDebugStats->UncompressedSizePerFrequency));
-			OutExtendedDebugStats->CompressedSizePerFrequency[ShaderEntry.Frequency] += ShaderEntry.Size;
+			EShaderFrequency Frequency = ShaderEntry.GetFrequency();
+			check(Frequency < UE_ARRAY_COUNT(OutExtendedDebugStats->NumShadersPerFrequency));
+			++OutExtendedDebugStats->NumShadersPerFrequency[Frequency];
+			check(Frequency < UE_ARRAY_COUNT(OutExtendedDebugStats->UncompressedSizePerFrequency));
+			OutExtendedDebugStats->UncompressedSizePerFrequency[Frequency] += ShaderEntry.GetUncompressedSize();
+			check(Frequency < UE_ARRAY_COUNT(OutExtendedDebugStats->UncompressedSizePerFrequency));
+			OutExtendedDebugStats->CompressedSizePerFrequency[Frequency] += ShaderEntry.Size;
 		}
 	}
 
@@ -1257,11 +1260,11 @@ void FSerializedShaderArchive::DumpContentsInPlaintext(FString& OutText) const
 			Out << TEXT("\n");
 
 			Out << TEXT("\t\t\tUncompressedSize : ");
-			Out << ShaderEntries[IdxEntry].UncompressedSize;
+			Out << ShaderEntries[IdxEntry].GetUncompressedSize();
 			Out << TEXT("\n");
 
 			Out << TEXT("\t\t\tFrequency : ");
-			Out << ShaderEntries[IdxEntry].Frequency;
+			Out << (uint8)ShaderEntries[IdxEntry].GetFrequency();
 			Out << TEXT("\n");
 
 			Out << TEXT("\t\t}\n");
@@ -1488,7 +1491,7 @@ bool FShaderCodeArchive::PreloadShaderMap(int32 ShaderMapIndex, FGraphEventArray
 		const FShaderCodeEntry& ShaderEntry = SerializedShaders.GetShaderEntries()[ShaderIndex];
 
 #if RHI_RAYTRACING
-		if (!IsRayTracingAllowed() && !IsCreateShadersOnLoadEnabled() && IsRayTracingShaderFrequency(static_cast<EShaderFrequency>(ShaderEntry.Frequency)))
+		if (!IsRayTracingAllowed() && !IsCreateShadersOnLoadEnabled() && IsRayTracingShaderFrequency(static_cast<EShaderFrequency>(ShaderEntry.GetFrequency())))
 		{
 			ShaderPreloadEntry.bNeverToBePreloaded = 1;
 			continue;
@@ -1613,19 +1616,20 @@ TRefCountPtr<FRHIShader> FShaderCodeArchive::CreateShader(int32 Index, bool bReq
 	}
 
 	const uint8* ShaderCode = (uint8*)ShaderPreloadEntry.Code;
-	if (ShaderEntry.UncompressedSize != ShaderEntry.Size)
+	if (ShaderEntry.IsCompressed())
 	{
-		uint8* UncompressedCode = reinterpret_cast<uint8*>(MemStack.Alloc(ShaderEntry.UncompressedSize, 16));
-		ShaderCodeArchive::DecompressShaderWithOodle(UncompressedCode, ShaderEntry.UncompressedSize, ShaderCode, ShaderEntry.Size);
+		uint32 UncompressedSize = ShaderEntry.GetUncompressedSize();
+		uint8* UncompressedCode = reinterpret_cast<uint8*>(MemStack.Alloc(UncompressedSize, 16));
+		ShaderCodeArchive::DecompressShaderWithOodle(UncompressedCode, UncompressedSize, ShaderCode, ShaderEntry.Size);
 		ShaderCode = (uint8*)UncompressedCode;
 	}
 
 	// detect the breach of contract early
 	ensureAlwaysMsgf(IsInRenderingThread() || GRHISupportsMultithreadedShaderCreation, TEXT("More than one thread is creating shaders, but GRHISupportsMultithreadedShaderCreation is false."));
 
-	const auto ShaderCodeView = MakeArrayView(ShaderCode, ShaderEntry.UncompressedSize);
+	const auto ShaderCodeView = MakeArrayView(ShaderCode, ShaderEntry.GetUncompressedSize());
 	const FSHAHash& ShaderHash = SerializedShaders.GetShaderHashes()[Index];
-	switch (ShaderEntry.Frequency)
+	switch (ShaderEntry.GetFrequency())
 	{
 	case SF_Vertex: Shader = RHICreateVertexShader(ShaderCodeView, ShaderHash); CheckShaderCreation(Shader, Index); break;
 	case SF_Mesh: Shader = RHICreateMeshShader(ShaderCodeView, ShaderHash); CheckShaderCreation(Shader, Index); break;
@@ -1826,7 +1830,7 @@ void FShaderCodeArchiveProcessor::ProcessSerializedShaders(const FSerializedShad
 	for (const FShaderCodeEntry& ShaderEntry : SerializedShaders->GetShaderEntries())
 	{
 		FIoStoreShaderCodeEntry& IoStoreShaderEntry = OutHeader.ShaderEntries.AddDefaulted_GetRef();
-		IoStoreShaderEntry.Frequency = ShaderEntry.Frequency;
+		IoStoreShaderEntry.Frequency = ShaderEntry.GetFrequency();
 	}
 
 	BuildShaderToShadermapsArray();
@@ -2020,7 +2024,7 @@ void FShaderCodeArchiveProcessor::SplitShaderGroupBySize(TArray<uint32>& Current
 	const TArrayView<const FShaderCodeEntry> ShaderEntries = SerializedShaders->GetShaderEntries();
 	for (uint32 ShaderIdx : CurrentShaderGroup)
 	{
-		GroupSize += ShaderEntries[ShaderIdx].UncompressedSize;
+		GroupSize += ShaderEntries[ShaderIdx].GetUncompressedSize();
 	}
 
 	if (LIKELY(GroupSize <= MaxUncompressedShaderGroupSize || CurrentShaderGroup.Num() == 1))
@@ -2045,17 +2049,21 @@ void FShaderCodeArchiveProcessor::SplitShaderGroupBySize(TArray<uint32>& Current
 				const TArrayView<const FShaderCodeEntry> ShaderEntries = SerializedShaders->GetShaderEntries();
 				const FShaderCodeEntry& ShaderEntryA = ShaderEntries[ShaderIndexA];
 				const FShaderCodeEntry& ShaderEntryB = ShaderEntries[ShaderIndexB];
-				if (ShaderEntryA.UncompressedSize != ShaderEntryB.UncompressedSize)
+				uint32 UncompressedSizeA = ShaderEntryA.GetUncompressedSize();
+				uint32 UncompressedSizeB = ShaderEntryB.GetUncompressedSize();
+				if (UncompressedSizeA != UncompressedSizeB)
 				{
-					return ShaderEntryA.UncompressedSize > ShaderEntryB.UncompressedSize;
+					return UncompressedSizeA > UncompressedSizeB;
 				}
 				if (ShaderEntryA.Size != ShaderEntryB.Size)
 				{
 					return ShaderEntryA.Size > ShaderEntryB.Size;
 				}
-				if (ShaderEntryA.Frequency != ShaderEntryB.Frequency)
+				EShaderFrequency FrequencyA = ShaderEntryA.GetFrequency();
+				EShaderFrequency FrequencyB = ShaderEntryB.GetFrequency();
+				if (FrequencyA != FrequencyB)
 				{
-					return ShaderEntryA.Frequency > ShaderEntryB.Frequency;
+					return FrequencyA > FrequencyB;
 				}
 				return ShaderEntryA.Offset > ShaderEntryB.Offset;
 			}
@@ -2074,7 +2082,7 @@ void FShaderCodeArchiveProcessor::SplitShaderGroupBySize(TArray<uint32>& Current
 			}
 
 			NewGroups[SmallestNewGroupIdx].Add(ShaderIdx);
-			NewGroupSizes[SmallestNewGroupIdx] += ShaderEntries[ShaderIdx].UncompressedSize;
+			NewGroupSizes[SmallestNewGroupIdx] += ShaderEntries[ShaderIdx].GetUncompressedSize();
 		}
 
 #if DO_CHECK // sanity checks
@@ -2110,17 +2118,21 @@ void FShaderCodeArchiveProcessor::AddNewShaderGroup(TArray<uint32>& ShaderIndice
 		{
 			const FShaderCodeEntry& ShaderEntryA = SerializedShaders->GetShaderEntries()[ShaderIndexA];
 			const FShaderCodeEntry& ShaderEntryB = SerializedShaders->GetShaderEntries()[ShaderIndexB];
-			if (ShaderEntryA.UncompressedSize != ShaderEntryB.UncompressedSize)
+			uint32 UncompressedSizeA = ShaderEntryA.GetUncompressedSize();
+			uint32 UncompressedSizeB = ShaderEntryB.GetUncompressedSize();
+			if (UncompressedSizeA != UncompressedSizeB)
 			{
-				return ShaderEntryA.UncompressedSize < ShaderEntryB.UncompressedSize;
+				return UncompressedSizeA < UncompressedSizeB;
 			}
 			if (ShaderEntryA.Size != ShaderEntryB.Size)
 			{
 				return ShaderEntryA.Size < ShaderEntryB.Size;
 			}
-			if (ShaderEntryA.Frequency != ShaderEntryB.Frequency)
+			EShaderFrequency FrequencyA = ShaderEntryA.GetFrequency();
+			EShaderFrequency FrequencyB = ShaderEntryB.GetFrequency();
+			if (FrequencyA != FrequencyB)
 			{
-				return ShaderEntryA.Frequency < ShaderEntryB.Frequency;
+				return FrequencyA < FrequencyB;
 			}
 			return ShaderEntryA.Offset < ShaderEntryB.Offset;
 		}
@@ -2149,9 +2161,10 @@ void FShaderCodeArchiveProcessor::AddNewShaderGroup(TArray<uint32>& ShaderIndice
 		// group hash is constructed from hashing the shaders in the group.
 		GroupHasher.Update(ShaderHashes[ShaderIndex].Hash, sizeof(FSHAHash));
 		// shader hash as of now excludes optional data, so we cannot rely on it, especially across the shader formats. Make the group hash a bit more robust by including the shader size in it.
-		GroupHasher.Update(reinterpret_cast<const uint8*>(&ShaderEntries[ShaderIndex].UncompressedSize), sizeof(FShaderCodeEntry::UncompressedSize));
+		uint32 UncompressedSize = ShaderEntries[ShaderIndex].GetUncompressedSize();
+		GroupHasher.Update(reinterpret_cast<const uint8*>(&UncompressedSize), sizeof(UncompressedSize));
 
-		CurrentGroupSize += ShaderEntries[ShaderIndex].UncompressedSize;
+		CurrentGroupSize += UncompressedSize;
 	}
 
 	// Shader hashes cannot be used to uniquely identify across shader formats due to aforementioned exclusion of optional data from it.
@@ -2785,4 +2798,3 @@ TRefCountPtr<FRHIShader> FIoStoreShaderCodeArchive::CreateShader(int32 ShaderInd
 
 	return Shader;
 }
-
diff --git a/Engine/Source/Runtime/RenderCore/Public/ShaderCodeArchive.h b/Engine/Source/Runtime/RenderCore/Public/ShaderCodeArchive.h
--- a/Engine/Source/Runtime/RenderCore/Public/ShaderCodeArchive.h
+++ b/Engine/Source/Runtime/RenderCore/Public/ShaderCodeArchive.h
@@ -69,19 +69,71 @@ static FArchive& operator <<(FArchive& Ar, FFileCachePreloadEntry& Ref)
 #endif
 struct FShaderCodeEntry
 {
-	uint64 Offset = 0;
-	uint32 Size = 0;
-	uint32 UncompressedSize = 0;
-	uint8 Frequency;
-
+	// If the corresponding shader is part of the library,
+	// it is stored uncompressed and assigned an index within
+	// the library.
+	static constexpr uint8  HasLibraryIndexTag   = 1u << SF_NumBits;
+	static constexpr uint32 FrequencyPayloadMask = HasLibraryIndexTag - 1u;
+
+	uint64 Offset   = 0;
+	uint32 Size     = 0;
+	
+	// The active member is specified by the tag stored
+	// in the TagAndFrequency field. Rely on access methods.
+	union
+	{
+		uint32 UncompressedSize_;
+		uint32 LibraryIndex_;
+	};
+	uint8 TagAndFrequency = 0;
+	
 	EShaderFrequency GetFrequency() const
 	{
-		return (EShaderFrequency)Frequency;
+		return (EShaderFrequency)(TagAndFrequency & FrequencyPayloadMask);
 	}
-
+	
+	void SetFrequency(EShaderFrequency Frequency)
+	{
+		check(Frequency <= FrequencyPayloadMask);
+		TagAndFrequency = (TagAndFrequency & HasLibraryIndexTag) | Frequency;
+	}
+	
+	uint32 GetUncompressedSize() const
+	{
+		return IsPartOfLibrary() ? Size : UncompressedSize_;
+	}
+	
+	void SetUncompressedSize(uint32 InUncompressedSize)
+	{
+		check(!IsPartOfLibrary());
+		UncompressedSize_ = InUncompressedSize;
+	}
+	
+	int32 GetLibraryIndex() const
+	{
+		return IsPartOfLibrary() ? LibraryIndex_ : INDEX_NONE;
+	}
+	
+	void SetLibraryIndex(int32 Index)
+	{
+		check(Index >= 0);
+		LibraryIndex_     = Index;
+		TagAndFrequency |= HasLibraryIndexTag;
+	}
+	
+	bool IsCompressed() const
+	{
+		return !IsPartOfLibrary() && Size != GetUncompressedSize();
+	}
+	
+	bool IsPartOfLibrary() const
+	{
+		return (TagAndFrequency & HasLibraryIndexTag) != 0;
+	}
+	
 	friend FArchive& operator <<(FArchive& Ar, FShaderCodeEntry& Ref)
 	{
-		return Ar << Ref.Offset << Ref.Size << Ref.UncompressedSize << Ref.Frequency;
+		return Ar << Ref.Offset << Ref.Size << (Ref.IsPartOfLibrary() ? Ref.LibraryIndex_ : Ref.UncompressedSize_) << Ref.TagAndFrequency;
 	}
 };
 #if USE_MMAPPED_SHADERARCHIVE
@@ -1002,4 +1054,4 @@ inline void FSerializedShaderArchive::FShaderMapAssetAssociations::FAssociatedAs
 	LatestSource = (ECookShaderLibrarySource)FMath::Min((uint8)LatestSource, (uint8)InCookSource);
 }
 
-#endif // WITH_EDITOR
\ No newline at end of file
+#endif // WITH_EDITOR
