Shader Parameter Parser changes
- Moving FParsedShaderParameter bools into a Flags type.
- Renaming a few variables for clarity on what they are for.
- Adding new flag to mark a parameter as "global" for the ReplaceGlobals path.
- Adding ConstantBuffer detection so they are not treated the same as Loose parameters.

#Synced-CL 49859967

diff --git a/Engine/Source/Developer/ShaderCompilerCommon/Private/ShaderCompilerCommon.cpp b/Engine/Source/Developer/ShaderCompilerCommon/Private/ShaderCompilerCommon.cpp
--- a/Engine/Source/Developer/ShaderCompilerCommon/Private/ShaderCompilerCommon.cpp
+++ b/Engine/Source/Developer/ShaderCompilerCommon/Private/ShaderCompilerCommon.cpp
@@ -908,7 +908,7 @@ void AddUnboundShaderParameterError(
 	check(CompilerInput.RootParametersStructure);
 
 	const FShaderParameterParser::FParsedShaderParameter& Member = ShaderParameterParser.FindParameterInfos(ParameterBindingName);
-	check(!Member.bIsBindable);
+	check(!EnumHasAnyFlags(Member.Flags, EParsedShaderParameterFlags::WasInRootParameters));
 
 	FShaderCompilerError Error(FString::Printf(
 		TEXT("Error: Shader parameter %.*s could not be bound to %s's shader parameter structure %s."),
diff --git a/Engine/Source/Runtime/RenderCore/Private/ShaderParameterParser.cpp b/Engine/Source/Runtime/RenderCore/Private/ShaderParameterParser.cpp
--- a/Engine/Source/Runtime/RenderCore/Private/ShaderParameterParser.cpp
+++ b/Engine/Source/Runtime/RenderCore/Private/ShaderParameterParser.cpp
@@ -208,6 +208,11 @@ EShaderParameterType FShaderParameterParser::ParseParameterType(FStringView InTy
 		return EShaderParameterType::UAV;
 	}
 
+	if (UntemplatedType == TEXTVIEW("ConstantBuffer"))
+	{
+		return EShaderParameterType::UniformBuffer;
+	}
+
 	return EShaderParameterType::LooseData;
 }
 
@@ -261,6 +266,14 @@ FStringView FShaderParameterParser::GetBindlessParameterPrefix(EShaderParameterT
 	return FStringView();
 }
 
+static bool IsBindlessHeap(FStringView VariableName)
+{
+	return VariableName.StartsWith(FShaderParameterParser::kBindlessSRVArrayPrefix)
+		|| VariableName.StartsWith(FShaderParameterParser::kBindlessUAVArrayPrefix)
+		|| VariableName.StartsWith(FShaderParameterParser::kBindlessCBVArrayPrefix)
+		|| VariableName.StartsWith(FShaderParameterParser::kBindlessSamplerArrayPrefix);
+}
+
 bool FShaderParameterParser::ParseParameters(
 	const FShaderParametersMetadata* RootParametersStructure,
 	TArray<FShaderCompilerError>& OutErrors)
@@ -280,7 +293,7 @@ bool FShaderParameterParser::ParseParameters(
 				uint16 ByteOffset)
 			{
 				FParsedShaderParameter ParsedParameter;
-				ParsedParameter.bIsBindable = true;
+				EnumAddFlags(ParsedParameter.Flags, EParsedShaderParameterFlags::WasInRootParameters);
 				ParsedParameter.ConstantBufferOffset = ByteOffset;
 				ParsedParameter.BaseType = Member.GetBaseType();
 				ParsedParameter.PrecisionModifier = Member.GetPrecision();
@@ -380,21 +393,19 @@ bool FShaderParameterParser::ParseParameters(
 				check(NameStartPos != -1);
 				check(NameEndPos != -1);
 
-				FStringView Type = ShaderSource.Mid(TypeStartPos, TypeEndPos - TypeStartPos + 1);
+				const FStringView Type = ShaderSource.Mid(TypeStartPos, TypeEndPos - TypeStartPos + 1);
 				FStringView Name = ShaderSource.Mid(NameStartPos, NameEndPos - NameStartPos + 1);
 
 				FStringView Leftovers = ShaderSource.Mid(NameEndPos + 1, (Cursor - 1) - (NameEndPos + 1) + 1);
 
-				EShaderParameterType ParsedParameterType = ParseAndRemoveBindlessParameterPrefix(Name);
-				const bool bBindlessIndex = (ParsedParameterType != EShaderParameterType::LooseData);
+				const EShaderParameterType ParsedBindlessType = ParseAndRemoveBindlessParameterPrefix(Name);
+				const EShaderParameterType ParsedParameterType = ParseParameterType(Type);
 
 				EBindlessConversionType BindlessConversionType = EBindlessConversionType::None;
-				EShaderParameterType ParsedConstantBufferType = ParsedParameterType;
+				EShaderParameterType ParsedConstantBufferType = ParsedBindlessType;
 
-				if (bBindlessEnabled && ParsedParameterType == EShaderParameterType::LooseData)
+				if (bBindlessEnabled && ParsedBindlessType == EShaderParameterType::LooseData)
 				{
-					ParsedParameterType = ParseParameterType(Type);
-
 					if (bBindlessEnabled)
 					{
 						if (ParsedParameterType == EShaderParameterType::SRV)
@@ -428,8 +439,7 @@ bool FShaderParameterParser::ParseParameters(
 
 				FParsedShaderParameter ParsedParameter;
 
-				EShaderParameterType ConstantBufferParameterType = EShaderParameterType::Num;
-				bool bMoveToRootConstantBuffer = false;
+				EShaderParameterType RootCBParameterType = EShaderParameterType::Num;
 				bool bUpdateParsedParameters = false;
 
  				const FString ParsedParameterKey(Name);
@@ -446,19 +456,18 @@ bool FShaderParameterParser::ParseParameters(
 						ParsedParameter = ParsedParameters.FindChecked(ParsedParameterKey);
 
 						// Erase the parameter to move it into the root constant buffer.
-						if (bNeedToMoveToRootConstantBuffer && ParsedParameter.bIsBindable)
+						if (bMoveGlobalsToRootCB && EnumHasAnyFlags(ParsedParameter.Flags, EParsedShaderParameterFlags::WasInRootParameters))
 						{
-							const EUniformBufferBaseType BaseType = ParsedParameter.BaseType;
-							bMoveToRootConstantBuffer =
-								BaseType == UBMT_INT32 ||
-								BaseType == UBMT_UINT32 ||
-								BaseType == UBMT_FLOAT32 ||
-								bBindlessIndex ||
-								(BindlessConversionType != EBindlessConversionType::None);
-
-							if (bMoveToRootConstantBuffer)
+							const bool bMoveToRootCB =
+								ParsedParameter.BaseType == UBMT_INT32 ||
+								ParsedParameter.BaseType == UBMT_UINT32 ||
+								ParsedParameter.BaseType == UBMT_FLOAT32 ||
+								ParsedBindlessType != EShaderParameterType::LooseData ||
+								BindlessConversionType != EBindlessConversionType::None;
+
+							if (bMoveToRootCB)
 							{
-								ConstantBufferParameterType = ParsedConstantBufferType;
+								RootCBParameterType = ParsedConstantBufferType;
 							}
 						}
 					}
@@ -479,8 +488,7 @@ bool FShaderParameterParser::ParseParameters(
 					ParsedParameter.ParsedCharOffsetStart = TypeQualifierStartPos != -1 ? TypeQualifierStartPos : TypeStartPos;
 					ParsedParameter.ParsedCharOffsetEnd = Cursor;
 					ParsedParameter.BindlessConversionType = BindlessConversionType;
-					ParsedParameter.ConstantBufferParameterType = ConstantBufferParameterType;
-					ParsedParameter.bGloballyCoherent = bGloballyCoherent;
+					ParsedParameter.RootCBParameterType = RootCBParameterType;
 					ParsedParameter.ParsedTypeDecl = TypeDecl;
 
 					if (ArrayStartPos != -1 && ArrayEndPos != -1)
@@ -488,6 +496,18 @@ bool FShaderParameterParser::ParseParameters(
 						ParsedParameter.ParsedArraySize = ShaderSource.Mid(ArrayStartPos + 1, ArrayEndPos - ArrayStartPos - 1);
 					}
 
+					if (bGloballyCoherent)
+					{
+						EnumAddFlags(ParsedParameter.Flags, EParsedShaderParameterFlags::GloballyCoherent);
+					}
+					if (!EnumHasAnyFlags(ParsedParameter.Flags, EParsedShaderParameterFlags::WasInRootParameters) && ParsedParameterType == EShaderParameterType::LooseData && ParsedBindlessType == EShaderParameterType::LooseData)
+					{
+						if (!EnumHasAnyFlags(PlatformConfiguration.Flags, EShaderParameterParserConfigurationFlags::BindlessUsesArrays) || !IsBindlessHeap(Name))
+						{
+							EnumAddFlags(ParsedParameter.Flags, EParsedShaderParameterFlags::IsGlobalParameter);
+						}
+					}
+
 					ParsedParameters.Add(ParsedParameterKey, ParsedParameter);
 				}
 
@@ -896,7 +916,7 @@ void FShaderParameterParser::RemoveMovingParametersFromSource(FString& Preproces
 		const FParsedShaderParameter& ParsedParameter = Itr.Value;
 
 		// If this parameter is going to be in the root constant buffer
-		if (ParsedParameter.ConstantBufferParameterType != EShaderParameterType::Num &&
+		if (ParsedParameter.RootCBParameterType != EShaderParameterType::Num &&
 			// but it's not being converted to bindless
 			ParsedParameter.BindlessConversionType == EBindlessConversionType::None &&
 			ParsedParameter.ParsedCharOffsetStart != INDEX_NONE)
@@ -941,7 +961,7 @@ FString FShaderParameterParser::GenerateBindlessParameterDeclaration(const FPars
 	const bool bIsUAV = (ParsedParameter.BindlessConversionType == EBindlessConversionType::UAV);
 	const FStringView Name = ParsedParameter.ParsedName;
 	const FStringView Type = ParsedParameter.ParsedType;
-	const TCHAR* StorageClass = ParsedParameter.bGloballyCoherent ? TEXT("globallycoherent ") : TEXT("");
+	const TCHAR* StorageClass = EnumHasAnyFlags(ParsedParameter.Flags, EParsedShaderParameterFlags::GloballyCoherent) ? TEXT("globallycoherent ") : TEXT("");
 	const FStringView IndexPrefix = ::GetBindlessParameterPrefix(ParsedParameter.BindlessConversionType);
 
 	TStringBuilder<64> IndexString;
@@ -950,7 +970,7 @@ FString FShaderParameterParser::GenerateBindlessParameterDeclaration(const FPars
 	TStringBuilder<512> Result;
 
 	// If we weren't going to be added to a root constant buffer, that means we need to declare our index before we declare our getter.
-	if (ParsedParameter.ConstantBufferParameterType == EShaderParameterType::Num)
+	if (ParsedParameter.RootCBParameterType == EShaderParameterType::Num)
 	{
 		// e.g. `uint BindlessResource_##Name;`
 		Result << TEXT("uint ") << IndexString << TEXT("; ");
@@ -1010,7 +1030,7 @@ void FShaderParameterParser::ApplyBindlessModifications(FString& PreprocessedSha
 
 		const bool bReplaceGlobals = EnumHasAnyFlags(PlatformConfiguration.Flags, EShaderParameterParserConfigurationFlags::ReplaceGlobals);
 
-		for (TPair<FString, FParsedShaderParameter>& Itr : ParsedParameters)
+		for (const TPair<FString, FParsedShaderParameter>& Itr : ParsedParameters)
 		{
 			const FParsedShaderParameter& ParsedParameter = Itr.Value;
 
@@ -1030,22 +1050,20 @@ void FShaderParameterParser::ApplyBindlessModifications(FString& PreprocessedSha
 			}
 			else if (bReplaceGlobals)
 			{
-				const bool IsGlobalParam = 
-					(ParsedParameter.BaseType == UBMT_INVALID) &&
-					!ParsedParameter.ParsedName.StartsWith(FShaderParameterParser::kBindlessSamplerArrayPrefix) &&
-					!ParsedParameter.ParsedName.StartsWith(FShaderParameterParser::kBindlessSRVArrayPrefix) &&
-					!ParsedParameter.ParsedName.StartsWith(FShaderParameterParser::kBindlessUAVArrayPrefix);
-
-				if (IsGlobalParam)
+				if (EnumHasAnyFlags(ParsedParameter.Flags, EParsedShaderParameterFlags::IsGlobalParameter))
 				{
 					FShaderCodeModifications Modif;
 					Modif.CharOffsetStart = ParsedParameter.ParsedCharOffsetStart;
 					Modif.CharOffsetEnd = ParsedParameter.ParsedCharOffsetEnd + 1;
 
 					const int32 NumChars = Modif.CharOffsetEnd - Modif.CharOffsetStart;
-					Modif.Replace = PlatformConfiguration.ReplaceGlobal(FStringView(&OriginalParsedShader[Modif.CharOffsetStart], NumChars), ParsedParameter.ParsedName);
+					const FStringView OriginalString(&OriginalParsedShader[Modif.CharOffsetStart], NumChars);
+					Modif.Replace = PlatformConfiguration.ReplaceGlobal(OriginalString, ParsedParameter.ParsedName);
 
-					Modifications.Add(Modif);
+					if (Modif.Replace != OriginalString)
+					{
+						Modifications.Add(Modif);
+					}
 				}
 			}
 		}
@@ -1119,7 +1137,7 @@ bool FShaderParameterParser::MoveShaderParametersToRootConstantBuffer(
 	bool bSuccess = true;
 
 	// Generate the root cbuffer content.
-	if (RootParametersStructure && bNeedToMoveToRootConstantBuffer)
+	if (RootParametersStructure && bMoveGlobalsToRootCB)
 	{
 		FStringBuilderBase ConstantBufferMembers;
 
@@ -1131,15 +1149,15 @@ bool FShaderParameterParser::MoveShaderParametersToRootConstantBuffer(
 				uint16 ByteOffset)
 		{
 			FParsedShaderParameter* ParsedParameter = ParsedParameters.Find(ShaderBindingName);
-			if (ParsedParameter && ParsedParameter->IsFound() && ParsedParameter->ConstantBufferParameterType != EShaderParameterType::Num)
+			if (ParsedParameter && ParsedParameter->IsFound() && ParsedParameter->RootCBParameterType != EShaderParameterType::Num)
 			{
 				const uint32 ConstantRegister = ByteOffset / 16;
 				const TCHAR* ConstantSwizzle = GetConstantSwizzle(ByteOffset);
 
 #define SVARG(N) N.Len(), N.GetData()
-				if (IsParameterBindless(ParsedParameter->ConstantBufferParameterType))
+				if (IsParameterBindless(ParsedParameter->RootCBParameterType))
 				{
-					const FStringView Prefix = GetBindlessParameterPrefix(ParsedParameter->ConstantBufferParameterType);
+					const FStringView Prefix = GetBindlessParameterPrefix(ParsedParameter->RootCBParameterType);
 
 					ConstantBufferMembers.Appendf(
 						TEXT("uint %.*s%.*s : packoffset(c%d%s);\n"),
@@ -1149,7 +1167,7 @@ bool FShaderParameterParser::MoveShaderParametersToRootConstantBuffer(
 						ConstantSwizzle
 					);
 				}
-				else if (ParsedParameter->ConstantBufferParameterType == EShaderParameterType::LooseData)
+				else if (ParsedParameter->RootCBParameterType == EShaderParameterType::LooseData)
 				{
 					if (!ParsedParameter->ParsedArraySize.IsEmpty())
 					{
@@ -1194,7 +1212,7 @@ bool FShaderParameterParser::MoveShaderParametersToRootConstantBuffer(
 			bMovedLoosedParametersToRootConstantBuffer = true;
 			bModifiedShader = true;
 		}
-	} // if (CompilerInput.RootParametersStructure && bNeedToMoveToRootConstantBuffer)
+	} // if (CompilerInput.RootParametersStructure && bMoveGlobalsToRootCB)
 
 	return bSuccess;
 }
@@ -1221,7 +1239,7 @@ bool FShaderParameterParser::ParseAndModify(const FShaderCompilerInput& Compiler
 		return true;
 	}
 
-	bNeedToMoveToRootConstantBuffer = bRootParametersModification;
+	bMoveGlobalsToRootCB = bRootParametersModification;
 	OriginalParsedShader = PreprocessedShaderSource;
 
 	if (!ParseParameters(CompilerInput.RootParametersStructure, OutErrors))
@@ -1240,7 +1258,7 @@ bool FShaderParameterParser::ParseAndModify(const FShaderCompilerInput& Compiler
 			ApplyBindlessModifications(PreprocessedShaderSource);
 		}
 
-		if (bNeedToMoveToRootConstantBuffer)
+		if (bMoveGlobalsToRootCB)
 		{
 			bResult = MoveShaderParametersToRootConstantBuffer(CompilerInput.RootParametersStructure, CompilerInput.Target.GetFrequency(), PreprocessedShaderSource);
 		}
diff --git a/Engine/Source/Runtime/RenderCore/Public/ShaderParameterParser.h b/Engine/Source/Runtime/RenderCore/Public/ShaderParameterParser.h
--- a/Engine/Source/Runtime/RenderCore/Public/ShaderParameterParser.h
+++ b/Engine/Source/Runtime/RenderCore/Public/ShaderParameterParser.h
@@ -36,6 +36,15 @@ enum class EShaderParameterParserConfigurationFlags
 };
 ENUM_CLASS_FLAGS(EShaderParameterParserConfigurationFlags)
 
+enum class EParsedShaderParameterFlags : uint8
+{
+	None                = 0,
+	GloballyCoherent    = 1 << 0,
+	WasInRootParameters = 1 << 1,
+	IsGlobalParameter   = 1 << 2,
+};
+ENUM_CLASS_FLAGS(EParsedShaderParameterFlags)
+
 /** Validates and moves all the shader loose data parameter defined in the root scope of the shader into the root uniform buffer. */
 class FShaderParameterParser
 {
@@ -84,28 +93,21 @@ public:
 
 	struct FParsedShaderParameter
 	{
-	public:
 		/** Original information about the member. */
-		TEnumAsByte<EUniformBufferBaseType> BaseType{ UBMT_INVALID };
-		TEnumAsByte<EShaderPrecisionModifier::Type> PrecisionModifier{ EShaderPrecisionModifier::Invalid };
+		EUniformBufferBaseType BaseType{ UBMT_INVALID };
+		EShaderPrecisionModifier::Type PrecisionModifier{ EShaderPrecisionModifier::Invalid };
 		uint32 NumRows = 0u;
 		uint32 NumColumns = 0u;
 		uint32 MemberSize = 0u;
 
+		/** Offset the member should be in the constant buffer. */
+		int32 ConstantBufferOffset = 0;
+
 		/** Information found about the member when parsing the preprocessed code. */
 		FStringView ParsedName; /** View into FShaderParameterParser::OriginalParsedShader */
 		FStringView ParsedType; /** View into FShaderParameterParser::OriginalParsedShader */
 		FStringView ParsedArraySize; /** View into FShaderParameterParser::OriginalParsedShader */
 
-		/** Offset the member should be in the constant buffer. */
-		int32 ConstantBufferOffset = 0;
-
-		/* Returns whether the shader parameter has been found when parsing. */
-		bool IsFound() const
-		{
-			return !ParsedType.IsEmpty();
-		}
-
 		int32 ParsedPragmaLineOffset = 0;
 		int32 ParsedLineOffset = 0;
 
@@ -113,15 +115,18 @@ public:
 		int32 ParsedCharOffsetStart = INDEX_NONE;
 		int32 ParsedCharOffsetEnd = INDEX_NONE;
 
-		EBindlessConversionType BindlessConversionType{};
-		EShaderParameterType ConstantBufferParameterType{};
-
-		bool bGloballyCoherent = false;
-		bool bIsBindable = false;
+		/** Only set when parameter needs to be moved to the root constant buffer. */
+		EShaderParameterType RootCBParameterType{};
 
+		EBindlessConversionType BindlessConversionType{};
 		EShaderCodeResourceBindingType ParsedTypeDecl = EShaderCodeResourceBindingType::Invalid;
+		EParsedShaderParameterFlags Flags{};
 
-		friend class FShaderParameterParser;
+		/* Returns whether the shader parameter has been found when parsing. */
+		bool IsFound() const
+		{
+			return !ParsedType.IsEmpty();
+		}
 	};
 
 	FShaderParameterParser() = delete;
@@ -245,8 +250,8 @@ protected:
 
 	bool bBindlessEnabled = false;
 
-	/** Indicates that parameters should be moved to the root cosntant buffer. */
-	bool bNeedToMoveToRootConstantBuffer = false;
+	/** Indicates that parameters should be moved to the root constant buffer. */
+	bool bMoveGlobalsToRootCB = false;
 
 	/** Indicates that parameters were actually moved to the root constant buffer. */
 	bool bMovedLoosedParametersToRootConstantBuffer = false;
