FShaderCodeEntry cleanup

#rnx

diff --git a/Engine/Source/Runtime/RenderCore/Public/ShaderCodeArchive.h b/Engine/Source/Runtime/RenderCore/Public/ShaderCodeArchive.h
--- a/Engine/Source/Runtime/RenderCore/Public/ShaderCodeArchive.h
+++ b/Engine/Source/Runtime/RenderCore/Public/ShaderCodeArchive.h
@@ -25,6 +25,8 @@
 #include "Templates/RefCounting.h"
 #include "UObject/NameTypes.h"
 
+#include <type_traits>
+
 #if WITH_EDITOR
 #include "Templates/Function.h"
 #endif
@@ -75,16 +77,19 @@ struct FShaderCodeEntry
 	static constexpr uint8  HasLibraryIndexTag   = 1u << SF_NumBits;
 	static constexpr uint32 FrequencyPayloadMask = HasLibraryIndexTag - 1u;
 
-	uint64 Offset   = 0;
-	uint32 Size     = 0;
+	uint64 Offset = 0;
+	uint32 Size   = 0;
 	
 	// The active member is specified by the tag stored
 	// in the TagAndFrequency field. Rely on access methods.
-	union
+	struct
 	{
-		uint32 UncompressedSize_;
-		uint32 LibraryIndex_;
-	};
+		union
+		{
+			uint32 UncompressedSize;
+			uint32 LibraryIndex;
+		};
+	} Payload {};	
 	uint8 TagAndFrequency = 0;
 	
 	EShaderFrequency GetFrequency() const
@@ -100,25 +105,25 @@ struct FShaderCodeEntry
 	
 	uint32 GetUncompressedSize() const
 	{
-		return IsPartOfLibrary() ? Size : UncompressedSize_;
+		return IsPartOfLibrary() ? Size : Payload.UncompressedSize;
 	}
 	
 	void SetUncompressedSize(uint32 InUncompressedSize)
 	{
 		check(!IsPartOfLibrary());
-		UncompressedSize_ = InUncompressedSize;
+		Payload.UncompressedSize = InUncompressedSize;
 	}
 	
 	int32 GetLibraryIndex() const
 	{
-		return IsPartOfLibrary() ? LibraryIndex_ : INDEX_NONE;
+		return IsPartOfLibrary() ? Payload.LibraryIndex : INDEX_NONE;
 	}
 	
 	void SetLibraryIndex(int32 Index)
 	{
 		check(Index >= 0);
-		LibraryIndex_     = Index;
-		TagAndFrequency |= HasLibraryIndexTag;
+		Payload.LibraryIndex  = Index;
+		TagAndFrequency      |= HasLibraryIndexTag;
 	}
 	
 	bool IsCompressed() const
@@ -133,12 +138,14 @@ struct FShaderCodeEntry
 	
 	friend FArchive& operator <<(FArchive& Ar, FShaderCodeEntry& Ref)
 	{
-		return Ar << Ref.Offset << Ref.Size << (Ref.IsPartOfLibrary() ? Ref.LibraryIndex_ : Ref.UncompressedSize_) << Ref.TagAndFrequency;
+		return Ar << Ref.Offset << Ref.Size << (Ref.IsPartOfLibrary() ? Ref.Payload.LibraryIndex : Ref.Payload.UncompressedSize) << Ref.TagAndFrequency;
 	}
 };
 #if USE_MMAPPED_SHADERARCHIVE
 #pragma pack(pop)
 #endif
+static_assert(std::is_standard_layout_v<FShaderCodeEntry>);
+static_assert(std::is_trivially_copyable_v<FShaderCodeEntry>);
 
 #if WITH_EDITOR
 namespace ShaderCodeArchive
