Shaders: Remove unused WorkingDirectory parameter from IShaderFormat::CompilePreprocessedShader

#Synced-CL 49858055

diff --git a/Engine/Plugins/FX/Niagara/Source/NiagaraEditor/Private/NiagaraCompiler.cpp b/Engine/Plugins/FX/Niagara/Source/NiagaraEditor/Private/NiagaraCompiler.cpp
--- a/Engine/Plugins/FX/Niagara/Source/NiagaraEditor/Private/NiagaraCompiler.cpp
+++ b/Engine/Plugins/FX/Niagara/Source/NiagaraEditor/Private/NiagaraCompiler.cpp
@@ -1706,7 +1706,7 @@ void FNiagaraEditorModule::TestCompileScriptFromConsole(const TArray<FString>& A
 
 			FVectorVMCompilationOutput CompilationOutput;
 			double StartTime = FPlatformTime::Seconds();
-			bool bSucceeded = TestCompileVectorVMShader(Input, FString(FPlatformProcess::ShaderDir()), CompilationOutput, GNiagaraSkipVectorVMBackendOptimizations != 0);
+			bool bSucceeded = TestCompileVectorVMShader(Input, CompilationOutput, GNiagaraSkipVectorVMBackendOptimizations != 0);
 			float DeltaTime = (float)(FPlatformTime::Seconds() - StartTime);
 
 			if (bSucceeded)
diff --git a/Engine/Source/Developer/Apple/MetalShaderFormat/Private/MetalShaderFormat.cpp b/Engine/Source/Developer/Apple/MetalShaderFormat/Private/MetalShaderFormat.cpp
--- a/Engine/Source/Developer/Apple/MetalShaderFormat/Private/MetalShaderFormat.cpp
+++ b/Engine/Source/Developer/Apple/MetalShaderFormat/Private/MetalShaderFormat.cpp
@@ -153,8 +153,7 @@ public:
 	virtual void CompilePreprocessedShader(
 		const FShaderCompilerInput& Input,
 		const FShaderPreprocessOutput& PreprocessOutput,
-		FShaderCompilerOutput& Output,
-		const FString& WorkingDirectory) const override final
+		FShaderCompilerOutput& Output) const override final
 	{
 		CheckShaderFormat(Input.ShaderFormat);
 		CompileMetalShader(Input, PreprocessOutput, Output);
diff --git a/Engine/Source/Developer/ShaderFormatOpenGL/Private/OpenGLShaderCompiler.cpp b/Engine/Source/Developer/ShaderFormatOpenGL/Private/OpenGLShaderCompiler.cpp
--- a/Engine/Source/Developer/ShaderFormatOpenGL/Private/OpenGLShaderCompiler.cpp
+++ b/Engine/Source/Developer/ShaderFormatOpenGL/Private/OpenGLShaderCompiler.cpp
@@ -3045,10 +3045,9 @@ static bool CompileToGlslWithShaderConductor(
  * @param Input - The input shader code and environment.
  * @param PreprocessOutput - The output of the shader preprocessing phase (see above)
  * @param Output - Contains shader compilation results upon return.
- * @param WorkingDirectory - (unused, part of IShaderFormat API)
  * @param Version - Target GLSL version for this compilation
  */
-void CompileOpenGLShader(const FShaderCompilerInput& Input, const FShaderPreprocessOutput& InPreprocessOutput, FShaderCompilerOutput& Output, const FString& WorkingDirectory, GLSLVersion Version)
+void CompileOpenGLShader(const FShaderCompilerInput& Input, const FShaderPreprocessOutput& InPreprocessOutput, FShaderCompilerOutput& Output, GLSLVersion Version)
 {
 	FString EntryPointName = Input.EntryPointName;
 	FString PreprocessedSource(InPreprocessOutput.GetSourceViewWide());
diff --git a/Engine/Source/Developer/ShaderFormatOpenGL/Private/ShaderFormatOpenGL.cpp b/Engine/Source/Developer/ShaderFormatOpenGL/Private/ShaderFormatOpenGL.cpp
--- a/Engine/Source/Developer/ShaderFormatOpenGL/Private/ShaderFormatOpenGL.cpp
+++ b/Engine/Source/Developer/ShaderFormatOpenGL/Private/ShaderFormatOpenGL.cpp
@@ -23,7 +23,6 @@ extern void CompileOpenGLShader(
 	const FShaderCompilerInput& Input,
 	const FShaderPreprocessOutput& InPreprocessOutput,
 	FShaderCompilerOutput& Output,
-	const FString& WorkingDirectory,
 	GLSLVersion Version);
 
 /** Version for shader format, this becomes part of the DDC key. */
@@ -116,10 +115,10 @@ public:
 		}
 	}
 
-	virtual void CompilePreprocessedShader(const FShaderCompilerInput& Input, const FShaderPreprocessOutput& PreprocessOutput, FShaderCompilerOutput& Output, const FString& WorkingDirectory) const override
+	virtual void CompilePreprocessedShader(const FShaderCompilerInput& Input, const FShaderPreprocessOutput& PreprocessOutput, FShaderCompilerOutput& Output) const override
 	{
 		CheckFormat(Input.ShaderFormat);
-		CompileOpenGLShader(Input, PreprocessOutput, Output, WorkingDirectory, TranslateFormatNameToEnum(Input.ShaderFormat));		
+		CompileOpenGLShader(Input, PreprocessOutput, Output, TranslateFormatNameToEnum(Input.ShaderFormat));
 	}
 
 	virtual const TCHAR* GetPlatformIncludeDirectory() const override
diff --git a/Engine/Source/Developer/ShaderFormatVectorVM/Internal/VectorVMTestCompile.h b/Engine/Source/Developer/ShaderFormatVectorVM/Internal/VectorVMTestCompile.h
--- a/Engine/Source/Developer/ShaderFormatVectorVM/Internal/VectorVMTestCompile.h
+++ b/Engine/Source/Developer/ShaderFormatVectorVM/Internal/VectorVMTestCompile.h
@@ -4,4 +4,4 @@
 
 #include "ShaderFormatVectorVM.h"
 
-extern SHADERFORMATVECTORVM_API bool TestCompileVectorVMShader(const FShaderCompilerInput& Input, const FString& WorkingDirectory, FVectorVMCompilationOutput& VMCompilationOutput, bool bSkipBackendOptimizations);
+extern SHADERFORMATVECTORVM_API bool TestCompileVectorVMShader(const FShaderCompilerInput& Input, FVectorVMCompilationOutput& VMCompilationOutput, bool bSkipBackendOptimizations);
diff --git a/Engine/Source/Developer/ShaderFormatVectorVM/Private/ShaderFormatVectorVM.cpp b/Engine/Source/Developer/ShaderFormatVectorVM/Private/ShaderFormatVectorVM.cpp
--- a/Engine/Source/Developer/ShaderFormatVectorVM/Private/ShaderFormatVectorVM.cpp
+++ b/Engine/Source/Developer/ShaderFormatVectorVM/Private/ShaderFormatVectorVM.cpp
@@ -17,8 +17,7 @@
 extern bool CompileVectorVMShader(
 	const FShaderCompilerInput& Input,
 	const FShaderPreprocessOutput& PreprocessOutput,
-	FShaderCompilerOutput& Output,
-	const FString& WorkingDirectory);
+	FShaderCompilerOutput& Output);
 
 static FName NAME_VVM_1_0(TEXT("VVM_1_0"));
 static const FGuid UE_SHADER_VECTORVM_VER = FGuid("51C3EC95-421E-43B1-8AA9-61B6D4EDC1BF");
@@ -76,10 +75,10 @@ public:
 		Input.Environment.CompilerFlags.Add(CFLAG_DisableSourceStripping);
 	}
 
-	virtual void CompilePreprocessedShader(const FShaderCompilerInput& Input, const FShaderPreprocessOutput& PreprocessOutput, FShaderCompilerOutput& Output, const FString& WorkingDirectory) const override
+	virtual void CompilePreprocessedShader(const FShaderCompilerInput& Input, const FShaderPreprocessOutput& PreprocessOutput, FShaderCompilerOutput& Output) const override
 	{
 		CheckFormat(Input.ShaderFormat);
-		CompileVectorVMShader(Input, PreprocessOutput, Output, WorkingDirectory);
+		CompileVectorVMShader(Input, PreprocessOutput, Output);
 	}
 
 	virtual const TCHAR* GetPlatformIncludeDirectory() const
diff --git a/Engine/Source/Developer/ShaderFormatVectorVM/Private/VectorVMShaderCompiler.cpp b/Engine/Source/Developer/ShaderFormatVectorVM/Private/VectorVMShaderCompiler.cpp
--- a/Engine/Source/Developer/ShaderFormatVectorVM/Private/VectorVMShaderCompiler.cpp
+++ b/Engine/Source/Developer/ShaderFormatVectorVM/Private/VectorVMShaderCompiler.cpp
@@ -30,7 +30,6 @@ bool CompileVectorVMShader(
 	const FShaderCompilerInput& Input,
 	const FShaderPreprocessOutput& PreprocessOutput,
 	FVectorVMCompilationOutput& VMCompilationOutput,
-	const FString& WorkingDirectory,
 	bool bSkipBackendOptimizations)
 {
 	SCOPE_CYCLE_COUNTER(STAT_VectorVM_Compiler_CompileShader_VectorVM);
@@ -117,10 +116,10 @@ bool CompileVectorVMShader(
 	return bResult;
 }
 
-bool CompileVectorVMShader(const FShaderCompilerInput& Input, const FShaderPreprocessOutput& PreprocessOutput, FShaderCompilerOutput& Output, const FString& WorkingDirectory)
+bool CompileVectorVMShader(const FShaderCompilerInput& Input, const FShaderPreprocessOutput& PreprocessOutput, FShaderCompilerOutput& Output)
 {
 	FVectorVMCompilationOutput CompilationOutput;
-	bool bResult = CompileVectorVMShader(Input, PreprocessOutput, CompilationOutput, WorkingDirectory, Input.Environment.CompilerFlags.Contains(CFLAG_SkipOptimizations));
+	bool bResult = CompileVectorVMShader(Input, PreprocessOutput, CompilationOutput, Input.Environment.CompilerFlags.Contains(CFLAG_SkipOptimizations));
 	
 	if (bResult)
 	{
@@ -137,7 +136,6 @@ bool CompileVectorVMShader(const FShaderCompilerInput& Input, const FShaderPrepr
 
 bool TestCompileVectorVMShader(
 	const FShaderCompilerInput& Input,
-	const FString& WorkingDirectory,
 	FVectorVMCompilationOutput& VMCompilationOutput,
 	bool bSkipBackendOptimizations)
 {
@@ -156,5 +154,5 @@ bool TestCompileVectorVMShader(
 		return false;
 	}
 
-	return CompileVectorVMShader(Input, PreprocessOutput, VMCompilationOutput, WorkingDirectory, bSkipBackendOptimizations);
+	return CompileVectorVMShader(Input, PreprocessOutput, VMCompilationOutput, bSkipBackendOptimizations);
 }
diff --git a/Engine/Source/Developer/TargetPlatform/Public/Interfaces/IShaderFormat.h b/Engine/Source/Developer/TargetPlatform/Public/Interfaces/IShaderFormat.h
--- a/Engine/Source/Developer/TargetPlatform/Public/Interfaces/IShaderFormat.h
+++ b/Engine/Source/Developer/TargetPlatform/Public/Interfaces/IShaderFormat.h
@@ -133,10 +133,9 @@ public:
 	 * Compile the specified preprocessed shader.
 	 */
 	virtual void CompilePreprocessedShader(
-		const FShaderCompilerInput& Input, 
-		const FShaderPreprocessOutput& PreprocessOutput, 
-		FShaderCompilerOutput& Output, 
-		const FString& WorkingDirectory) const {}
+		const FShaderCompilerInput& Input,
+		const FShaderPreprocessOutput& PreprocessOutput,
+		FShaderCompilerOutput& Output) const {}
 
 	/**
 	 * Compile the specified preprocessed shaders; only called if the call to RequiresSecondaryCompile given the first preprocess 
@@ -146,12 +145,11 @@ public:
 	 * require secondary compilation and is used by the RHI to differentiate packed shader code from single shaders)
 	 */
 	virtual void CompilePreprocessedShader(
-		const FShaderCompilerInput& Input, 
-		const FShaderPreprocessOutput& PrimaryPreprocessOutput, 
-		const FShaderPreprocessOutput& SecondaryPreprocessOutput, 
-		FShaderCompilerOutput& PrimaryOutput, 
-		FShaderCompilerOutput& SecondaryOutput,
-		const FString& WorkingDirectory) const {}
+		const FShaderCompilerInput& Input,
+		const FShaderPreprocessOutput& PrimaryPreprocessOutput,
+		const FShaderPreprocessOutput& SecondaryPreprocessOutput,
+		FShaderCompilerOutput& PrimaryOutput,
+		FShaderCompilerOutput& SecondaryOutput) const {}
 
 	/**
 	 * Predicate which should return true if a second preprocess & compilation is required given the initial preprocess output.
diff --git a/Engine/Source/Developer/VulkanShaderFormat/Private/VulkanShaderCompiler.cpp b/Engine/Source/Developer/VulkanShaderFormat/Private/VulkanShaderCompiler.cpp
--- a/Engine/Source/Developer/VulkanShaderFormat/Private/VulkanShaderCompiler.cpp
+++ b/Engine/Source/Developer/VulkanShaderFormat/Private/VulkanShaderCompiler.cpp
@@ -288,7 +288,7 @@ struct FVulkanShaderParameterParserPlatformConfiguration : public FShaderParamet
 	TMap<FStringView, FStringView>& ReplacedGlobals;
 };
 
-void CompileVulkanShader(const FShaderCompilerInput& Input, const FShaderPreprocessOutput& InPreprocessOutput, FShaderCompilerOutput& Output, const class FString& WorkingDirectory)
+void CompileVulkanShader(const FShaderCompilerInput& Input, const FShaderPreprocessOutput& InPreprocessOutput, FShaderCompilerOutput& Output)
 {
 	check(IsVulkanShaderFormat(Input.ShaderFormat));
 
diff --git a/Engine/Source/Developer/VulkanShaderFormat/Private/VulkanShaderFormat.cpp b/Engine/Source/Developer/VulkanShaderFormat/Private/VulkanShaderFormat.cpp
--- a/Engine/Source/Developer/VulkanShaderFormat/Private/VulkanShaderFormat.cpp
+++ b/Engine/Source/Developer/VulkanShaderFormat/Private/VulkanShaderFormat.cpp
@@ -21,8 +21,7 @@ extern void ModifyVulkanCompilerInput(FShaderCompilerInput& Input);
 extern void CompileVulkanShader(
 	const FShaderCompilerInput& Input,
 	const FShaderPreprocessOutput& InPreprocessOutput,
-	FShaderCompilerOutput& Output,
-	const FString& WorkingDirectory);
+	FShaderCompilerOutput& Output);
 
 static const FGuid UE_SHADER_VULKAN_ES3_1_VER = FGuid("B84F72C8-3ECD-411E-993C-D7C7CEE26F28");
 static const FGuid UE_SHADER_VULKAN_SM5_VER = FGuid("0715D8EE-9907-4A25-93AD-A3902C8E069A");
@@ -87,9 +86,9 @@ public:
 		ModifyVulkanCompilerInput(Input);
 	}
 
-	virtual void CompilePreprocessedShader(const FShaderCompilerInput& Input, const FShaderPreprocessOutput& PreprocessOutput, FShaderCompilerOutput& Output,const FString& WorkingDirectory) const override
+	virtual void CompilePreprocessedShader(const FShaderCompilerInput& Input, const FShaderPreprocessOutput& PreprocessOutput, FShaderCompilerOutput& Output) const override
 	{
-		CompileVulkanShader(Input, PreprocessOutput, Output, WorkingDirectory);
+		CompileVulkanShader(Input, PreprocessOutput, Output);
 	}
 
 	virtual const TCHAR* GetPlatformIncludeDirectory() const
diff --git a/Engine/Source/Developer/Windows/ShaderFormatD3D/Private/D3DShaderCompiler.cpp b/Engine/Source/Developer/Windows/ShaderFormatD3D/Private/D3DShaderCompiler.cpp
--- a/Engine/Source/Developer/Windows/ShaderFormatD3D/Private/D3DShaderCompiler.cpp
+++ b/Engine/Source/Developer/Windows/ShaderFormatD3D/Private/D3DShaderCompiler.cpp
@@ -1431,7 +1431,7 @@ struct FD3DShaderParameterParserPlatformConfiguration : public FShaderParameterP
 	const FString HitGroupSystemVertexBufferName;
 };
 
-void CompileD3DShader(const FShaderCompilerInput& Input, const FShaderPreprocessOutput& InPreprocessOutput, FShaderCompilerOutput& Output, const FString& WorkingDirectory, ED3DShaderModel ShaderModel)
+void CompileD3DShader(const FShaderCompilerInput& Input, const FShaderPreprocessOutput& InPreprocessOutput, FShaderCompilerOutput& Output, ED3DShaderModel ShaderModel)
 {
 	TRACE_CPUPROFILER_EVENT_SCOPE(CompileD3DShader);
 
diff --git a/Engine/Source/Developer/Windows/ShaderFormatD3D/Private/ShaderFormatD3D.cpp b/Engine/Source/Developer/Windows/ShaderFormatD3D/Private/ShaderFormatD3D.cpp
--- a/Engine/Source/Developer/Windows/ShaderFormatD3D/Private/ShaderFormatD3D.cpp
+++ b/Engine/Source/Developer/Windows/ShaderFormatD3D/Private/ShaderFormatD3D.cpp
@@ -186,13 +186,12 @@ public:
 	virtual void CompilePreprocessedShader(
 		const FShaderCompilerInput& Input, 
 		const FShaderPreprocessOutput& PreprocessOutput, 
-		FShaderCompilerOutput& Output,
-		const FString& WorkingDirectory) const
+		FShaderCompilerOutput& Output) const
 	{
 		// Fill in ShaderDiagnosticDatas into the output result before calling CompileD3DShader as it will use and serialize its content
 		Output.ShaderDiagnosticDatas = PreprocessOutput.GetDiagnosticDatas();
 
-		CompileD3DShader(Input, PreprocessOutput, Output, WorkingDirectory, DetermineShaderModel(Input));
+		CompileD3DShader(Input, PreprocessOutput, Output, DetermineShaderModel(Input));
 	}
 
 	static void AddShaderTargetDefines(FShaderCompilerInput& Input, uint32 ShaderTargetMajor, uint32 ShaderTargetMinor)
diff --git a/Engine/Source/Developer/Windows/ShaderFormatD3D/Private/ShaderFormatD3D.h b/Engine/Source/Developer/Windows/ShaderFormatD3D/Private/ShaderFormatD3D.h
--- a/Engine/Source/Developer/Windows/ShaderFormatD3D/Private/ShaderFormatD3D.h
+++ b/Engine/Source/Developer/Windows/ShaderFormatD3D/Private/ShaderFormatD3D.h
@@ -57,7 +57,6 @@ void CompileD3DShader(
 	const FShaderCompilerInput& Input,
 	const FShaderPreprocessOutput& InPreprocessOutput,
 	FShaderCompilerOutput& Output,
-	const FString& WorkingDirectory,
 	ED3DShaderModel ShaderModel);
 
 /**
diff --git a/Engine/Source/Programs/ShaderCompileWorker/Private/ShaderCompileWorker.cpp b/Engine/Source/Programs/ShaderCompileWorker/Private/ShaderCompileWorker.cpp
--- a/Engine/Source/Programs/ShaderCompileWorker/Private/ShaderCompileWorker.cpp
+++ b/Engine/Source/Programs/ShaderCompileWorker/Private/ShaderCompileWorker.cpp
@@ -803,7 +803,7 @@ private:
 				}
 
 				// Process the job.
-				CompileShader(GetShaderFormats(), Job, WorkingDirectory, &GNumProcessedJobs);
+				CompileShader(GetShaderFormats(), Job, &GNumProcessedJobs);
 				if (bIsDebugCompile)
 				{
 					// if we're in debug compile mode and we have a single job, skip the pipeline job section below so
@@ -864,7 +864,7 @@ private:
 					}
 				}
 
-				CompileShaderPipeline(GetShaderFormats(), &PipelineJob, WorkingDirectory, &GNumProcessedJobs);
+				CompileShaderPipeline(GetShaderFormats(), &PipelineJob, &GNumProcessedJobs);
 			}
 		}
 	}
diff --git a/Engine/Source/Runtime/Engine/Private/ShaderCompiler/ShaderCompilerUtil.cpp b/Engine/Source/Runtime/Engine/Private/ShaderCompiler/ShaderCompilerUtil.cpp
--- a/Engine/Source/Runtime/Engine/Private/ShaderCompiler/ShaderCompilerUtil.cpp
+++ b/Engine/Source/Runtime/Engine/Private/ShaderCompiler/ShaderCompilerUtil.cpp
@@ -23,14 +23,13 @@ void FShaderCompileUtilities::ExecuteShaderCompileJob(FShaderCommonCompileJob& J
 
 	check(!Job.bFinalized);
 
-	FString WorkingDir = FPlatformProcess::ShaderDir();
 	static ITargetPlatformManagerModule& TPM = GetTargetPlatformManagerRef();
 	auto* SingleJob = Job.GetSingleShaderJob();
 	const TArray<const IShaderFormat*> ShaderFormats = TPM.GetShaderFormats();
 	if (SingleJob)
 	{
 		const FName Format = (SingleJob->Input.ShaderFormat != NAME_None) ? SingleJob->Input.ShaderFormat : LegacyShaderPlatformToShaderFormat(EShaderPlatform(SingleJob->Input.Target.Platform));
-		CompileShader(ShaderFormats, *SingleJob, WorkingDir);
+		CompileShader(ShaderFormats, *SingleJob);
 	}
 	else
 	{
@@ -54,7 +53,7 @@ void FShaderCompileUtilities::ExecuteShaderCompileJob(FShaderCommonCompileJob& J
 			}
 		}
 
-		CompileShaderPipeline(ShaderFormats, PipelineJob, WorkingDir);
+		CompileShaderPipeline(ShaderFormats, PipelineJob);
 	}
 
 	Job.bFinalized = true;
diff --git a/Engine/Source/Runtime/RenderCore/Private/ShaderCore.cpp b/Engine/Source/Runtime/RenderCore/Private/ShaderCore.cpp
--- a/Engine/Source/Runtime/RenderCore/Private/ShaderCore.cpp
+++ b/Engine/Source/Runtime/RenderCore/Private/ShaderCore.cpp
@@ -1947,7 +1947,7 @@ public:
 		return true;
 	}
 
-	static void InvokeCompile(const IShaderFormat* Compiler, FShaderCompileJob& Job, const FString& WorkingDirectory, FString& OutExceptionMsg, FString& OutExceptionCallstack)
+	static void InvokeCompile(const IShaderFormat* Compiler, FShaderCompileJob& Job, FString& OutExceptionMsg, FString& OutExceptionCallstack)
 	{
 #if PLATFORM_WINDOWS
 		__try
@@ -1960,11 +1960,11 @@ public:
 			{
 				check(Job.SecondaryPreprocessOutput.IsValid());
 				Job.SecondaryPreprocessOutput->DecompressCode();
-				Compiler->CompilePreprocessedShader(Job.Input, Job.PreprocessOutput, *Job.SecondaryPreprocessOutput, Job.Output, *Job.SecondaryOutput, WorkingDirectory);
+				Compiler->CompilePreprocessedShader(Job.Input, Job.PreprocessOutput, *Job.SecondaryPreprocessOutput, Job.Output, *Job.SecondaryOutput);
 			}
 			else
 			{
-				Compiler->CompilePreprocessedShader(Job.Input, Job.PreprocessOutput, Job.Output, WorkingDirectory);
+				Compiler->CompilePreprocessedShader(Job.Input, Job.PreprocessOutput, Job.Output);
 			}
 		}
 #if PLATFORM_WINDOWS
@@ -1975,7 +1975,7 @@ public:
 #endif
 	}
 
-	static void CompileShaderInternal(const IShaderFormat* Compiler, FShaderCompileJob& Job, const FString& WorkingDirectory, FString& OutExceptionMsg, FString& OutExceptionCallstack, int32* CompileCount)
+	static void CompileShaderInternal(const IShaderFormat* Compiler, FShaderCompileJob& Job, FString& OutExceptionMsg, FString& OutExceptionCallstack, int32* CompileCount)
 	{
 		double TimeStart = FPlatformTime::Seconds();
 
@@ -1992,7 +1992,7 @@ public:
 		{
 			Job.SecondaryOutput = MakeUnique<FShaderCompilerOutput>();
 		}
-		InvokeCompile(Compiler, Job, WorkingDirectory, OutExceptionMsg, OutExceptionCallstack);
+		InvokeCompile(Compiler, Job, OutExceptionMsg, OutExceptionCallstack);
 		if (Job.SecondaryOutput.IsValid())
 		{
 			Job.Output.bSucceeded = Job.Output.bSucceeded && Job.SecondaryOutput->bSucceeded;
@@ -2079,16 +2079,16 @@ bool PreprocessShader(FShaderCommonCompileJob* Job)
 	return false;
 }
 
-void CompileShader(const TArray<const IShaderFormat*>& ShaderFormats, FShaderCompilerInput& Input, FShaderCompilerOutput& Output, const FString& WorkingDirectory, int32* CompileCount)
+static void CompileShader(const TArray<const IShaderFormat*>& ShaderFormats, FShaderCompilerInput& Input, FShaderCompilerOutput& Output, int32* CompileCount)
 {
 	FShaderCompileJob TempJob;
 	TempJob.Input = Input;
 	
-	CompileShader(ShaderFormats, TempJob, WorkingDirectory, CompileCount);
+	CompileShader(ShaderFormats, TempJob, CompileCount);
 	Output = TempJob.Output;
 }
 
-void CompileShader(const TArray<const IShaderFormat*>& ShaderFormats, FShaderCompileJob& Job, const FString& WorkingDirectory, int32* CompileCount)
+void CompileShader(const TArray<const IShaderFormat*>& ShaderFormats, FShaderCompileJob& Job, int32* CompileCount)
 {
 	TRACE_CPUPROFILER_EVENT_SCOPE(CompileShader);
 
@@ -2104,7 +2104,7 @@ void CompileShader(const TArray<const IShaderFormat*>& ShaderFormats, FShaderCom
 	}
 	FString ExceptionMsg;
 	FString ExceptionCallstack;
-	FInternalShaderCompilerFunctions::CompileShaderInternal(Compiler, Job, WorkingDirectory, ExceptionMsg, ExceptionCallstack, CompileCount);
+	FInternalShaderCompilerFunctions::CompileShaderInternal(Compiler, Job, ExceptionMsg, ExceptionCallstack, CompileCount);
 	if (!Job.Output.bSucceeded && (!ExceptionMsg.IsEmpty() || !ExceptionCallstack.IsEmpty()))
 	{
 		Job.Output.AddErrorf(
@@ -2118,7 +2118,7 @@ void CompileShader(const TArray<const IShaderFormat*>& ShaderFormats, FShaderCom
 }
 
 // Compiles all shader pipeline stage jobs and returns true if all succeeded.
-static bool CompileShaderPipelineStages(const TArray<const IShaderFormat*>& ShaderFormats, FShaderPipelineCompileJob* PipelineJob, const FString& WorkingDirectory, int32* CompileCount, bool* bOutInterpolatorRemovalFailed = nullptr)
+static bool CompileShaderPipelineStages(const TArray<const IShaderFormat*>& ShaderFormats, FShaderPipelineCompileJob* PipelineJob, int32* CompileCount, bool* bOutInterpolatorRemovalFailed = nullptr)
 {
 	FShaderCompileJob* CurrentJob = nullptr;
 	for (int32 Index = 0; Index < PipelineJob->StageJobs.Num(); ++Index)
@@ -2142,7 +2142,7 @@ static bool CompileShaderPipelineStages(const TArray<const IShaderFormat*>& Shad
 		}
 
 		// Compile the shader directly through the platform dll (directly from the shader dir as the working directory)
-		CompileShader(ShaderFormats, *CurrentJob, WorkingDirectory, CompileCount);
+		CompileShader(ShaderFormats, *CurrentJob, CompileCount);
 
 		CurrentJob->bSucceeded = CurrentJob->Output.bSucceeded;
 		if (!CurrentJob->Output.bSucceeded)
@@ -2192,7 +2192,7 @@ static void ResetShaderCompilerOutputWithRetainedErrors(FShaderCompilerOutput& O
 	Output.bInterpolatorRemovalFailed = bInterpolatorRemovalFailed;
 }
 
-void CompileShaderPipeline(const TArray<const IShaderFormat*>& ShaderFormats, FShaderPipelineCompileJob* PipelineJob, const FString& WorkingDirectory, int32* CompileCount)
+void CompileShaderPipeline(const TArray<const IShaderFormat*>& ShaderFormats, FShaderPipelineCompileJob* PipelineJob, int32* CompileCount)
 {
 	checkf(PipelineJob->StageJobs.Num() > 0, TEXT("Pipeline %s has zero jobs!"), PipelineJob->Key.ShaderPipeline->GetName());
 
@@ -2201,7 +2201,7 @@ void CompileShaderPipeline(const TArray<const IShaderFormat*>& ShaderFormats, FS
 	if (bEnableRemovingUnused)
 	{
 		bool bInterpolatorRemovalFailed = false;
-		if (!CompileShaderPipelineStages(ShaderFormats, PipelineJob, WorkingDirectory, CompileCount, &bInterpolatorRemovalFailed))
+		if (!CompileShaderPipelineStages(ShaderFormats, PipelineJob, CompileCount, &bInterpolatorRemovalFailed))
 		{
 			return;
 		}
@@ -2227,7 +2227,7 @@ void CompileShaderPipeline(const TArray<const IShaderFormat*>& ShaderFormats, FS
 				}
 			);
 
-			if (!CompileShaderPipelineStages(ShaderFormats, PipelineJob, WorkingDirectory, CompileCount))
+			if (!CompileShaderPipelineStages(ShaderFormats, PipelineJob, CompileCount))
 			{
 				return;
 			}
@@ -2235,7 +2235,7 @@ void CompileShaderPipeline(const TArray<const IShaderFormat*>& ShaderFormats, FS
 	}
 	else
 	{
-		if (!CompileShaderPipelineStages(ShaderFormats, PipelineJob, WorkingDirectory, CompileCount))
+		if (!CompileShaderPipelineStages(ShaderFormats, PipelineJob, CompileCount))
 		{
 			return;
 		}
diff --git a/Engine/Source/Runtime/RenderCore/Public/ShaderCompilerCore.h b/Engine/Source/Runtime/RenderCore/Public/ShaderCompilerCore.h
--- a/Engine/Source/Runtime/RenderCore/Public/ShaderCompilerCore.h
+++ b/Engine/Source/Runtime/RenderCore/Public/ShaderCompilerCore.h
@@ -657,8 +657,8 @@ extern RENDERCORE_API const IShaderFormat* FindShaderFormat(FName Format, const
 
 // Executes preprocessing for the given job, if the job is marked to be preprocessed independently prior to compilation.
 extern RENDERCORE_API bool PreprocessShader(FShaderCommonCompileJob* Job);
-extern RENDERCORE_API void CompileShader(const TArray<const IShaderFormat*>& ShaderFormats, FShaderCompileJob& Job, const FString& WorkingDirectory, int32* CompileCount = nullptr);
-extern RENDERCORE_API void CompileShaderPipeline(const TArray<const IShaderFormat*>& ShaderFormats, FShaderPipelineCompileJob* PipelineJob, const FString& WorkingDirectory, int32* CompileCount = nullptr);
+extern RENDERCORE_API void CompileShader(const TArray<const IShaderFormat*>& ShaderFormats, FShaderCompileJob& Job, int32* CompileCount = nullptr);
+extern RENDERCORE_API void CompileShaderPipeline(const TArray<const IShaderFormat*>& ShaderFormats, FShaderPipelineCompileJob* PipelineJob, int32* CompileCount = nullptr);
 
 /**
  * Validates the format of a virtual shader file path.
