Fixed: GenerateMips pixel shader path where each mip level was offset by one texel toward (0,0)
#okforgtihub public

#Synced-CL 49983345

diff --git a/Engine/Shaders/Private/ComputeGenerateMips.usf b/Engine/Shaders/Private/ComputeGenerateMips.usf
--- a/Engine/Shaders/Private/ComputeGenerateMips.usf
+++ b/Engine/Shaders/Private/ComputeGenerateMips.usf
@@ -74,7 +74,6 @@ void BuildIndirectDispatchArgsCS(uint DT_ID : SV_DispatchThreadID)
 
 #else // !GENMIPS_COMPUTE
 
-float2 HalfTexelSize;
 float Level;
 
 void MainVS(in float4 InPosition : ATTRIBUTE0, in float2 InUV : ATTRIBUTE1,	out float4 OutPosition : SV_POSITION, out float2 OutUV : TEXCOORD0)
@@ -87,7 +86,7 @@ void MainVS(in float4 InPosition : ATTRIBUTE0, in float2 InUV : ATTRIBUTE1,	out
 
 void MainPS(float4 InPosition : SV_POSITION,	float2 InUV : TEXCOORD0, out float4 OutColor : SV_Target0)
 {
-    OutColor = MipInSRV.SampleLevel(MipSampler, InUV + HalfTexelSize, Level);
+	OutColor = MipInSRV.SampleLevel(MipSampler, InUV, Level);
 }
 
 #endif // GENMIPS_COMPUTE
\ No newline at end of file
diff --git a/Engine/Source/Runtime/RenderCore/Private/GenerateMips.cpp b/Engine/Source/Runtime/RenderCore/Private/GenerateMips.cpp
--- a/Engine/Source/Runtime/RenderCore/Private/GenerateMips.cpp
+++ b/Engine/Source/Runtime/RenderCore/Private/GenerateMips.cpp
@@ -59,7 +59,6 @@ public:
 	SHADER_USE_PARAMETER_STRUCT(FGenerateMipsPS, FGlobalShader);
 
 	BEGIN_SHADER_PARAMETER_STRUCT(FParameters, )
-		SHADER_PARAMETER(FVector2f, HalfTexelSize)
 		SHADER_PARAMETER(float, Level)
 		SHADER_PARAMETER_RDG_TEXTURE_SRV(Texture2D, MipInSRV)
 		SHADER_PARAMETER_SAMPLER(SamplerState, MipSampler)
@@ -167,7 +166,6 @@ void FGenerateMips::ExecuteRaster(FRDGBuilder& GraphBuilder, ERHIFeatureLevel::T
 			SRVDesc.FirstArraySlice = SliceIndex;
 
 			FGenerateMipsPS::FParameters* PassParameters = GraphBuilder.AllocParameters<FGenerateMipsPS::FParameters>();
-			PassParameters->HalfTexelSize = FVector2f(0.5f / DestTextureSize.X, 0.5f / DestTextureSize.Y);
 			PassParameters->Level = InputMipLevel;
 			PassParameters->MipInSRV = GraphBuilder.CreateSRV(SRVDesc);
 			PassParameters->MipSampler = Sampler;
