Defer the actual enqueue of rendering work until on-commit.

Unfortunately we can't move `TUniqueFunction` into the on-commit lambda for reasons, so use a shared pointer only under transactions to work around this limitation.

#Synced-CL 49988233

diff --git a/Engine/Source/Runtime/RenderCore/Public/RenderingThread.h b/Engine/Source/Runtime/RenderCore/Public/RenderingThread.h
--- a/Engine/Source/Runtime/RenderCore/Public/RenderingThread.h
+++ b/Engine/Source/Runtime/RenderCore/Public/RenderingThread.h
@@ -6,6 +6,7 @@
 
 #pragma once
 
+#include "AutoRTFM.h"
 #include "Async/TaskGraphInterfaces.h"
 #include "Containers/Array.h"
 #include "Containers/List.h"
@@ -1114,6 +1115,20 @@ public:
 	template <typename RenderCommandTag>
 	static void Enqueue(TUniqueFunction<void(FRHICommandListImmediate&)>&& Function)
 	{
+		if (AutoRTFM::IsClosed())
+		{
+			// We cannot move the `TUniqueFunction` into the `OnCommit` handler, so as a workaround we need to shove
+			// it into a shared pointer and fire that in instead. Not ideal, but only affects the code when used under
+			// transactions.
+			AutoRTFM::OnCommit([SharedFunction = MakeShared<TUniqueFunction<void(FRHICommandListImmediate&)>>(MoveTemp(Function))]
+			{
+				TUniqueFunction<void(FRHICommandListImmediate&)> Function = MoveTemp(*SharedFunction);
+				Enqueue<RenderCommandTag>(MoveTemp(Function));
+			});
+
+			return;
+		}
+
 #if WITH_STATE_STREAM
 		RegisterDebugCommand(FString(RenderCommandTag::Get().GetName()));
 #else
