Use atomic for ref count to fix TSAN errors.

#Synced-CL 49999797

diff --git a/Engine/Source/Runtime/RenderCore/Public/RenderTargetPool.h b/Engine/Source/Runtime/RenderCore/Public/RenderTargetPool.h
--- a/Engine/Source/Runtime/RenderCore/Public/RenderTargetPool.h
+++ b/Engine/Source/Runtime/RenderCore/Public/RenderTargetPool.h
@@ -42,22 +42,24 @@ struct FPooledRenderTarget final : public IPooledRenderTarget
 
 	uint32 AddRef() const override
 	{
-		return uint32(FPlatformAtomics::InterlockedIncrement(&NumRefs));
+		return NumRefs.fetch_add(1, std::memory_order_relaxed) + 1;
 	}
 
 	uint32 Release() override
 	{
-		const int32 Refs = FPlatformAtomics::InterlockedDecrement(&NumRefs);
-		if (Refs == 0)
+		const uint32 OldNumRefs = NumRefs.fetch_sub(1, std::memory_order_release);
+		check(OldNumRefs > 0);
+		if (OldNumRefs == 1)
 		{
+			std::atomic_thread_fence(std::memory_order_acquire);
 			delete this;
 		}
-		return uint32(Refs);
+		return OldNumRefs - 1;
 	}
 
 	uint32 GetRefCount() const override
 	{
-		return uint32(NumRefs);
+		return NumRefs.load(std::memory_order_relaxed);
 	}
 
 	RENDERCORE_API bool IsFree() const override;
@@ -74,7 +76,7 @@ private:
 	FPooledRenderTargetDesc Desc;
 
 	/** For pool management (only if NumRef == 0 the element can be reused) */
-	mutable int32 NumRefs = 0;
+	mutable std::atomic_uint32_t NumRefs{0};
 
 	/** Allows to defer the release to save performance on some hardware (DirectX) */
 	uint32 UnusedForNFrames = 0;
