Change type of global static bindless material UB from structtype to ConstantBuffer<structtype> to significantly improve bindless codegen on PC.

#Synced-CL 49998540

diff --git a/Engine/Source/Runtime/RenderCore/Private/ShaderParameters.cpp b/Engine/Source/Runtime/RenderCore/Private/ShaderParameters.cpp
--- a/Engine/Source/Runtime/RenderCore/Private/ShaderParameters.cpp
+++ b/Engine/Source/Runtime/RenderCore/Private/ShaderParameters.cpp
@@ -521,7 +521,7 @@ static FString CreateHLSLUniformBufferDeclaration(const TCHAR* UniformBufferName
 				UniformBufferName
 			);
 
-			Builder.Appendf(TEXT("static UB_CBO_TYPE(%s) %s;\n"),
+			Builder.Appendf(TEXT("static UB_SAFE_TYPE(%s) %s;\n"),
 				UniformBufferName, UniformBufferName
 			);
 
