Add the ability to specify an output directory where SHK shader files can be copied during a cook.  The main goal here is to have a central, shared location to place SHK files  preferably for CIS or nightly builds so we can track and catch regressions.

- The command line "-mirrorshkfiles" will enable the copying of the shk files.
- You must specify an output location via an ini file where all the results will go.

[/Script/Engine.ShaderCompilerStats]
SHKFilesLocation=E:\SHKFiles

- Output files exist with the file names like so: StreamName-Lyra-50002989-Windows-Lyra-PCD3D_SM5.shk
- StreamName, ProjectName, Changelist, Platform, ShaderLibraryName, ShaderPlatform.
- This way all the SHK files from multiple builds and multiple streams can coexisting in a single location.
- We also overwrite existing files with the same name so a single build at a single changelist will only generate a single artifact.

[RN] minor | Rendering - Materials and Shaders

#Synced-CL 50009848

#Synced-CL 50009848

diff --git a/Engine/Source/Runtime/RenderCore/Private/ShaderLibrary/EditorShaderStableInfo.cpp b/Engine/Source/Runtime/RenderCore/Private/ShaderLibrary/EditorShaderStableInfo.cpp
--- a/Engine/Source/Runtime/RenderCore/Private/ShaderLibrary/EditorShaderStableInfo.cpp
+++ b/Engine/Source/Runtime/RenderCore/Private/ShaderLibrary/EditorShaderStableInfo.cpp
@@ -32,6 +32,9 @@
 #include "Serialization/CompactBinaryWriter.h"
 #include "RHIStrings.h"
 #include "CookArtifactReader.h"
+#include "Misc/App.h"
+#include "Misc/ConfigCacheIni.h"
+#include "Misc/EngineVersion.h"
 #endif
 
 #if WITH_EDITOR
@@ -195,6 +198,27 @@ bool FEditorShaderStableInfo::SaveToDisk(FString const& OutputDir, FString& OutS
 		// Only the primary cooker needs to write to the output directory, child cookers only write to the Saved directory
 		FString OutputFilePath = GetStableInfoArchiveFilename(OutputDir, LibraryName, FormatName);
 
+		if (FParse::Param(FCommandLine::Get(), TEXT("mirrorshkfiles")))
+		{
+			FString MirrorLocation;
+			GConfig->GetString(TEXT("/Script/Engine.ShaderCompilerStats"), TEXT("SHKFilesLocation"), MirrorLocation, GGameIni);
+			if (!MirrorLocation.IsEmpty())
+			{
+				FString TargetType = "Default";
+				FParse::Value(FCommandLine::Get(), TEXT("target="), TargetType);
+				if (TargetType == TEXT("Default"))
+				{
+					FParse::Value(FCommandLine::Get(), TEXT("targetplatform="), TargetType);
+				}
+				const FString CopyLocation = FPaths::Combine(*MirrorLocation, FString::Printf(TEXT("%s-%s-%d-%s-%s-%s.shk"), *FApp::GetBranchName(), FApp::GetProjectName(), FEngineVersion::Current().GetChangelist(), *TargetType, *LibraryName, *FormatName.ToString()));
+				const uint32 CopyResult = IFileManager::Get().Copy(*CopyLocation, *IntermediateFormatPath, true, true);
+				if (CopyResult == COPY_Fail)
+				{
+					UE_LOG(LogShaderLibrary, Warning, TEXT("Copy from Src: %s to Dest: %s failed."), *IntermediateFormatPath, *CopyLocation);
+				}
+			}
+		}
+
 		// Copy to output location - support for iterative native library cooking
 		uint32 Result = IFileManager::Get().Copy(*OutputFilePath, *IntermediateFormatPath, true, true);
 		if (Result == COPY_OK)
