~19% optimization to LoadStableKeysFile by preventing unecessary allocations by reserveing the TArray size for all shader entries.

- 7511 seconds -> 6029 seconds in my large test file with million of entries.

[RN] minor | Rendering - Materials and Shaders

#Synced-CL 50009848

diff --git a/Engine/Source/Runtime/RenderCore/Private/PipelineCacheUtilities.cpp b/Engine/Source/Runtime/RenderCore/Private/PipelineCacheUtilities.cpp
--- a/Engine/Source/Runtime/RenderCore/Private/PipelineCacheUtilities.cpp
+++ b/Engine/Source/Runtime/RenderCore/Private/PipelineCacheUtilities.cpp
@@ -345,6 +345,7 @@ bool UE::PipelineCacheUtilities::LoadStableKeysFile(const FStringView& Filename,
 		*Archive << Hashes[IdxHash];
 	}
 
+	InOutArray.Reserve(Header.NumEntries);
 	for (int64 Idx = 0; Idx < Header.NumEntries; ++Idx)
 	{
 		FStableShaderKeyAndValue Item;
