[RenderCore]
* Added RenderCommandStateStream (only compiled in when WITH_STATE_STREAM 1)
* Removed code from RenderingThread.h

#Synced-CL 50066719

diff --git a/Engine/Source/Runtime/RenderCore/Private/RenderCommandStateStream.cpp b/Engine/Source/Runtime/RenderCore/Private/RenderCommandStateStream.cpp
new file mode 100644
--- /dev/null
+++ b/Engine/Source/Runtime/RenderCore/Private/RenderCommandStateStream.cpp
@@ -0,0 +1,105 @@
+// Copyright Epic Games, Inc. All Rights Reserved.
+
+#include "RenderCommandStateStream.h"
+
+#if WITH_STATE_STREAM
+
+#include "RenderCommandFence.h"
+#include "RenderingThread.h"
+#include "StateStream.h"
+#include "StateStreamCreator.h"
+#include "StateStreamDebugRenderer.h"
+#include "StateStreamManagerImpl.h"
+
+////////////////////////////////////////////////////////////////////////////////////////////////////
+
+class FRenderCommandStateStream : public IStateStream
+{
+public:
+	FRenderCommandStateStream(FSceneInterface& InScene)
+	{
+		// TODO: Flush
+		//FRenderCommandFence Fence;
+		//Fence.BeginFence();
+		//Fence.Wait(false);
+	}
+
+	virtual void Game_BeginTick() override
+	{
+	}
+	virtual void Game_EndTick(StateStreamTime AbsoluteTime) override
+	{
+	}
+	virtual void Game_Exit() override
+	{
+	}
+	virtual void* Game_GetVoidPointer() override
+	{
+		return this;
+	}
+	virtual void Game_DebugRender(IStateStreamDebugRenderer& Renderer)
+	{
+		TStringBuilder<1024> DebugLine;
+		DebugLine.Appendf(TEXT("State stream : [%s]"), GetDebugName());
+		Renderer.DrawText(*DebugLine, FColor::Orange);
+	}
+
+	virtual void Render_Update(StateStreamTime AbsoluteTime) override
+	{
+	}
+	virtual void Render_PostUpdate() override
+	{
+	}
+	virtual void Render_Exit() override
+	{
+	}
+	virtual void Render_GarbageCollect() override
+	{
+	}
+	virtual void Render_Enable(bool Enable) override
+	{
+	}
+	virtual uint32 GetId() override
+	{
+		return RenderCommandStateStreamId;
+	}
+	virtual const TCHAR* GetDebugName() override { return TEXT("RenderCommand"); }
+
+	enum { Id = RenderCommandStateStreamId };
+
+	//TMap<FString,uint32> DebugElementsCounts;
+};
+
+////////////////////////////////////////////////////////////////////////////////////////////////////
+
+bool RenderCommandStateStream_AddCommand(TUniqueFunction<void(FRHICommandListImmediate&)>& Function, const FRenderCommandTag& Tag)
+{
+	if (IsInGameThread() || IsInParallelGameThread())
+	{
+		//if (RenderCommandStateStream_AddCommand(Function, Tag))
+		//	return;
+		FStateStreamManagerImpl::GetDebugRenderCommands().AddDebugElements(Tag.GetName(), 1);
+
+		// These should be added to current tick of RenderCommandStateStream
+
+	}
+	else // if (GRenderingThreadHeartbeat != FRunnableThread::GetRunnableThread())
+	{
+		// These can be called in between game ticks but still needs to be in sync with ticks.. will need to queue up separately
+		// and inject to beginning of next tick when begin?
+
+		// Things that could end up here
+		//   GRenderingThreadHeartbeat (FDefaultGameMoviePlayer, etc)
+		//   AsyncPackage2 loading (FMaterialShaderMap, FTextureReference, etc)
+		//   FSlateRHIRenderer::ReleaseDrawBuffer
+	}
+	return false;
+}
+
+////////////////////////////////////////////////////////////////////////////////////////////////////
+
+STATESTREAM_CREATOR_INSTANCE(FRenderCommandStateStream)
+
+////////////////////////////////////////////////////////////////////////////////////////////////////
+
+#endif
diff --git a/Engine/Source/Runtime/RenderCore/Private/RenderCommandStateStream.h b/Engine/Source/Runtime/RenderCore/Private/RenderCommandStateStream.h
new file mode 100644
--- /dev/null
+++ b/Engine/Source/Runtime/RenderCore/Private/RenderCommandStateStream.h
@@ -0,0 +1,23 @@
+// Copyright Epic Games, Inc. All Rights Reserved.
+
+#pragma once
+
+#include "Misc/Build.h"
+
+#if WITH_STATE_STREAM
+
+#include "Templates/Function.h"
+
+class FRHICommandListImmediate;
+class FRenderCommandTag;
+
+////////////////////////////////////////////////////////////////////////////////////////////////////
+
+inline constexpr uint32 RenderCommandStateStreamId = 14;
+
+// This will take Function if returning true
+bool RenderCommandStateStream_AddCommand(TUniqueFunction<void(FRHICommandListImmediate&)>& Function, const FRenderCommandTag& Tag);
+
+////////////////////////////////////////////////////////////////////////////////////////////////////
+
+#endif
\ No newline at end of file
diff --git a/Engine/Source/Runtime/RenderCore/Private/RenderingThread.cpp b/Engine/Source/Runtime/RenderCore/Private/RenderingThread.cpp
--- a/Engine/Source/Runtime/RenderCore/Private/RenderingThread.cpp
+++ b/Engine/Source/Runtime/RenderCore/Private/RenderingThread.cpp
@@ -17,6 +17,7 @@
 #include "Misc/CoreDelegates.h"
 #include "Misc/ScopeLock.h"
 #include "RHI.h"
+#include "RenderCommandStateStream.h"
 #include "RenderCore.h"
 #include "RenderCommandFence.h"
 #include "RenderDeferredCleanup.h"
@@ -32,7 +33,6 @@
 #include "ProfilingDebugging/CsvProfiler.h"
 #include "Async/TaskTrace.h"
 #include "DataDrivenShaderPlatformInfo.h"
-#include "StateStreamManagerImpl.h"
 #include "HAL/ThreadManager.h"
 #include "ProfilingDebugging/CountersTrace.h"
 #include "Async/ParallelFor.h"
@@ -1902,6 +1902,11 @@ void FRenderThreadCommandPipe::EnqueueAndLaunch(FRenderCommandList* CommandList)
 
 void FRenderThreadCommandPipe::EnqueueAndLaunch(TUniqueFunction<void(FRHICommandListImmediate&)>&& Function, const FRenderCommandTag& Tag)
 {
+#if WITH_STATE_STREAM
+	if (RenderCommandStateStream_AddCommand(Function, Tag))
+		return;
+#endif
+
 	Mutex.Lock();
 	const bool bWasEmpty = Context->CommandList.Enqueue(FRenderCommandFunctionVariant(TInPlaceType<TUniqueFunction<void(FRHICommandListImmediate&)>>(), MoveTemp(Function)), Tag);
 	Mutex.Unlock();
@@ -2425,18 +2430,6 @@ void FRenderCommandPipe::ExecuteCommands(UE::RenderCommandPipe::FCommandList& Co
 
 //////////////////////////////////////////////////////////////////////////
 
-#if WITH_STATE_STREAM
-void FRenderCommandDispatcher::RegisterDebugCommand(const FString CommandName)
-{
-	if (IsInGameThread() || IsInParallelGameThread())
-	{
-		FStateStreamManagerImpl::GetDebugRenderCommands().AddDebugElements(CommandName, 1);
-	}
-}
-#endif
-
-//////////////////////////////////////////////////////////////////////////
-
 TAutoConsoleVariable<int32> CVarAllowOneFrameThreadLag(
 	TEXT("r.OneFrameThreadLag"),
 	1,
diff --git a/Engine/Source/Runtime/RenderCore/Public/RenderingThread.h b/Engine/Source/Runtime/RenderCore/Public/RenderingThread.h
--- a/Engine/Source/Runtime/RenderCore/Public/RenderingThread.h
+++ b/Engine/Source/Runtime/RenderCore/Public/RenderingThread.h
@@ -1095,12 +1095,6 @@ private:
 
 class FRenderCommandDispatcher
 {
-private : 
-#if WITH_STATE_STREAM
-	/** Register the debug command to state stream */
-	RENDERCORE_API static void RegisterDebugCommand(const FString CommandName);
-#endif
-	
 public:
 	/**
 	 * Call to submit a command list into a parent command list or render command pipes. If the parent command list is null the recording instance
@@ -1129,15 +1123,11 @@ public:
 			return;
 		}
 
-#if WITH_STATE_STREAM
-		RegisterDebugCommand(FString(RenderCommandTag::Get().GetName()));
-#else
 		if (FRenderCommandList* CommandList = FRenderCommandList::GetInstanceTLS())
 		{
 			CommandList->Enqueue<RenderCommandTag>(MoveTemp(Function));
 			return;
 		}
-#endif
 
 		FRenderThreadCommandPipe::Enqueue<RenderCommandTag>(MoveTemp(Function));
 	}
@@ -1145,9 +1135,6 @@ public:
 	template <typename RenderCommandTag>
 	static void Enqueue(FRenderCommandPipe* Pipe, FRenderCommandPipe::FCommandListFunction&& Function)
 	{
-#if WITH_STATE_STREAM
-		RegisterDebugCommand(FString(RenderCommandTag::Get().GetName()));
-#else
 		if (FRenderCommandList* CommandList = FRenderCommandList::GetInstanceTLS())
 		{
 			CommandList->Enqueue<RenderCommandTag>(Pipe, MoveTemp(Function));
@@ -1158,7 +1145,6 @@ public:
 		{
 			return;
 		}
-#endif
 
 		FRenderThreadCommandPipe::Enqueue<RenderCommandTag>([Function = MoveTemp(Function)](FRHICommandListImmediate& RHICmdList) { Function(RHICmdList); });
 	}
@@ -1172,9 +1158,6 @@ public:
 	template <typename RenderCommandTag>
 	static void Enqueue(FRenderCommandPipe* Pipe, FRenderCommandPipe::FEmptyFunction&& Function)
 	{
-#if WITH_STATE_STREAM
-		RegisterDebugCommand(FString(RenderCommandTag::Get().GetName()));
-#else
 		if (FRenderCommandList* CommandListSet = FRenderCommandList::GetInstanceTLS())
 		{
 			CommandListSet->Enqueue<RenderCommandTag>(Pipe, MoveTemp(Function));
@@ -1185,7 +1168,6 @@ public:
 		{
 			return;
 		}
-#endif
 
 		FRenderThreadCommandPipe::Enqueue<RenderCommandTag>([Function = MoveTemp(Function)](FRHICommandListImmediate&) { Function(); });
 	}
