[RenderCore]
* Moved render command tag to its own file

#Synced-CL 50091045

diff --git a/Engine/Source/Runtime/RenderCore/Public/RenderCommandTag.h b/Engine/Source/Runtime/RenderCore/Public/RenderCommandTag.h
new file mode 100644
--- /dev/null
+++ b/Engine/Source/Runtime/RenderCore/Public/RenderCommandTag.h
@@ -0,0 +1,134 @@
+// Copyright Epic Games, Inc. All Rights Reserved.
+
+#pragma once
+
+#include "Stats/Stats.h"
+
+////////////////////////////////////////////////////////////////////////////////////////////////////
+
+class FRenderCommandTag
+{
+public:
+	const TCHAR* GetName() const
+	{
+		return Name;
+	}
+
+	TStatId GetStatId() const
+	{
+		return StatId;
+	}
+
+	uint32& GetSpecId() const
+	{
+		return SpecId;
+	}
+
+protected:
+	FRenderCommandTag(const TCHAR* InName, TStatId InStatId)
+		: Name(InName)
+		, StatId(InStatId)
+	{}
+
+private:
+	const TCHAR* Name;
+	TStatId StatId;
+	mutable uint32 SpecId;
+};
+
+////////////////////////////////////////////////////////////////////////////////////////////////////
+
+DECLARE_STATS_GROUP(TEXT("Render Thread Commands"), STATGROUP_RenderThreadCommands, STATCAT_Advanced);
+
+////////////////////////////////////////////////////////////////////////////////////////////////////
+/** Type that contains profiler data necessary to mark up render commands for various profilers. */
+
+template <typename TSTR>
+class TRenderCommandTag : public FRenderCommandTag
+{
+public:
+	static const TRenderCommandTag& Get()
+	{
+#if STATS
+		struct FStatData
+		{
+			typedef FStatGroup_STATGROUP_RenderThreadCommands TGroup;
+			static inline const char* GetStatName()
+			{
+				return TSTR::CStr();
+			}
+			static inline const TCHAR* GetDescription()
+			{
+				return TSTR::TStr();
+			}
+			static inline EStatDataType::Type GetStatType()
+			{
+				return EStatDataType::ST_int64;
+			}
+			static inline bool IsClearEveryFrame()
+			{
+				return true;
+			}
+			static inline bool IsCycleStat()
+			{
+				return true;
+			}
+			static inline FPlatformMemory::EMemoryCounterRegion GetMemoryRegion()
+			{
+				return FPlatformMemory::MCR_Invalid;
+			}
+		};
+		static FThreadSafeStaticStat<FStatData> Stat;
+		static TRenderCommandTag Tag(TSTR::TStr(), Stat.GetStatId());
+#else
+		static TRenderCommandTag Tag(TSTR::TStr(), TStatId());
+#endif
+
+		return Tag;
+	}
+
+private:
+	TRenderCommandTag(const TCHAR* InName, TStatId InStatId)
+		: FRenderCommandTag(InName, InStatId)
+	{}
+};
+
+////////////////////////////////////////////////////////////////////////////////////////////////////
+/** Type that contains profiler data necessary to mark up render commands for various profilers. */
+
+template <typename TSTR>
+class TRenderCommandTagNoMarker : public FRenderCommandTag
+{
+public:
+	static const TRenderCommandTagNoMarker& Get()
+	{
+		static TRenderCommandTagNoMarker Tag(TSTR::TStr(), TStatId());
+		return Tag;
+	}
+
+private:
+	TRenderCommandTagNoMarker(const TCHAR* InName, TStatId InStatId)
+		: FRenderCommandTag(InName, TStatId())
+	{
+	}
+};
+
+
+/** Declares a new render command tag type from a name. */
+#define DECLARE_RENDER_COMMAND_TAG(Type, Name) \
+	struct UE_JOIN(TSTR_, Name, __LINE__) \
+	{  \
+		static const char* CStr() { return #Name; } \
+		static const TCHAR* TStr() { return TEXT(#Name); } \
+	}; \
+	using Type = TRenderCommandTag<UE_JOIN(TSTR_, Name, __LINE__)>;
+
+#define DECLARE_RENDER_COMMAND_TAG_NOMARKER(Type, Name) \
+	struct UE_JOIN(TSTR_, Name, __LINE__) \
+	{  \
+		static const char* CStr() { return #Name; } \
+		static const TCHAR* TStr() { return TEXT(#Name); } \
+	}; \
+	using Type = TRenderCommandTagNoMarker<UE_JOIN(TSTR_, Name, __LINE__)>;
+
+////////////////////////////////////////////////////////////////////////////////////////////////////
diff --git a/Engine/Source/Runtime/RenderCore/Public/RenderingThread.h b/Engine/Source/Runtime/RenderCore/Public/RenderingThread.h
--- a/Engine/Source/Runtime/RenderCore/Public/RenderingThread.h
+++ b/Engine/Source/Runtime/RenderCore/Public/RenderingThread.h
@@ -19,10 +19,10 @@
 #include "Misc/TVariant.h"
 #include "MultiGPU.h"
 #include "ProfilingDebugging/CpuProfilerTrace.h"
+#include "RenderCommandTag.h"
 #include "RHI.h"
 #include "RHICommandList.h"
 #include "Serialization/MemoryLayout.h"
-#include "Stats/Stats.h"
 #include "Templates/Atomic.h"
 #include "Templates/Function.h"
 #include "Templates/UnrealTemplate.h"
@@ -135,8 +135,6 @@ UE_TRACE_CHANNEL_EXTERN(RenderCommandsChannel, RENDERCORE_API);
 // ideally this would be inline, however that changes the module dependency situation
 extern RENDERCORE_API class FRHICommandListImmediate& GetImmediateCommandList_ForRenderCommand();
 
-DECLARE_STATS_GROUP(TEXT("Render Thread Commands"), STATGROUP_RenderThreadCommands, STATCAT_Advanced);
-
 // Log render commands on server for debugging
 #if 0 // UE_SERVER && UE_BUILD_DEBUG
 	#define LogRenderCommand(TypeName)				UE_LOG(LogRHI, Warning, TEXT("Render command '%s' is being executed on a dedicated server."), TEXT(#TypeName))
@@ -151,123 +149,6 @@ DECLARE_STATS_GROUP(TEXT("Render Thread Commands"), STATGROUP_RenderThreadComman
 	#define	ShouldExecuteOnRenderThread()			(LIKELY(GIsThreadedRendering || !IsInGameThread()))
 #endif // UE_SERVER
 
-class FRenderCommandTag
-{
-public:
-	const TCHAR* GetName() const
-	{
-		return Name;
-	}
-
-	TStatId GetStatId() const
-	{
-		return StatId;
-	}
-
-	uint32& GetSpecId() const
-	{
-		return SpecId;
-	}
-
-protected:
-	FRenderCommandTag(const TCHAR* InName, TStatId InStatId)
-		: Name(InName)
-		, StatId(InStatId)
-	{}
-
-private:
-	const TCHAR* Name;
-	TStatId StatId;
-	mutable uint32 SpecId;
-};
-
-/** Type that contains profiler data necessary to mark up render commands for various profilers. */
-template <typename TSTR>
-class TRenderCommandTag : public FRenderCommandTag
-{
-public:
-	static const TRenderCommandTag& Get()
-	{
-#if STATS
-		struct FStatData
-		{
-			typedef FStatGroup_STATGROUP_RenderThreadCommands TGroup;
-			static inline const char* GetStatName()
-			{
-				return TSTR::CStr();
-			}
-			static inline const TCHAR* GetDescription()
-			{
-				return TSTR::TStr();
-			}
-			static inline EStatDataType::Type GetStatType()
-			{
-				return EStatDataType::ST_int64;
-			}
-			static inline bool IsClearEveryFrame()
-			{
-				return true;
-			}
-			static inline bool IsCycleStat()
-			{
-				return true;
-			}
-			static inline FPlatformMemory::EMemoryCounterRegion GetMemoryRegion()
-			{
-				return FPlatformMemory::MCR_Invalid;
-			}
-		};
-		static FThreadSafeStaticStat<FStatData> Stat;
-		static TRenderCommandTag Tag(TSTR::TStr(), Stat.GetStatId());
-#else
-		static TRenderCommandTag Tag(TSTR::TStr(), TStatId());
-#endif
-
-		return Tag;
-	}
-
-private:
-	TRenderCommandTag(const TCHAR* InName, TStatId InStatId)
-		: FRenderCommandTag(InName, InStatId)
-	{}
-};
-
-/** Type that contains profiler data necessary to mark up render commands for various profilers. */
-template <typename TSTR>
-class TRenderCommandTagNoMarker : public FRenderCommandTag
-{
-public:
-	static const TRenderCommandTagNoMarker& Get()
-	{
-		static TRenderCommandTagNoMarker Tag(TSTR::TStr(), TStatId());
-		return Tag;
-	}
-
-private:
-	TRenderCommandTagNoMarker(const TCHAR* InName, TStatId InStatId)
-		: FRenderCommandTag(InName, TStatId())
-	{
-	}
-};
-
-
-/** Declares a new render command tag type from a name. */
-#define DECLARE_RENDER_COMMAND_TAG(Type, Name) \
-	struct UE_JOIN(TSTR_, Name, __LINE__) \
-	{  \
-		static const char* CStr() { return #Name; } \
-		static const TCHAR* TStr() { return TEXT(#Name); } \
-	}; \
-	using Type = TRenderCommandTag<UE_JOIN(TSTR_, Name, __LINE__)>;
-
-#define DECLARE_RENDER_COMMAND_TAG_NOMARKER(Type, Name) \
-	struct UE_JOIN(TSTR_, Name, __LINE__) \
-	{  \
-		static const char* CStr() { return #Name; } \
-		static const TCHAR* TStr() { return TEXT(#Name); } \
-	}; \
-	using Type = TRenderCommandTagNoMarker<UE_JOIN(TSTR_, Name, __LINE__)>;
-
 /** Describes which pipes are configured to use the render command pipe system. */
 enum class ERenderCommandPipeMode : uint8
 {
