[RenderCore]
* Added Category to RenderCommandTag
* Added category to a few ENQUEUE_RENDER_COMMAND just to test out
* Fixed so RenderCommandStateStream can render render commands per category

#Synced-CL 50091045

diff --git a/Engine/Source/Runtime/Engine/Private/GameEngine.cpp b/Engine/Source/Runtime/Engine/Private/GameEngine.cpp
--- a/Engine/Source/Runtime/Engine/Private/GameEngine.cpp
+++ b/Engine/Source/Runtime/Engine/Private/GameEngine.cpp
@@ -2140,7 +2140,7 @@ void UGameEngine::Tick( float DeltaSeconds, bool bIdleMode )
 		GameAudioDeviceManager->UpdateActiveAudioDevices(bIsAnyNonPreviewWorldUnpaused);
 	}
 
-	ENQUEUE_RENDER_COMMAND(TickPoolElements)([](FRHICommandListImmediate& RHICmdList)
+	ENQUEUE_RENDER_COMMAND(TickPoolElements, LoopTick)([](FRHICommandListImmediate& RHICmdList)
 	{			
 		GRenderTargetPool.TickPoolElements();
 		FRDGBuilder::TickPoolElements();
diff --git a/Engine/Source/Runtime/Engine/Private/UnrealClient.cpp b/Engine/Source/Runtime/Engine/Private/UnrealClient.cpp
--- a/Engine/Source/Runtime/Engine/Private/UnrealClient.cpp
+++ b/Engine/Source/Runtime/Engine/Private/UnrealClient.cpp
@@ -1555,7 +1555,7 @@ bool FViewport::TakeHighResScreenShot()
 static void HighResScreenshotBeginFrame(FDummyViewport* DummyViewport)
 {
 	GFrameCounter++;
-	ENQUEUE_RENDER_COMMAND(BeginFrameCommand)(
+	ENQUEUE_RENDER_COMMAND(BeginFrameCommand, LoopTick)(
 		[DummyViewport, CurrentFrameCounter = GFrameCounter](FRHICommandListImmediate& RHICmdList)
 	{
 		GFrameCounterRenderThread = CurrentFrameCounter;
@@ -1570,7 +1570,7 @@ static void HighResScreenshotBeginFrame(FDummyViewport* DummyViewport)
 
 static void HighResScreenshotEndFrame(FDummyViewport* DummyViewport)
 {
-	ENQUEUE_RENDER_COMMAND(EndFrameCommand)(
+	ENQUEUE_RENDER_COMMAND(EndFrameCommand, LoopTick)(
 		[DummyViewport](FRHICommandListImmediate& RHICmdList)
 	{
 		if (DummyViewport)
@@ -1768,7 +1768,7 @@ FRHIGPUMask FViewport::GetGPUMask(FRHICommandListImmediate& RHICmdList) const
 void FViewport::EnqueueBeginRenderFrame(const bool bShouldPresent)
 {
 	FViewport* Viewport = this;
-	ENQUEUE_RENDER_COMMAND(BeginDrawingCommand)(
+	ENQUEUE_RENDER_COMMAND(BeginDrawingCommand, LoopTick)(
 		[Viewport](FRHICommandListImmediate& RHICmdList)
 		{
 			Viewport->BeginRenderFrame(RHICmdList);
@@ -1779,7 +1779,7 @@ void FViewport::EnqueueBeginRenderFrame(const bool bShouldPresent)
 void FViewport::EnqueueEndRenderFrame(const bool bLockToVsync, const bool bShouldPresent)
 {
 	FEndDrawingCommandParams Params = { this, (uint32)bLockToVsync, (uint32)GInputLatencyTimer.GameThreadTrigger, (uint32)(PresentAndStopMovieDelay > 0 ? 0 : bShouldPresent) };
-	ENQUEUE_RENDER_COMMAND(EndDrawingCommand)(
+	ENQUEUE_RENDER_COMMAND(EndDrawingCommand, LoopTick)(
 		[Params](FRHICommandListImmediate& RHICmdList)
 		{
 			ViewportEndDrawing(RHICmdList, Params);
diff --git a/Engine/Source/Runtime/RenderCore/Private/RenderCommandStateStream.cpp b/Engine/Source/Runtime/RenderCore/Private/RenderCommandStateStream.cpp
--- a/Engine/Source/Runtime/RenderCore/Private/RenderCommandStateStream.cpp
+++ b/Engine/Source/Runtime/RenderCore/Private/RenderCommandStateStream.cpp
@@ -18,6 +18,7 @@
 class FRenderCommandStateStream;
 FRenderCommandStateStream* g_renderCommandStateStream;
 
+constexpr uint32 RenderCommandCategoryCount = uint32(ERenderCommandCategory::Unknown) + 1;
 
 class FRenderCommandStateStream : public IStateStream
 {
@@ -59,28 +60,45 @@ public:
 
 #if UE_RENDER_COMMAND_DEBUG_DRAW
 		uint32 X = 500;
-		uint32 Y = 100;
 
-		DebugLine.Reset();
-
-		FScopeLock Lock(&RenderCommandsLock);
-		DebugLine.Appendf(TEXT("Num RenderCommands = %d"), RenderCommands.Num());
-		Renderer.DrawText(X, Y, *DebugLine, FColor::Blue);
-		Y += 20;
-
-		for (auto& Pair : RenderCommands)
+		static const TCHAR* Categories[] = 
+		{
+			TEXT("EnterTick"),
+			TEXT("LoopTick"),
+			TEXT("LeaveTick"),
+			TEXT("Unknown"),
+		};
+		static_assert(UE_ARRAY_COUNT(Categories) == RenderCommandCategoryCount);
+
+		for (uint32 I=0; I!=4; ++I)
 		{
-			if (Y == 1000)
+			uint32 Y = 100;
+			FScopeLock Lock(&RenderCommandsLock[I]);
+			uint32 Num = RenderCommands[I].Num();
+			if (!Num)
 			{
-				Y = 100;
-				X += 300;
+				continue;
 			}
 			DebugLine.Reset();
-			DebugLine.Appendf(TEXT("%s, %d"), *Pair.Key, Pair.Value);
-			Renderer.DrawText(X, Y, *DebugLine, FColor::Orange);
+			DebugLine.Appendf(TEXT("%s RenderCommands = %d"), Categories[I], Num);
+			Renderer.DrawText(X, Y, *DebugLine, FColor::Blue);
 			Y += 20;
+
+			for (auto& Pair : RenderCommands[I])
+			{
+				if (Y == 1000)
+				{
+					Y = 100;
+					X += 300;
+				}
+				DebugLine.Reset();
+				DebugLine.Appendf(TEXT("%s, %d"), *Pair.Key, Pair.Value);
+				Renderer.DrawText(X, Y, *DebugLine, FColor::Orange);
+				Y += 20;
+			}
+			RenderCommands[I].Empty();
+			X += 300;
 		}
-		RenderCommands.Empty();
 #endif
 	}
 
@@ -115,8 +133,8 @@ public:
 			//check(g_renderCommandStateStream->ActiveDefaultLaneId != DefaultLaneId);
 
 #if UE_RENDER_COMMAND_DEBUG_DRAW
-			FScopeLock Lock(&RenderCommandsLock);
-			++RenderCommands.FindOrAdd(Tag.GetName(), 0);
+			FScopeLock Lock(&RenderCommandsLock[(uint32)Tag.GetCategory()]);
+			++RenderCommands[(uint32)Tag.GetCategory()].FindOrAdd(Tag.GetName(), 0);
 #endif
 
 			// These should be added to current tick of RenderCommandStateStream
@@ -138,8 +156,8 @@ public:
 
 	uint32 ActiveDefaultLaneId = DefaultLaneId;
 
-	FCriticalSection RenderCommandsLock;
-	TMap<FString,uint32> RenderCommands;
+	FCriticalSection RenderCommandsLock[RenderCommandCategoryCount];
+	TMap<FString,uint32> RenderCommands[RenderCommandCategoryCount];
 };
 
 ////////////////////////////////////////////////////////////////////////////////////////////////////
diff --git a/Engine/Source/Runtime/RenderCore/Public/RenderCommandTag.h b/Engine/Source/Runtime/RenderCore/Public/RenderCommandTag.h
--- a/Engine/Source/Runtime/RenderCore/Public/RenderCommandTag.h
+++ b/Engine/Source/Runtime/RenderCore/Public/RenderCommandTag.h
@@ -6,6 +6,17 @@
 
 ////////////////////////////////////////////////////////////////////////////////////////////////////
 
+enum class ERenderCommandCategory : uint8
+{
+	EnterTick, // Only allowed to execute when entering a tick from previous tick. Typically creating resources
+	LoopTick,  // Allowed to be executed multiple times if rendering is looping the same tick
+	LeaveTick, // Only allowed to be executed when leaving a tick to a new one. Typically releasing resources
+
+	Unknown,   // Unknown category. Will not work with state stream path. ALWAYS LAST
+};
+
+////////////////////////////////////////////////////////////////////////////////////////////////////
+
 class FRenderCommandTag
 {
 public:
@@ -24,15 +35,22 @@ public:
 		return SpecId;
 	}
 
+	ERenderCommandCategory GetCategory() const
+	{
+		return Category;
+	}
+
 protected:
-	FRenderCommandTag(const TCHAR* InName, TStatId InStatId)
+	FRenderCommandTag(const TCHAR* InName, TStatId InStatId, ERenderCommandCategory InCategory)
 		: Name(InName)
 		, StatId(InStatId)
+		, Category(InCategory)
 	{}
 
 private:
 	const TCHAR* Name;
 	TStatId StatId;
+	ERenderCommandCategory Category;
 	mutable uint32 SpecId;
 };
 
@@ -79,17 +97,17 @@ public:
 			}
 		};
 		static FThreadSafeStaticStat<FStatData> Stat;
-		static TRenderCommandTag Tag(TSTR::TStr(), Stat.GetStatId());
+		static TRenderCommandTag Tag(Stat.GetStatId());
 #else
-		static TRenderCommandTag Tag(TSTR::TStr(), TStatId());
+		static TRenderCommandTag Tag({});
 #endif
 
 		return Tag;
 	}
 
 private:
-	TRenderCommandTag(const TCHAR* InName, TStatId InStatId)
-		: FRenderCommandTag(InName, InStatId)
+	TRenderCommandTag(TStatId InStatId)
+		: FRenderCommandTag(TSTR::TStr(), InStatId, TSTR::GetCategory())
 	{}
 };
 
@@ -102,32 +120,33 @@ class TRenderCommandTagNoMarker : public FRenderCommandTag
 public:
 	static const TRenderCommandTagNoMarker& Get()
 	{
-		static TRenderCommandTagNoMarker Tag(TSTR::TStr(), TStatId());
+		static TRenderCommandTagNoMarker Tag;
 		return Tag;
 	}
 
 private:
-	TRenderCommandTagNoMarker(const TCHAR* InName, TStatId InStatId)
-		: FRenderCommandTag(InName, TStatId())
+	TRenderCommandTagNoMarker()
+		: FRenderCommandTag(TSTR::TStr(), TStatId(), TSTR::GetCategory())
 	{
 	}
 };
 
-
 /** Declares a new render command tag type from a name. */
-#define DECLARE_RENDER_COMMAND_TAG(Type, Name) \
+#define DECLARE_RENDER_COMMAND_TAG(Type, Name, ...) \
 	struct UE_JOIN(TSTR_, Name, __LINE__) \
 	{  \
 		static const char* CStr() { return #Name; } \
 		static const TCHAR* TStr() { return TEXT(#Name); } \
+		static constexpr ERenderCommandCategory GetCategory() { return __VA_OPT__(ERenderCommandCategory::__VA_ARGS__;) ERenderCommandCategory::Unknown; } \
 	}; \
 	using Type = TRenderCommandTag<UE_JOIN(TSTR_, Name, __LINE__)>;
 
-#define DECLARE_RENDER_COMMAND_TAG_NOMARKER(Type, Name) \
+#define DECLARE_RENDER_COMMAND_TAG_NOMARKER(Type, Name, ...) \
 	struct UE_JOIN(TSTR_, Name, __LINE__) \
 	{  \
 		static const char* CStr() { return #Name; } \
 		static const TCHAR* TStr() { return TEXT(#Name); } \
+		static constexpr ERenderCommandCategory GetCategory() { return __VA_OPT__(ERenderCommandCategory::__VA_ARGS__;) ERenderCommandCategory::Unknown; } \
 	}; \
 	using Type = TRenderCommandTagNoMarker<UE_JOIN(TSTR_, Name, __LINE__)>;
 
diff --git a/Engine/Source/Runtime/RenderCore/Public/RenderingThread.h b/Engine/Source/Runtime/RenderCore/Public/RenderingThread.h
--- a/Engine/Source/Runtime/RenderCore/Public/RenderingThread.h
+++ b/Engine/Source/Runtime/RenderCore/Public/RenderingThread.h
@@ -1082,12 +1082,12 @@ public:
  *  Alternative implementations accept either a reference or pointer to an FRenderCommandPipe instance to schedule on an async
  *  pipe, if enabled.
  */
-#define ENQUEUE_RENDER_COMMAND(Type) \
-	DECLARE_RENDER_COMMAND_TAG(UE_JOIN(FRenderCommandTag_, Type, __LINE__), Type) \
+#define ENQUEUE_RENDER_COMMAND(Type, ...) \
+	DECLARE_RENDER_COMMAND_TAG(UE_JOIN(FRenderCommandTag_, Type, __LINE__), Type, __VA_ARGS__) \
 	FRenderCommandDispatcher::Enqueue<UE_JOIN(FRenderCommandTag_, Type, __LINE__)>
 
-#define ENQUEUE_RENDER_COMMAND_NOMARKER(Type) \
-	DECLARE_RENDER_COMMAND_TAG_NOMARKER(UE_JOIN(FRenderCommandTag_, Type, __LINE__), Type) \
+#define ENQUEUE_RENDER_COMMAND_NOMARKER(Type, ...) \
+	DECLARE_RENDER_COMMAND_TAG_NOMARKER(UE_JOIN(FRenderCommandTag_, Type, __LINE__), Type, __VA_ARGS__) \
 	FRenderCommandDispatcher::Enqueue<UE_JOIN(FRenderCommandTag_, Type, __LINE__)>
 
 ////////////////////////////////////
