IncrementalCook: ShaderLibrary: When we assert because a deleted shadertype is referenced ("Shader library contains a reference to an FShaderType with hash %lld which does not exist"), add feedback about the asset that is referencing the shadertype.
#rnx

#Synced-CL 50149715

diff --git a/Engine/Source/Editor/UnrealEd/Classes/CookOnTheSide/CookOnTheFlyServer.h b/Engine/Source/Editor/UnrealEd/Classes/CookOnTheSide/CookOnTheFlyServer.h
--- a/Engine/Source/Editor/UnrealEd/Classes/CookOnTheSide/CookOnTheFlyServer.h
+++ b/Engine/Source/Editor/UnrealEd/Classes/CookOnTheSide/CookOnTheFlyServer.h
@@ -1117,7 +1117,7 @@ private:
 	void ShutdownShaderLibraryCookerAndCompilers(const FString& ProjectShaderLibraryName);
 
 	/** Log shader stats during CookByTheBookFinished. */
-	void DumpShaderTypeStats(const FString& PlatformNameString);
+	void DumpShaderTypeStats(const ITargetPlatform* TargetPlatform, const FString& PlatformNameString);
 
 	/**
 	 * Calls the ShaderPipelineCacheToolsCommandlet to build a upipelinecache file from the stable pipeline cache
diff --git a/Engine/Source/Editor/UnrealEd/Private/CookOnTheFlyServer.cpp b/Engine/Source/Editor/UnrealEd/Private/CookOnTheFlyServer.cpp
--- a/Engine/Source/Editor/UnrealEd/Private/CookOnTheFlyServer.cpp
+++ b/Engine/Source/Editor/UnrealEd/Private/CookOnTheFlyServer.cpp
@@ -9613,7 +9613,7 @@ void UCookOnTheFlyServer::CookByTheBookFinishedInternal()
 			// Save shader code map
 			SaveShaderLibrary(TargetPlatform, LibraryName);
 			CreatePipelineCache(TargetPlatform, LibraryName);
-			DumpShaderTypeStats(PlatformNameString);
+			DumpShaderTypeStats(TargetPlatform, PlatformNameString);
 		}
 
 		if (FParse::Param(FCommandLine::Get(), TEXT("fastcook")))
diff --git a/Engine/Source/Editor/UnrealEd/Private/Cooker/ShaderLibraryCooking.cpp b/Engine/Source/Editor/UnrealEd/Private/Cooker/ShaderLibraryCooking.cpp
--- a/Engine/Source/Editor/UnrealEd/Private/Cooker/ShaderLibraryCooking.cpp
+++ b/Engine/Source/Editor/UnrealEd/Private/Cooker/ShaderLibraryCooking.cpp
@@ -6,6 +6,7 @@
 #include "Commandlets/ShaderPipelineCacheToolsCommandlet.h"
 #include "Cooker/CookDirector.h"
 #include "Cooker/CookOnTheFlyServerInterface.h"
+#include "Cooker/CookPackageData.h"
 #include "Cooker/CookPlatformManager.h"
 #include "Cooker/CookTypes.h"
 #include "Cooker/CookWorkerClient.h"
@@ -635,11 +636,30 @@ void UCookOnTheFlyServer::ShutdownShaderLibraryCookerAndCompilers(const FString&
 	ShutdownShaderCompilers(PlatformManager->GetSessionPlatforms());
 }
 
-void UCookOnTheFlyServer::DumpShaderTypeStats(const FString& PlatformNameString)
+void UCookOnTheFlyServer::DumpShaderTypeStats(const ITargetPlatform* TargetPlatform, const FString& PlatformNameString)
 {
+	using namespace UE::Cook;
+
+	auto GetCookStatus = [this, TargetPlatform]
+	(FName AssetName, bool& bOutImplemented, bool& bOutReferenced, bool& bOutRecooked)
+		{
+			bOutImplemented = true;
+			bOutReferenced = false;
+			bOutRecooked = false;
+
+			FPackageData* PackageData = PackageDatas->FindPackageDataByPackageName(AssetName);
+			if (!PackageData)
+			{
+				return;
+			}
+			FPackagePlatformData& PlatformData = PackageData->FindOrAddPlatformData(TargetPlatform);
+			bOutReferenced = PlatformData.IsCookSucceeded();
+			bOutRecooked = !PlatformData.IsIncrementallySkipped();
+		};
 	FShaderLibraryCooker::DumpShaderTypeStats(
 		GShaderCompilingManager->GetAbsoluteShaderDebugInfoDirectory(),
-		ConvertToFullSandboxPath(GetMetadataDirectory(), true, PlatformNameString));
+		ConvertToFullSandboxPath(GetMetadataDirectory(), true, PlatformNameString),
+		GetCookStatus);
 }
 
 bool UCookOnTheFlyServer::RecompileChangedShaders(const TArray<const ITargetPlatform*>& TargetPlatforms)
diff --git a/Engine/Source/Runtime/RenderCore/Private/ShaderLibrary/ShaderCodeLibrary.cpp b/Engine/Source/Runtime/RenderCore/Private/ShaderLibrary/ShaderCodeLibrary.cpp
--- a/Engine/Source/Runtime/RenderCore/Private/ShaderLibrary/ShaderCodeLibrary.cpp
+++ b/Engine/Source/Runtime/RenderCore/Private/ShaderLibrary/ShaderCodeLibrary.cpp
@@ -1802,31 +1802,51 @@ public:
 		return bOK;
 	}
 
-	void DumpShaderTypeStats(const FString& DebugInfoDir, const FString& MetaDataDir)
+	/**
+	 * Returns true if the specified shader type is referenced by any asset for staging,
+	 * and optionally reports referencers.
+	 */
+	bool IsShaderTypeReferencedByStaging(int32 PlatformId, uint64 InShaderTypeHash,
+		TArray<FName>* OutAssetNames = nullptr, int32 MaxNumAssetNames = 1)
 	{
-		// Returns true if the specified shader type is referenced by any asset for staging
-		auto IsShaderTypeReferencedByStaging = [this](int32 PlatformId, uint64 InShaderTypeHash) -> bool
-			{
-				bool bOutIsReferencedByStaging = false;
-				if (FEditorShaderCodeArchive* ShaderArchive = this->EditorShaderCodeArchive[PlatformId])
+		bool bOutIsReferencedByStaging = false;
+		if (OutAssetNames)
+		{
+			OutAssetNames->Empty();
+		}
+		if (FEditorShaderCodeArchive* ShaderArchive = this->EditorShaderCodeArchive[PlatformId])
+		{
+			ShaderArchive->ForEachAssetReferencingShaderType(
+				InShaderTypeHash,
+				[&bOutIsReferencedByStaging, OutAssetNames, MaxNumAssetNames]
+				(FName AssetName, const FSerializedShaderArchive::FShaderMapAssetAssociations::FAssociatedAssetData& AssetData) -> bool
 				{
-					ShaderArchive->ForEachAssetReferencingShaderType(
-						InShaderTypeHash,
-						[&bOutIsReferencedByStaging](FName AssetName, const FSerializedShaderArchive::FShaderMapAssetAssociations::FAssociatedAssetData& AssetData) -> bool
+					if (AssetData.bReferencedByStaging)
+					{
+						bOutIsReferencedByStaging = true;
+						if (OutAssetNames)
 						{
-							if (AssetData.bReferencedByStaging)
+							OutAssetNames->Add(AssetName);
+							if (MaxNumAssetNames >= 1 && OutAssetNames->Num() >= MaxNumAssetNames)
 							{
-								// Mark shader type as referenced and stop iterating through further assets
-								bOutIsReferencedByStaging = true;
-								return false;
+								return false; // stop iterating
 							}
-							return true;
 						}
-					);
+						else
+						{
+							return false; // stop iterating
+						}
+					}
+					return true; // keep iterating
 				}
-				return bOutIsReferencedByStaging;
-			};
+			);
+		}
+		return bOutIsReferencedByStaging;
+	};
 
+	void DumpShaderTypeStats(const FString& DebugInfoDir, const FString& MetaDataDir,
+		TFunctionRef<void(FName AssetName, bool& bImplemented, bool& bOutReferenced, bool& bOutRecooked)> GetCookStatus)
+	{
 		for (int32 PlatformId = 0; PlatformId < EShaderPlatform::SP_NumPlatforms; ++PlatformId)
 		{
 			const FShaderCodeStats& PlatformCodeStats = EditorShaderCodeStats[PlatformId];
@@ -1846,9 +1866,45 @@ public:
 					{
 						// This error needs to be suppressed for data used by ShaderMaps or Shaders
 						// that are not referenced by any assets that have bReferencedByStaging.
-						checkf(!IsShaderTypeReferencedByStaging(PlatformId, Pair.Key),
-							TEXT("Shader library contains a reference to an FShaderType with hash %llu which does not exist.\n")
-							TEXT("If it has been removed or renamed, please bump CookIncrementalVersion to force a rebuild"), Pair.Key);
+						TArray<FName> AssetNames;
+						bool bReferencedbyStaging = IsShaderTypeReferencedByStaging(PlatformId, Pair.Key,
+							&AssetNames, 1 /* MaxNumAssetNames */);
+						if (bReferencedbyStaging)
+						{
+							FString CookStatus = TEXT("It is unknown whether this asset was incrementally skipped in the current cook.");
+							if (!AssetNames.IsEmpty())
+							{
+								bool bImplemented = false;
+								bool bReferenced = false;
+								bool bRecooked = false;
+								GetCookStatus(AssetNames[0], bImplemented, bReferenced, bRecooked);
+								if (bImplemented)
+								{
+									if (!bReferenced)
+									{
+										CookStatus = TEXT("This asset was not referenced in the current cook; ")
+											TEXT("there is a Cooker bug that reported it to the ShaderCodeLibrary as referenced by staging.");
+									}
+									else if (bRecooked)
+									{
+										CookStatus = TEXT("This asset was referenced and recooked in the current cook; ")
+											TEXT("there is a Cooker or ShaderCodeLibrary bug that failed to remove its now-stale previously cooked shader that uses the removed or renamed ShaderType.");
+									}
+									else
+									{
+										CookStatus = TEXT("This asset was referenced but was incrementally skipped in the current cook; ")
+											TEXT("the change that removed or renamed the ShaderType is missing a version bump that triggers the recook of the asset.");
+									}
+								}
+							}
+
+							checkf(false,
+								TEXT("Shader library contains a reference to an FShaderType with hash %" UINT64_FMT " which does not exist.\n")
+								TEXT("Removed or renamed ShaderTypes are supposed to be impossible to stage; recooks should not be able ")
+								TEXT("to reference them and incremental cooks are supposed to recook assets that use them.\n")
+								TEXT("This shadertype however is referenced by an asset which is being staged: \"%s\". %s"),
+								Pair.Key, AssetNames.IsEmpty() ? TEXT("<Unknown>") : *AssetNames[0].ToString(), *CookStatus);
+						}
 					}
 				}
 				
@@ -2420,11 +2476,12 @@ void FShaderLibraryCooker::AddShaderStableKeyValue(EShaderPlatform ShaderPlatfor
 	}
 }
 
-void FShaderLibraryCooker::DumpShaderTypeStats(const FString& DebugInfoDir, const FString& MetaDataDir)
+void FShaderLibraryCooker::DumpShaderTypeStats(const FString& DebugInfoDir, const FString& MetaDataDir,
+	TFunctionRef<void(FName AssetName, bool& bImplemented, bool& bOutReferenced, bool& bOutRecooked)> GetCookStatus)
 {
 	if (FShaderLibrariesCollection::Impl)
 	{
-		FShaderLibrariesCollection::Impl->DumpShaderTypeStats(DebugInfoDir, MetaDataDir);
+		FShaderLibrariesCollection::Impl->DumpShaderTypeStats(DebugInfoDir, MetaDataDir, GetCookStatus);
 	}
 }
 
diff --git a/Engine/Source/Runtime/RenderCore/Public/ShaderCodeLibrary.h b/Engine/Source/Runtime/RenderCore/Public/ShaderCodeLibrary.h
--- a/Engine/Source/Runtime/RenderCore/Public/ShaderCodeLibrary.h
+++ b/Engine/Source/Runtime/RenderCore/Public/ShaderCodeLibrary.h
@@ -23,6 +23,7 @@
 #include "RHI.h"
 #include "RHIDefinitions.h"
 #include "Serialization/Archive.h"
+#include "Templates/Function.h"
 #include "Templates/RefCounting.h"
 #include "UObject/NameTypes.h"
 
@@ -410,7 +411,8 @@ struct FShaderLibraryCooker
 	static RENDERCORE_API void DeleteShaderLibraryChunks(const ITargetPlatform* TargetPlatform,
 		const FString& SandboxDestinationPath, const FString& SandboxMetadataPath);
 
-	static RENDERCORE_API void DumpShaderTypeStats(const FString& DebugInfoPath, const FString& MetaDataPath);
+	static RENDERCORE_API void DumpShaderTypeStats(const FString& DebugInfoPath, const FString& MetaDataPath,
+		TFunctionRef<void(FName AssetName, bool& bImplemented, bool& bOutReferenced, bool& bOutRecooked)> GetCookStatus);
 
 	// Create a smaller 'patch' library that only contains data from 'NewMetaDataDir' not contained in any of 'OldMetaDataDirs'
 	static RENDERCORE_API bool CreatePatchLibrary(TArray<FString> const& OldMetaDataDirs, FString const& NewMetaDataDir, FString const& OutDir, bool bNativeFormat, bool bNeedsDeterministicOrder);
