Fix check(OptionalDataSize == -1) in FShaderCode::GetFinalizedResource when we fail to compile the first stage of a pipeline job: the other stages were skipped, setting only Target/ValidateInputHash but not calling ShaderCode.FinalizeShaderCode();
This happens in FShaderCompileDistributedThreadRunnable_Interface::CompilingLoop, when bCompileJobsSucceeded=false: we try to call FShaderCompileUtilities::ExecuteShaderCompileJob locally, ending up in CompileShaderPipelineStages. Subsequent stages are skipped, FinalizeShaderCode not being called on them. ProcessFinishedJob would eventually call GetFinalizedResource triggering the check fail

The work is now centralized in FinalizeShaderCompileJob to handle all cases

#Synced-CL 0

diff --git a/Engine/Source/Runtime/RenderCore/Private/ShaderCore.cpp b/Engine/Source/Runtime/RenderCore/Private/ShaderCore.cpp
--- a/Engine/Source/Runtime/RenderCore/Private/ShaderCore.cpp
+++ b/Engine/Source/Runtime/RenderCore/Private/ShaderCore.cpp
@@ -1975,6 +1975,14 @@ public:
 #endif
 	}
 
+	static void FinalizeShaderCompileJob(FShaderCompilerOutput& Output, const FShaderCompilerInput& Input)
+	{
+		Output.Target = Input.Target;
+		// also set the job input hash on the output struct for validation purposes
+		Output.ValidateInputHash = Input.Hash;
+		Output.ShaderCode.FinalizeShaderCode();
+	}
+
 	static void CompileShaderInternal(const IShaderFormat* Compiler, FShaderCompileJob& Job, FString& OutExceptionMsg, FString& OutExceptionCallstack, int32* CompileCount)
 	{
 		double TimeStart = FPlatformTime::Seconds();
@@ -1997,10 +2005,7 @@ public:
 		{
 			Job.Output.bSucceeded = Job.Output.bSucceeded && Job.SecondaryOutput->bSucceeded;
 			// ensure the target field is set on the job output struct as we use it for validation during serialization
-			Job.SecondaryOutput->Target = Job.Input.Target;
-			// also set the job input hash on the output struct for validation purposes
-			Job.SecondaryOutput->ValidateInputHash = Job.Input.Hash;
-			Job.SecondaryOutput->ShaderCode.FinalizeShaderCode();
+			FinalizeShaderCompileJob(*Job.SecondaryOutput, Job.Input);
 			if (Job.SecondaryOutput->bSucceeded)
 			{
 				Job.SecondaryOutput->GenerateOutputHash();
@@ -2008,11 +2013,7 @@ public:
 			Job.Output.bSucceeded &= CombineCodeOutputs(Compiler, Job);
 		}
 
-		// ensure the target field is set on the job output struct as we use it for validation during serialization
-		Job.Output.Target = Job.Input.Target;
-		// also set the job input hash on the output struct for validation purposes
-		Job.Output.ValidateInputHash = Job.Input.Hash;
-		Job.Output.ShaderCode.FinalizeShaderCode();
+		FinalizeShaderCompileJob(Job.Output, Job.Input);
 
 		// clear out the modified source/entry point fields if they aren't needed to avoid the data transfer/ddc overhead
 		if (!Job.Input.NeedsOriginalShaderSource())
@@ -2152,8 +2153,11 @@ static bool CompileShaderPipelineStages(const TArray<const IShaderFormat*>& Shad
 			for (int32 SkipIndex = Index + 1; SkipIndex < PipelineJob->StageJobs.Num(); SkipIndex++)
 			{
 				FShaderCompileJob& Job = *PipelineJob->StageJobs[SkipIndex];
-				Job.Output.Target = Job.Input.Target;
-				Job.Output.ValidateInputHash = Job.Input.Hash;
+				FInternalShaderCompilerFunctions::FinalizeShaderCompileJob(Job.Output, Job.Input);
+				if (Job.SecondaryOutput.IsValid())
+				{
+					FInternalShaderCompilerFunctions::FinalizeShaderCompileJob(*Job.SecondaryOutput, Job.Input);
+				}
 			}
 			return false;
 		}
