Replace r.HDR.UI.UseOSPaperWhite with r.HDR.UI.Luminance.Mode and add r.HDR.ReferenceWhite/r.HDR.ReferenceWhite.Mode

r.HDR.ReferenceWhite/r.HDR.ReferenceWhite.Mode control paper whtie values for scene content, and r.HDR.UI.Luminance/r.HDR.UI.Luminance.Mode control paper white values for UI content.

r.HDR.UI.Luminance is being demoted a bit as it will not be used unless r.HDR.UI.Luminance.Mode is set to non-default 1 to allow both scene and UI content to be controlled by one value by default.

#Synced-CL 50463326

diff --git a/Engine/Source/Runtime/RenderCore/Private/RenderCore.cpp b/Engine/Source/Runtime/RenderCore/Private/RenderCore.cpp
--- a/Engine/Source/Runtime/RenderCore/Private/RenderCore.cpp
+++ b/Engine/Source/Runtime/RenderCore/Private/RenderCore.cpp
@@ -439,12 +439,28 @@ static TAutoConsoleVariable<int32> CVarHDRAcesVersion(
 TAutoConsoleVariable<int32> CVarHDROutputEnabled(
 	TEXT("r.HDR.EnableHDROutput"),
 	0,
-	TEXT("Creates an HDR compatible swap-chain and enables HDR display output.")
+	TEXT("Creates an HDR compatible swap-chain and enables HDR display output.\n")
 	TEXT("0: Disabled (default)\n")
 	TEXT("1: Enable hardware-specific implementation\n"),
 	ECVF_RenderThreadSafe
 );
 
+TAutoConsoleVariable<float> CVarHDRReferenceWhite(
+	TEXT("r.HDR.ReferenceWhite"),
+	203.f,
+	TEXT("Luminance in nits to redefine HDR Reference White as specified in BT.2408."),
+	ECVF_RenderThreadSafe
+);
+
+TAutoConsoleVariable<int32> CVarHDRReferenceWhiteMode(
+	TEXT("r.HDR.ReferenceWhite.Mode"),
+	0,
+	TEXT("Specifies how scene content incorporates paper white values.\n")
+	TEXT("0: Use the platform's paper white if available, otherwise r.HDR.ReferenceWhite (default)\n")
+	TEXT("1: Use r.HDR.ReferenceWhite"),
+	ECVF_RenderThreadSafe
+);
+
 RENDERCORE_API bool IsHDREnabled()
 {
 	return GRHISupportsHDROutput && CVarHDROutputEnabled.GetValueOnAnyThread() != 0;
@@ -645,7 +661,16 @@ RENDERCORE_API void HDRGetMetaDataForWindow(FHDRMetaData& OutHDRMetaData, const
 	OutHDRMetaData.MinimumLuminanceInNits = OutHDRMetaData.bHDRSupported ? DisplayList[BestDisplayForWindow].MinimumLuminanceInNits : 0.f;
 	OutHDRMetaData.MaximumLuminanceInNits = OutHDRMetaData.bHDRSupported ? DisplayList[BestDisplayForWindow].MaximumLuminanceInNits : 100.f;
 	OutHDRMetaData.MaximumFullFrameLuminanceInNits = OutHDRMetaData.bHDRSupported ? DisplayList[BestDisplayForWindow].MaximumFullFrameLuminanceInNits : 100.f;
-	OutHDRMetaData.HDRPaperWhiteInNits = DisplayList[BestDisplayForWindow].HDRPaperWhiteInNits;
+#if PLATFORM_PROVIDES_HDR_PAPER_WHITE
+	if (CVarHDRReferenceWhiteMode.GetValueOnAnyThread() == 0)
+	{
+		OutHDRMetaData.HDRPaperWhiteInNits = DisplayList[BestDisplayForWindow].HDRPaperWhiteInNits;
+	}
+	else
+#endif
+	{
+		OutHDRMetaData.HDRPaperWhiteInNits = CVarHDRReferenceWhite.GetValueOnAnyThread();
+	}
 	OutHDRMetaData.LimitingColorSpace = OutHDRMetaData.bHDRSupported ? DisplayList[BestDisplayForWindow].LimitingColorSpace : UE::Color::FColorSpace::GetSRGB();
 
 	if (OutHDRMetaData.bHDRSupported)
diff --git a/Engine/Source/Runtime/Renderer/Private/PostProcess/PostProcessSelectionOutline.cpp b/Engine/Source/Runtime/Renderer/Private/PostProcess/PostProcessSelectionOutline.cpp
--- a/Engine/Source/Runtime/Renderer/Private/PostProcess/PostProcessSelectionOutline.cpp
+++ b/Engine/Source/Runtime/Renderer/Private/PostProcess/PostProcessSelectionOutline.cpp
@@ -380,8 +380,8 @@ FScreenPassTexture AddSelectionOutlinePass(
 		}
 
 		float UILuminance = 203.f;
-		static const auto* const CVarHDRUIUseOSPaperWhite = IConsoleManager::Get().FindTConsoleVariableDataBool(TEXT("r.HDR.UI.UseOSPaperWhite"));
-		if (CVarHDRUIUseOSPaperWhite && CVarHDRUIUseOSPaperWhite->GetValueOnRenderThread())
+		static const auto* const CVarHDRUILuminanceMode = IConsoleManager::Get().FindTConsoleVariableDataInt(TEXT("r.HDR.UI.Luminance.Mode"));
+		if (CVarHDRUILuminanceMode && (CVarHDRUILuminanceMode->GetValueOnRenderThread() == 0))
 		{
 			UILuminance = View.Family->RenderTarget->GetHDRPaperWhiteInNits();
 		}
diff --git a/Engine/Source/Runtime/SlateRHIRenderer/Private/SlatePostProcessor.cpp b/Engine/Source/Runtime/SlateRHIRenderer/Private/SlatePostProcessor.cpp
--- a/Engine/Source/Runtime/SlateRHIRenderer/Private/SlatePostProcessor.cpp
+++ b/Engine/Source/Runtime/SlateRHIRenderer/Private/SlatePostProcessor.cpp
@@ -34,10 +34,10 @@ float GetSlateHDRUILuminance()
 	return CVar ? CVar->GetFloat() : 203.f;
 }
 
-bool GetSlateHDRUIUseOSPaperWhite()
+int32 GetSlateHDRUILuminanceMode()
 {
-	static auto CVar = IConsoleManager::Get().FindConsoleVariable(TEXT("r.HDR.UI.UseOSPaperWhite"));
-	return CVar ? CVar->GetBool() : false;
+	static auto CVar = IConsoleManager::Get().FindConsoleVariable(TEXT("r.HDR.UI.Luminance.Mode"));
+	return CVar ? CVar->GetInt() : 1;
 }
 
 int32 GetSlateHDRUICompositeEOTF()
@@ -161,7 +161,7 @@ FScreenPassTexture AddSlateCompositeHDRForBlurPass(FRDGBuilder& GraphBuilder, co
 	PassParameters->UISampler = TStaticSamplerState<SF_Bilinear>::GetRHI();
 	PassParameters->UITextureSize = InputViewport.Extent;
 	PassParameters->UILevel = GetSlateHDRUILevel();
-	PassParameters->UILuminance = GetSlateHDRUIUseOSPaperWhite() ? Inputs.HDRPaperWhiteInNits : GetSlateHDRUILuminance();
+	PassParameters->UILuminance = GetSlateHDRUILuminanceMode() ? GetSlateHDRUILuminance() : Inputs.HDRPaperWhiteInNits;
 	PassParameters->UICompositeEOTF = GetSlateHDRUICompositeEOTF();
 	PassParameters->UICompositeEOTFGamma = GetSlateHDRUICompositeEOTFGamma();
 	PassParameters->UICompositeBlend = GetSlateHDRUICompositeBlend();
diff --git a/Engine/Source/Runtime/SlateRHIRenderer/Private/SlateRHIRenderer.cpp b/Engine/Source/Runtime/SlateRHIRenderer/Private/SlateRHIRenderer.cpp
--- a/Engine/Source/Runtime/SlateRHIRenderer/Private/SlateRHIRenderer.cpp
+++ b/Engine/Source/Runtime/SlateRHIRenderer/Private/SlateRHIRenderer.cpp
@@ -88,15 +88,14 @@ static TAutoConsoleVariable<float> CVarHDRUILuminance(
 	TEXT("Base Luminance in nits for UI elements when compositing into HDR framebuffer. Gets multiplied by r.HDR.UI.Level"),
 	ECVF_RenderThreadSafe);
 
-static TAutoConsoleVariable<bool> CVarHDRUIUseOSPaperWhite(
-	TEXT("r.HDR.UI.UseOSPaperWhite"),
-#if PLATFORM_PROVIDES_HDR_PAPER_WHITE
-	true,
-#else
-	false,
-#endif
-	TEXT("Whether to use the paper white value provided by the OS for UI luminance."),
-	ECVF_RenderThreadSafe);
+TAutoConsoleVariable<int32> CVarHDRUILuminanceMode(
+	TEXT("r.HDR.UI.Luminance.Mode"),
+	0,
+	TEXT("Specifies how UI content incorporates paper white values.\n")
+	TEXT("0: Use the white value according to r.HDR.ReferenceWhite.Mode (default)\n")
+	TEXT("1: Use r.HDR.UI.Luminance"),
+	ECVF_RenderThreadSafe
+);
 
 static TAutoConsoleVariable<int32> CVarHDRUICompositeEOTF(
 	TEXT("r.HDR.UI.CompositeEOTF"),
@@ -874,9 +873,9 @@ void FSlateRHIRenderer::DrawWindowViewport_RenderThread(ISlateViewport* SlateVie
 		CommonParameters.UITexture = ElementsTexture;
 		CommonParameters.UISampler = TStaticSamplerState<SF_Point>::GetRHI();
 		CommonParameters.UILevel = CVarUILevel.GetValueOnRenderThread();
-		CommonParameters.UILuminance = CVarHDRUIUseOSPaperWhite.GetValueOnRenderThread()
-			? ViewportInfo.HDRPaperWhiteInNits
-			: CVarHDRUILuminance.GetValueOnRenderThread();
+		CommonParameters.UILuminance = CVarHDRUILuminanceMode.GetValueOnRenderThread()
+				? CVarHDRUILuminance.GetValueOnRenderThread()
+				: ViewportInfo.HDRPaperWhiteInNits;
 		CommonParameters.UICompositeEOTF = CVarHDRUICompositeEOTF.GetValueOnRenderThread();
 		CommonParameters.UICompositeEOTFGamma = CVarHDRUICompositeEOTFGamma.GetValueOnRenderThread();
 		CommonParameters.UICompositeBlend = CVarHDRUICompositeBlend.GetValueOnRenderThread();
