Update paper white support again

- Rename reference white to paper white to avoid confusion.
- Implement paper white support for PostProcessDeviceEncodingOnly.

#Synced-CL 50479750

diff --git a/Engine/Shaders/Private/PostProcessDeviceEncodingOnly.usf b/Engine/Shaders/Private/PostProcessDeviceEncodingOnly.usf
--- a/Engine/Shaders/Private/PostProcessDeviceEncodingOnly.usf
+++ b/Engine/Shaders/Private/PostProcessDeviceEncodingOnly.usf
@@ -27,6 +27,9 @@ uint OutputDevice;
 uint OutputGamut;
 float OutputMaxLuminance;
 
+float HDRPaperWhiteRatio;
+float InvHDRPaperWhiteRatio;
+
 float4 ACESMinMaxData;
 float4 ACESMidData;
 float4 ACESCoefsLow_0;
@@ -112,7 +115,7 @@ float4 DeviceEncodingOnlyCommonPS(
 		float3 ODTColor = ACESOutputTransform(
 			FilmColor * ACESSceneColorMultiplier,
 			(float3x3)WorkingColorSpace.ToAP0,
-			OutputMaxLuminance,
+			OutputMaxLuminance * InvHDRPaperWhiteRatio,
 			ACESReachTable,
 			ACESGamutTable,
 			ACESGammaTable,
@@ -120,6 +123,7 @@ float4 DeviceEncodingOnlyCommonPS(
 			LimitingColorSpace.RGBtoXYZ,
 			LimitingColorSpace.XYZtoRGB,
 			EncodingColorSpace);
+		ODTColor *= HDRPaperWhiteRatio;
 #else
 		FACESTonemapParams AcesParams = ComputeACESTonemapParams(ACESMinMaxData, ACESMidData, ACESCoefsLow_0, ACESCoefsHigh_0, ACESCoefsLow_4, ACESCoefsHigh_4, ACESSceneColorMultiplier, ACESGamutCompression);
 		float3 ODTColor = ACESOutputTransform( FilmColor, (float3x3)WorkingColorSpace.ToAP0, AcesParams);
@@ -139,7 +143,7 @@ float4 DeviceEncodingOnlyCommonPS(
 		float3 ODTColor = ACESOutputTransform(
 			FilmColor * ACESSceneColorMultiplier,
 			(float3x3)WorkingColorSpace.ToAP0,
-			OutputMaxLuminance,
+			OutputMaxLuminance * InvHDRPaperWhiteRatio,
 			ACESReachTable,
 			ACESGamutTable,
 			ACESGammaTable,
@@ -147,6 +151,7 @@ float4 DeviceEncodingOnlyCommonPS(
 			LimitingColorSpace.RGBtoXYZ,
 			LimitingColorSpace.XYZtoRGB,
 			GetColorSpace(TONEMAPPER_GAMUT_sRGB_D65));
+		ODTColor *= HDRPaperWhiteRatio;
 		OutDeviceColor = ODTColor / UE_SCRGB_WHITE_NITS;
 #else
 		FACESTonemapParams AcesParams = ComputeACESTonemapParams(ACESMinMaxData, ACESMidData, ACESCoefsLow_0, ACESCoefsHigh_0, ACESCoefsLow_4, ACESCoefsHigh_4, ACESSceneColorMultiplier, ACESGamutCompression);
diff --git a/Engine/Source/Runtime/RenderCore/Private/RenderCore.cpp b/Engine/Source/Runtime/RenderCore/Private/RenderCore.cpp
--- a/Engine/Source/Runtime/RenderCore/Private/RenderCore.cpp
+++ b/Engine/Source/Runtime/RenderCore/Private/RenderCore.cpp
@@ -417,19 +417,19 @@ TAutoConsoleVariable<int32> CVarHDROutputEnabled(
 	ECVF_RenderThreadSafe
 );
 
-TAutoConsoleVariable<float> CVarHDRReferenceWhite(
-	TEXT("r.HDR.ReferenceWhite"),
+TAutoConsoleVariable<float> CVarHDRPaperWhite(
+	TEXT("r.HDR.PaperWhite"),
 	203.f,
 	TEXT("Luminance in nits to redefine HDR Reference White as specified in BT.2408."),
 	ECVF_RenderThreadSafe
 );
 
-TAutoConsoleVariable<int32> CVarHDRReferenceWhiteMode(
-	TEXT("r.HDR.ReferenceWhite.Mode"),
+TAutoConsoleVariable<int32> CVarHDRPaperWhiteMode(
+	TEXT("r.HDR.PaperWhite.Mode"),
 	0,
 	TEXT("Specifies how scene content incorporates paper white values.\n")
-	TEXT("0: Use the platform's paper white if available, otherwise r.HDR.ReferenceWhite (default)\n")
-	TEXT("1: Use r.HDR.ReferenceWhite"),
+	TEXT("0: Use the platform's paper white if available, otherwise r.HDR.PaperWhite (default)\n")
+	TEXT("1: Use r.HDR.PaperWhite"),
 	ECVF_RenderThreadSafe
 );
 
@@ -634,14 +634,14 @@ RENDERCORE_API void HDRGetMetaDataForWindow(FHDRMetaData& OutHDRMetaData, const
 	OutHDRMetaData.MaximumLuminanceInNits = OutHDRMetaData.bHDRSupported ? DisplayList[BestDisplayForWindow].MaximumLuminanceInNits : 100.f;
 	OutHDRMetaData.MaximumFullFrameLuminanceInNits = OutHDRMetaData.bHDRSupported ? DisplayList[BestDisplayForWindow].MaximumFullFrameLuminanceInNits : 100.f;
 #if PLATFORM_PROVIDES_HDR_PAPER_WHITE
-	if (CVarHDRReferenceWhiteMode.GetValueOnAnyThread() == 0)
+	if (CVarHDRPaperWhiteMode.GetValueOnAnyThread() == 0)
 	{
 		OutHDRMetaData.HDRPaperWhiteInNits = DisplayList[BestDisplayForWindow].HDRPaperWhiteInNits;
 	}
 	else
 #endif
 	{
-		OutHDRMetaData.HDRPaperWhiteInNits = CVarHDRReferenceWhite.GetValueOnAnyThread();
+		OutHDRMetaData.HDRPaperWhiteInNits = CVarHDRPaperWhite.GetValueOnAnyThread();
 	}
 	OutHDRMetaData.LimitingColorSpace = OutHDRMetaData.bHDRSupported ? DisplayList[BestDisplayForWindow].LimitingColorSpace : UE::Color::FColorSpace::GetSRGB();
 
diff --git a/Engine/Source/Runtime/Renderer/Private/PostProcess/PostProcessDeviceEncodingOnly.cpp b/Engine/Source/Runtime/Renderer/Private/PostProcess/PostProcessDeviceEncodingOnly.cpp
--- a/Engine/Source/Runtime/Renderer/Private/PostProcess/PostProcessDeviceEncodingOnly.cpp
+++ b/Engine/Source/Runtime/Renderer/Private/PostProcess/PostProcessDeviceEncodingOnly.cpp
@@ -114,6 +114,8 @@ BEGIN_SHADER_PARAMETER_STRUCT(FDeviceEncodingOnlyParameters, )
 	SHADER_PARAMETER_RDG_TEXTURE(Texture2D, ColorTexture)
 	SHADER_PARAMETER_SAMPLER(SamplerState, ColorSampler)
 	SHADER_PARAMETER(float, EditorNITLevel)
+	SHADER_PARAMETER(float, HDRPaperWhiteRatio)
+	SHADER_PARAMETER(float, InvHDRPaperWhiteRatio)
 END_SHADER_PARAMETER_STRUCT()
 
 class FDeviceEncodingOnlyPS : public FGlobalShader
@@ -250,6 +252,10 @@ FScreenPassTexture AddDeviceEncodingOnlyPass(FRDGBuilder& GraphBuilder, const FV
 	CommonParameters.ColorSampler = BilinearClampSampler;
 	CommonParameters.EditorNITLevel = EditorNITLevel;
 
+	const float HDRPaperWhiteRatio = ViewFamily.RenderTarget->GetHDRPaperWhiteInNits() / 203.f;
+	CommonParameters.HDRPaperWhiteRatio = HDRPaperWhiteRatio;
+	CommonParameters.InvHDRPaperWhiteRatio = 1.f / HDRPaperWhiteRatio;
+
 	// Generate permutation vector for the desktop tonemapper.
 	DeviceEncodingOnlyPermutation::FDesktopDomain DesktopPermutationVector;
 
diff --git a/Engine/Source/Runtime/SlateRHIRenderer/Private/SlateRHIRenderer.cpp b/Engine/Source/Runtime/SlateRHIRenderer/Private/SlateRHIRenderer.cpp
--- a/Engine/Source/Runtime/SlateRHIRenderer/Private/SlateRHIRenderer.cpp
+++ b/Engine/Source/Runtime/SlateRHIRenderer/Private/SlateRHIRenderer.cpp
@@ -92,7 +92,7 @@ TAutoConsoleVariable<int32> CVarHDRUILuminanceMode(
 	TEXT("r.HDR.UI.Luminance.Mode"),
 	0,
 	TEXT("Specifies how UI content incorporates paper white values.\n")
-	TEXT("0: Use the white value according to r.HDR.ReferenceWhite.Mode (default)\n")
+	TEXT("0: Use the paper white value according to r.HDR.PaperWhite.Mode (default)\n")
 	TEXT("1: Use r.HDR.UI.Luminance"),
 	ECVF_RenderThreadSafe
 );
