[IO] Allow internal access to name of backend that was assgined to IoRequest for logging purposes.

This change makes use of the new internal debug information to log the assigned backend for each shader we create if LogShaderLibrary has VeryVerbose logging enabled.
Additionally, write out debug information in `DumpShaderCodeArchiveHashes()` in UTF-8; Some POSIX tools such as `grep` don't handle UTF-16 files well and it's a waste of memory for this information.

#rnx

#Synced-CL 50573404

diff --git a/Engine/Source/Runtime/Core/Internal/IO/IoDispatcherInternal.h b/Engine/Source/Runtime/Core/Internal/IO/IoDispatcherInternal.h
--- a/Engine/Source/Runtime/Core/Internal/IO/IoDispatcherInternal.h
+++ b/Engine/Source/Runtime/Core/Internal/IO/IoDispatcherInternal.h
@@ -1,31 +1,23 @@
 // Copyright Epic Games, Inc. All Rights Reserved.
 
 #pragma once
+#include "HAL/Platform.h"
 #include "IO/IoDispatcher.h"
 #include "IO/IoChunkId.h"
 #include "Containers/StringFwd.h"
 
 // This header is used for IoDispatcher functionality we don't want to expose outside of the engine
 
+struct FIoRequestDebugInfo
+{
+	const TCHAR* BackendName = nullptr;
+};
+
 class FIoDispatcherInternal final
 {
 public:
-	static bool HasPackageData()
-	{
-		// Checking for the global script objects chunk is currently the best means 
-		// to know if the IoDispatcher will be loading packaged data
-		const FIoDispatcher& Dispatcher = FIoDispatcher::Get();
-		static const bool HasScriptObjectsChunk = Dispatcher.DoesChunkExist(CreateIoChunkId(0, 0, EIoChunkType::ScriptObjects));
-		return HasScriptObjectsChunk;
-	}
-	
-	static void MountContainerMeta(const FString& FilePath)
-	{
-		FIoDispatcher::Get().MountContainerMeta(FilePath);
-	}
-
-	static FUtf8StringView GetFilename(const FIoChunkId& ChunkId, FUtf8StringBuilderBase& OutFilename, FUtf8StringView& OutContainerName)
-	{
-		return FIoDispatcher::Get().GetFilename(ChunkId, OutFilename, OutContainerName);
-	}
+	static CORE_API bool HasPackageData();
+	static CORE_API void MountContainerMeta(const FString& FilePath);
+	static CORE_API FUtf8StringView GetFilename(const FIoChunkId& ChunkId, FUtf8StringBuilderBase& OutFilename, FUtf8StringView& OutContainerName);
+	static CORE_API void GetDebugInfo(const FIoRequest& IoRequest, FIoRequestDebugInfo& OutDebugInfo);
 };
diff --git a/Engine/Source/Runtime/Core/Private/IO/IoDispatcher.cpp b/Engine/Source/Runtime/Core/Private/IO/IoDispatcher.cpp
--- a/Engine/Source/Runtime/Core/Private/IO/IoDispatcher.cpp
+++ b/Engine/Source/Runtime/Core/Private/IO/IoDispatcher.cpp
@@ -12,6 +12,7 @@
 #include "IO/IoDispatcherBackend.h"
 #include "IO/IoDispatcherConfig.h"
 #include "IO/IoDispatcherChunkDecoder.h"
+#include "IO/IoDispatcherInternal.h"
 #include "IO/IoDispatcherPrivate.h"
 #include "IO/IoOffsetLength.h"
 #include "IO/PlatformIoDispatcher.h"
@@ -1484,3 +1485,34 @@ FIoRequest::Release()
 		Impl = nullptr;
 	}
 }
+
+bool
+FIoDispatcherInternal::HasPackageData()
+{
+	// Checking for the global script objects chunk is currently the best means 
+	// to know if the IoDispatcher will be loading packaged data
+	const FIoDispatcher& Dispatcher = FIoDispatcher::Get();
+	static const bool HasScriptObjectsChunk = Dispatcher.DoesChunkExist(CreateIoChunkId(0, 0, EIoChunkType::ScriptObjects));
+	return HasScriptObjectsChunk;
+}
+
+void
+FIoDispatcherInternal::MountContainerMeta(const FString& FilePath)
+{
+	FIoDispatcher::Get().MountContainerMeta(FilePath);
+}
+
+FUtf8StringView
+FIoDispatcherInternal::GetFilename(const FIoChunkId& ChunkId, FUtf8StringBuilderBase& OutFilename, FUtf8StringView& OutContainerName)
+{
+	return FIoDispatcher::Get().GetFilename(ChunkId, OutFilename, OutContainerName);
+}
+
+void
+FIoDispatcherInternal::GetDebugInfo(const FIoRequest& IoRequest, FIoRequestDebugInfo& OutDebugInfo)
+{
+	if (IoRequest.Impl && IoRequest.Impl->Backend)
+	{
+		OutDebugInfo.BackendName = IoRequest.Impl->Backend->GetName();
+	}
+}
diff --git a/Engine/Source/Runtime/Core/Public/IO/IoDispatcher.h b/Engine/Source/Runtime/Core/Public/IO/IoDispatcher.h
--- a/Engine/Source/Runtime/Core/Public/IO/IoDispatcher.h
+++ b/Engine/Source/Runtime/Core/Public/IO/IoDispatcher.h
@@ -261,6 +261,7 @@ private:
 
 	friend class FIoDispatcher;
 	friend class FIoDispatcherImpl;
+	friend class FIoDispatcherInternal;
 	friend class FIoBatch;
 };
 
diff --git a/Engine/Source/Runtime/Core/Public/IO/IoDispatcherBackend.h b/Engine/Source/Runtime/Core/Public/IO/IoDispatcherBackend.h
--- a/Engine/Source/Runtime/Core/Public/IO/IoDispatcherBackend.h
+++ b/Engine/Source/Runtime/Core/Public/IO/IoDispatcherBackend.h
@@ -106,6 +106,7 @@ public:
 
 private:
 	friend class FIoDispatcherImpl;
+	friend class FIoDispatcherInternal;
 	friend class FIoRequest;
 	friend class FIoBatch;
 	friend class FIoRequestStats;
diff --git a/Engine/Source/Runtime/RenderCore/Private/ShaderLibrary/ShaderCodeArchive.cpp b/Engine/Source/Runtime/RenderCore/Private/ShaderLibrary/ShaderCodeArchive.cpp
--- a/Engine/Source/Runtime/RenderCore/Private/ShaderLibrary/ShaderCodeArchive.cpp
+++ b/Engine/Source/Runtime/RenderCore/Private/ShaderLibrary/ShaderCodeArchive.cpp
@@ -23,6 +23,7 @@
 #include "Serialization/StaticMemoryReader.h"
 #include "ProfilingDebugging/IoStoreTrace.h"
 #include "IO/IoDispatcherBackend.h"
+#include "IO/IoDispatcherInternal.h"
 
 #if WITH_EDITORONLY_DATA
 #include "Misc/Optional.h"
@@ -2289,44 +2290,42 @@ void FIoStoreShaderCodeArchive::SaveIoStoreShaderCodeArchive(const FIoStoreShade
 }
 
 // Dumps all shader, shader map, and shader group hashes and how they are connected to a text file for debugging purposes.
-static void DumpShaderCodeArchiveHashes(const FIoStoreShaderCodeArchiveHeader& InLibraryHeader, const TCHAR* InFilename)
+static void DumpShaderCodeArchiveHashes(const FIoStoreShaderCodeArchiveHeader& InLibraryHeader, FArchive& Ar)
 {
-	if (TUniquePtr<FArchive> DumpFile = TUniquePtr<FArchive>(IFileManager::Get().CreateFileWriter(InFilename)))
-	{
-		FString DumpContent;
-
-		DumpContent.Appendf(TEXT("--------------- ShaderGroupIoHashes: ---------------\n"));
-		for (int32 GroupIndex = 0; GroupIndex < InLibraryHeader.ShaderGroupIoHashes.Num(); ++GroupIndex)
-		{
-			const FIoChunkId& ShaderGroupChunkId = InLibraryHeader.ShaderGroupIoHashes[GroupIndex];
-			const FIoStoreShaderGroupEntry& GroupEntry = InLibraryHeader.ShaderGroupEntries[GroupIndex];
-			DumpContent.Appendf(TEXT("%s (%d shaders)\n"), *LexToString(ShaderGroupChunkId), GroupEntry.NumShaders);
+	FString DebugOutput;
 
-			for (uint32 ShaderIndex = 0; ShaderIndex < GroupEntry.NumShaders; ++ShaderIndex)
-			{
-				const int32 MappedShaderIndex = InLibraryHeader.ShaderIndices[GroupEntry.ShaderIndicesOffset + ShaderIndex];
-				DumpContent.Appendf(TEXT("  %s\n"), *InLibraryHeader.ShaderHashes[MappedShaderIndex].ToString());
-			}
-			DumpContent.Append(TEXT("\n"));
-		}
-		DumpContent.Append(TEXT("\n\n\n"));
+	DebugOutput.Appendf(TEXT("--------------- ShaderGroupIoHashes: ---------------\n"));
+	for (int32 GroupIndex = 0; GroupIndex < InLibraryHeader.ShaderGroupIoHashes.Num(); ++GroupIndex)
+	{
+		const FIoChunkId& ShaderGroupChunkId = InLibraryHeader.ShaderGroupIoHashes[GroupIndex];
+		const FIoStoreShaderGroupEntry& GroupEntry = InLibraryHeader.ShaderGroupEntries[GroupIndex];
+		DebugOutput.Appendf(TEXT("%s (%d shaders)\n"), *LexToString(ShaderGroupChunkId), GroupEntry.NumShaders);
 
-		DumpContent.Appendf(TEXT("--------------- ShaderHashes: ---------------\n"));
-		for (const FSHAHash& ShaderHash : InLibraryHeader.ShaderHashes)
+		for (uint32 ShaderIndex = 0; ShaderIndex < GroupEntry.NumShaders; ++ShaderIndex)
 		{
-			DumpContent.Appendf(TEXT("%s\n"), *ShaderHash.ToString());
+			const int32 MappedShaderIndex = InLibraryHeader.ShaderIndices[GroupEntry.ShaderIndicesOffset + ShaderIndex];
+			DebugOutput.Appendf(TEXT("  %s\n"), *InLibraryHeader.ShaderHashes[MappedShaderIndex].ToString());
 		}
-		DumpContent.Append(TEXT("\n\n\n"));
+		DebugOutput.Append(TEXT("\n"));
+	}
+	DebugOutput.Append(TEXT("\n\n\n"));
 
-		DumpContent.Appendf(TEXT("--------------- ShaderMapHashes: ---------------\n"));
-		for (const FSHAHash& ShaderMapHash : InLibraryHeader.ShaderMapHashes)
-		{
-			DumpContent.Appendf(TEXT("%s\n"), *ShaderMapHash.ToString());
-		}
-		DumpContent.Append(TEXT("\n\n\n"));
+	DebugOutput.Appendf(TEXT("--------------- ShaderHashes: ---------------\n"));
+	for (const FSHAHash& ShaderHash : InLibraryHeader.ShaderHashes)
+	{
+		DebugOutput.Appendf(TEXT("%s\n"), *ShaderHash.ToString());
+	}
+	DebugOutput.Append(TEXT("\n\n\n"));
 
-		DumpFile->Serialize(DumpContent.GetCharArray().GetData(), DumpContent.Len() * sizeof(TCHAR));
+	DebugOutput.Appendf(TEXT("--------------- ShaderMapHashes: ---------------\n"));
+	for (const FSHAHash& ShaderMapHash : InLibraryHeader.ShaderMapHashes)
+	{
+		DebugOutput.Appendf(TEXT("%s\n"), *ShaderMapHash.ToString());
 	}
+	DebugOutput.Append(TEXT("\n\n\n"));
+
+	FTCHARToUTF8 DebugOutputUtf8Converter(*DebugOutput);
+	Ar.Serialize(const_cast<UTF8CHAR*>(reinterpret_cast<const UTF8CHAR*>(DebugOutputUtf8Converter.Get())), DebugOutputUtf8Converter.Length());
 }
 
 FIoStoreShaderCodeArchive* FIoStoreShaderCodeArchive::Create(EShaderPlatform InPlatform, const FString& InLibraryName, FIoDispatcher& InIoDispatcher)
@@ -2386,7 +2385,10 @@ FIoStoreShaderCodeArchive* FIoStoreShaderCodeArchive::Create(EShaderPlatform InP
 					UE_LOG(LogShaderLibrary, Display, TEXT("Dump IoStoreShaderCodeArchive '%s-%s' header to file '%s'"),
 						*InLibraryName, *PlatformName.ToString(), *DumpFilename);
 
-					DumpShaderCodeArchiveHashes(Library->Header, *DumpFilename);
+					if (TUniquePtr<FArchive> DumpFile = TUniquePtr<FArchive>(IFileManager::Get().CreateFileWriter(*DumpFilename)))
+					{
+						DumpShaderCodeArchiveHashes(Library->Header, *DumpFile);
+					}
 				}
 			}
 
@@ -2805,6 +2807,13 @@ int32 FIoStoreShaderCodeArchive::FindShaderIndex(const FSHAHash& Hash)
 	return INDEX_NONE;
 }
 
+static const TCHAR* GetIoRequestBackendName(const FIoRequest& IoRequest)
+{
+	FIoRequestDebugInfo DebugInfo;
+	FIoDispatcherInternal::GetDebugInfo(IoRequest, DebugInfo);
+	return DebugInfo.BackendName ? DebugInfo.BackendName : TEXT("<None>");
+}
+
 TRefCountPtr<FRHIShader> FIoStoreShaderCodeArchive::CreateShader(int32 ShaderIndex, bool bRequired)
 {
 	LLM_SCOPE(ELLMTag::Shaders);
@@ -2874,14 +2883,12 @@ TRefCountPtr<FRHIShader> FIoStoreShaderCodeArchive::CreateShader(int32 ShaderInd
 	// we should not have this situation, but if we do, log a little more info about it before dying
 	if (UNLIKELY(PreloadEntryPtr->IoRequest.GetResult() == nullptr))
 	{
-		FString ShaderGroupIoHash;
-		Header.ShaderGroupIoHashes[GroupIndex].ToString(ShaderGroupIoHash);
 		UE_LOG(LogShaderLibrary, Warning, TEXT("IoRequest backing buffer is not there when we need it, about to crash. Library '%s', shader hash '%s' (freq=%d), group %d IoChunkId '%s', started preloading at frame %d, current frame is %d, preloading %s"),
 			*GetName(),
 			*Header.ShaderHashes[ShaderIndex].ToString(),
 			ShaderEntry.Frequency,
 			GroupIndex,
-			*ShaderGroupIoHash,
+			*LexToString(Header.ShaderGroupIoHashes[GroupIndex]),
 			PreloadEntryPtr->FramePreloadStarted,
 			GFrameNumber,
 			PreloadEntryPtr->PreloadEvent.IsValid() ? (PreloadEntryPtr->PreloadEvent->IsComplete() ? TEXT("is complete") : TEXT("is NOT complete!")) : TEXT("event is NOT VALID!")
@@ -2890,6 +2897,10 @@ TRefCountPtr<FRHIShader> FIoStoreShaderCodeArchive::CreateShader(int32 ShaderInd
 
 	const uint8* ShaderCode = PreloadEntryPtr->IoRequest.GetResultOrDie().Data();
 
+	// Verbose logging from what IoDispatcher backend this shader chunk was loaded. If this is null, GetResultOrDie() will have already crashed.
+	UE_LOG(LogShaderLibrary, VeryVerbose, TEXT("Received shader chunk '%s' in library '%s' from '%s'"),
+		*LexToString(Header.ShaderGroupIoHashes[GroupIndex]), *GetName(), GetIoRequestBackendName(PreloadEntryPtr->IoRequest));
+
 	FMemStackBase& MemStack = FMemStack::Get();
 	FMemMark Mark(MemStack);
 	FIoStoreShaderGroupEntry& GroupEntry = Header.ShaderGroupEntries[GroupIndex];
