RHI: Initial code drop for specialization constants (WIP):
- define a type for the container that will hold specialization constant values (FSpecializationContainerType)
- add FSpecializationContainerType to bound shader states for graphic pipelines
- add FSpecializationContainerType when fetching shader from shadermap with permutation vector
- add a flag to RHIShaders to know if they contain specialization constants
- add bSupportsSpecializationConstants to DDSPI
- add specialization support to Permutations and PermutationDomains.
- add specialization count and functions to ShaderType
- generate possible specialization constant values for shader compilation
- evaluate resource usage for each specialization constant combination and attach the results as OptionalData during compilation
- add specialization constants from FBoundShaderStateInput to pipeline precaching hashes

#rnx
#tests APT and ReplayRuns on Windows (D3D12 and Vulkan) and Android.

#Synced-CL 50561067

diff --git a/Engine/Plugins/Compositing/OpenColorIO/Source/OpenColorIO/Public/OpenColorIOShaderType.h b/Engine/Plugins/Compositing/OpenColorIO/Source/OpenColorIO/Public/OpenColorIOShaderType.h
--- a/Engine/Plugins/Compositing/OpenColorIO/Source/OpenColorIO/Public/OpenColorIOShaderType.h
+++ b/Engine/Plugins/Compositing/OpenColorIO/Source/OpenColorIO/Public/OpenColorIOShaderType.h
@@ -80,10 +80,13 @@ public:
 		const TCHAR* InFunctionName,
 		uint32 InFrequency,					// ugly - ignored for OCIO shaders but needed for IMPLEMENT_SHADER_TYPE macro magic
 		int32 InTotalPermutationCount,
+		int32 InTotalSpecializationCount,
 		ConstructSerializedType InConstructSerializedRef,
 		ConstructCompiledType InConstructCompiledRef,
 		ShouldCompilePermutationType InShouldCompilePermutationRef,
 		ShouldPrecachePermutationType InShouldPrecachePermutationRef,
+		GetUnspecializedIdType InGetUnspecializedIdRef,
+		GetSpecializationValuesType InGetSpecializationValuesRef,
 		GetRayTracingPayloadTypeType InGetRayTracingPayloadTypeRef,
 		GetShaderBindingLayoutType InGetShaderBindingLayoutTypeRef,
 #if WITH_EDITOR
@@ -96,12 +99,15 @@ public:
 #if WITH_EDITOR
 		, GetPermutationIdStringType InGetPermutationIdStringRef = nullptr
 #endif // WITH_EDITOR
-		):
-		FShaderType(EShaderTypeForDynamicCast::OCIO, InTypeLayout, InName, InSourceFilename, InFunctionName, SF_Pixel, InTotalPermutationCount,
+	):
+		FShaderType(EShaderTypeForDynamicCast::OCIO, InTypeLayout, InName, InSourceFilename, InFunctionName, SF_Pixel, InTotalPermutationCount, InTotalSpecializationCount,
+
 			InConstructSerializedRef,
 			InConstructCompiledRef,
 			InShouldCompilePermutationRef,
 			InShouldPrecachePermutationRef,
+			InGetUnspecializedIdRef,
+			InGetSpecializationValuesRef,
 			InGetRayTracingPayloadTypeRef,
 			InGetShaderBindingLayoutTypeRef,
 #if WITH_EDITOR
diff --git a/Engine/Plugins/Experimental/NNERuntimeIREE/Source/NNERuntimeIREEShader/Internal/NNERuntimeIREEShaderType.h b/Engine/Plugins/Experimental/NNERuntimeIREE/Source/NNERuntimeIREEShader/Internal/NNERuntimeIREEShaderType.h
--- a/Engine/Plugins/Experimental/NNERuntimeIREE/Source/NNERuntimeIREEShader/Internal/NNERuntimeIREEShaderType.h
+++ b/Engine/Plugins/Experimental/NNERuntimeIREE/Source/NNERuntimeIREEShader/Internal/NNERuntimeIREEShaderType.h
@@ -75,10 +75,13 @@ public:
 		const TCHAR* InFunctionName,
 		uint32 InFrequency,					// ugly - ignored for but needed for IMPLEMENT_SHADER_TYPE macro magic
 		int32 InTotalPermutationCount,
+		int32 InTotalSpecializationCount,
 		ConstructSerializedType InConstructSerializedRef,
 		ConstructCompiledType InConstructCompiledRef,
 		ShouldCompilePermutationType InShouldCompilePermutationRef,
 		ShouldPrecachePermutationType InShouldPrecachePermutationRef,
+		GetUnspecializedIdType InGetUnspecializedIdRef,
+		GetSpecializationValuesType InGetSpecializationValuesRef,
 		GetRayTracingPayloadTypeType InGetRayTracingPayloadTypeRef,
 		GetShaderBindingLayoutType InGetShaderBindingLayoutTypeRef,
 #if WITH_EDITOR
@@ -100,10 +103,13 @@ public:
 			InFunctionName, 
 			SF_Compute, 
 			InTotalPermutationCount,
+			InTotalSpecializationCount,
 			InConstructSerializedRef,
 			InConstructCompiledRef,
 			InShouldCompilePermutationRef,
 			InShouldPrecachePermutationRef,
+			InGetUnspecializedIdRef,
+			InGetSpecializationValuesRef,
 			InGetRayTracingPayloadTypeRef,
 			InGetShaderBindingLayoutTypeRef,
 #if WITH_EDITOR
diff --git a/Engine/Plugins/FX/Niagara/Source/Niagara/Private/NiagaraGpuComputeDispatch.cpp b/Engine/Plugins/FX/Niagara/Source/Niagara/Private/NiagaraGpuComputeDispatch.cpp
--- a/Engine/Plugins/FX/Niagara/Source/Niagara/Private/NiagaraGpuComputeDispatch.cpp
+++ b/Engine/Plugins/FX/Niagara/Source/Niagara/Private/NiagaraGpuComputeDispatch.cpp
@@ -1752,7 +1752,8 @@ void FNiagaraGpuComputeDispatch::DispatchStage(FRDGBuilder& GraphBuilder, const
 			[this, ShaderParametersMetadata, DispatchParameters, ComputeShader, TickPtr=&Tick, InstanceDataPtr=&InstanceData, SimStageDataPtr=&SimStageData](FRHICommandList& RHICmdList)
 			{
 				FRHIComputeShader* RHIComputeShader = ComputeShader.GetComputeShader();
-				SetComputePipelineState(RHICmdList, RHIComputeShader);
+				FComputePipelineStateInitializer Initializer(RHIComputeShader, ComputeShader.GetSpecializationConstants());
+				SetComputePipelineState(RHICmdList, Initializer);
 
 				FRHIBatchedShaderParameters& BatchedParameters = RHICmdList.GetScratchShaderParameters();
 
@@ -1845,7 +1846,8 @@ void FNiagaraGpuComputeDispatch::DispatchStage(FRDGBuilder& GraphBuilder, const
 			[this, ShaderParametersMetadata, DispatchParameters, ComputeShader, ThreadGroupCount, TickPtr=&Tick, InstanceDataPtr=&InstanceData, SimStageDataPtr=&SimStageData](FRHICommandList& RHICmdList)
 			{
 				FRHIComputeShader* RHIComputeShader = ComputeShader.GetComputeShader();
-				SetComputePipelineState(RHICmdList, RHIComputeShader);
+				FComputePipelineStateInitializer Initializer(RHIComputeShader, ComputeShader.GetSpecializationConstants());
+				SetComputePipelineState(RHICmdList, Initializer);
 
 				FRHIBatchedShaderParameters& BatchedParameters = RHICmdList.GetScratchShaderParameters();
 
diff --git a/Engine/Plugins/FX/Niagara/Source/NiagaraShader/Public/NiagaraShaderType.h b/Engine/Plugins/FX/Niagara/Source/NiagaraShader/Public/NiagaraShaderType.h
--- a/Engine/Plugins/FX/Niagara/Source/NiagaraShader/Public/NiagaraShaderType.h
+++ b/Engine/Plugins/FX/Niagara/Source/NiagaraShader/Public/NiagaraShaderType.h
@@ -86,10 +86,13 @@ public:
 		const TCHAR* InFunctionName,
 		uint32 InFrequency,					// ugly - ignored for Niagara shaders but needed for IMPLEMENT_SHADER_TYPE macro magic
 		int32 InTotalPermutationCount,
+		int32 InTotalSpecializationCount,
 		ConstructSerializedType InConstructSerializedRef,
 		ConstructCompiledType InConstructCompiledRef,
 		ShouldCompilePermutationType InShouldCompilePermutationRef,
 		ShouldPrecachePermutationType InShouldPrecachePermutationRef,
+		GetUnspecializedIdType InGetUnspecializedIdRef,
+		GetSpecializationValuesType InGetSpecializationValuesRef,
 		GetRayTracingPayloadTypeType InGetRayTracingPayloadTypeRef,
 		GetShaderBindingLayoutType InGetShaderBindingLayoutTypeRef,
 #if WITH_EDITOR
@@ -103,11 +106,13 @@ public:
 		, GetPermutationIdStringType InGetPermutationIdStringRef = nullptr
 #endif // WITH_EDITOR
 	)
-		: FShaderType(EShaderTypeForDynamicCast::Niagara, InTypeLayout, InName, InSourceFilename, InFunctionName, SF_Compute, InTotalPermutationCount,
+		: FShaderType(EShaderTypeForDynamicCast::Niagara, InTypeLayout, InName, InSourceFilename, InFunctionName, SF_Compute, InTotalPermutationCount, InTotalSpecializationCount,
 			InConstructSerializedRef,
 			InConstructCompiledRef,
 			InShouldCompilePermutationRef,
 			InShouldPrecachePermutationRef,
+			InGetUnspecializedIdRef,
+			InGetSpecializationValuesRef,
 			InGetRayTracingPayloadTypeRef,
 			InGetShaderBindingLayoutTypeRef,
 #if WITH_EDITOR
diff --git a/Engine/Plugins/Runtime/ComputeFramework/Source/ComputeFramework/Private/ComputeFramework/ComputeKernelShaderType.h b/Engine/Plugins/Runtime/ComputeFramework/Source/ComputeFramework/Private/ComputeFramework/ComputeKernelShaderType.h
--- a/Engine/Plugins/Runtime/ComputeFramework/Source/ComputeFramework/Private/ComputeFramework/ComputeKernelShaderType.h
+++ b/Engine/Plugins/Runtime/ComputeFramework/Source/ComputeFramework/Private/ComputeFramework/ComputeKernelShaderType.h
@@ -71,10 +71,13 @@ public:
 		const TCHAR* InFunctionName,
 		uint32 InFrequency,					// ugly - ignored for but needed for IMPLEMENT_SHADER_TYPE macro magic
 		int32 InTotalPermutationCount,
+		int32 InTotalSpecializationCount,
 		ConstructSerializedType InConstructSerializedRef,
 		ConstructCompiledType InConstructCompiledRef,
 		ShouldCompilePermutationType InShouldCompilePermutationRef,
 		ShouldPrecachePermutationType InShouldPrecachePermutationRef,
+		GetUnspecializedIdType InGetUnspecializedIdRef,
+		GetSpecializationValuesType InGetSpecializationValuesRef,
 		GetRayTracingPayloadTypeType InGetRayTracingPayloadTypeRef,
 		GetShaderBindingLayoutType InGetShaderBindingLayoutTypeRef,
 #if WITH_EDITOR
@@ -96,10 +99,13 @@ public:
 			InFunctionName, 
 			SF_Compute, 
 			InTotalPermutationCount,
+			InTotalSpecializationCount,
 			InConstructSerializedRef,
 			InConstructCompiledRef,
 			InShouldCompilePermutationRef,
 			InShouldPrecachePermutationRef,
+			InGetUnspecializedIdRef,
+			InGetSpecializationValuesRef,
 			InGetRayTracingPayloadTypeRef,
 			InGetShaderBindingLayoutTypeRef,
 #if WITH_EDITOR
diff --git a/Engine/Shaders/Public/Platform.ush b/Engine/Shaders/Public/Platform.ush
--- a/Engine/Shaders/Public/Platform.ush
+++ b/Engine/Shaders/Public/Platform.ush
@@ -1489,3 +1489,10 @@ uint WaveGetActiveLaneIndexLast()
 #define GET_SCALAR_ARRAY_ELEMENT(PackedArray, ElementIndex) PackedArray[(uint)ElementIndex>>2u][(uint)ElementIndex&3u]
 #endif
 
+// Specialization constant support
+#ifndef UE_DECL
+	#define UE_DECL(Name)
+#endif
+#ifndef UE_EVAL
+	#define UE_EVAL(Name) Name##_DEFAULT
+#endif
\ No newline at end of file
diff --git a/Engine/Source/Developer/ShaderCompilerCommon/Internal/ShaderCompilerCommonInternal.cpp b/Engine/Source/Developer/ShaderCompilerCommon/Internal/ShaderCompilerCommonInternal.cpp
--- a/Engine/Source/Developer/ShaderCompilerCommon/Internal/ShaderCompilerCommonInternal.cpp
+++ b/Engine/Source/Developer/ShaderCompilerCommon/Internal/ShaderCompilerCommonInternal.cpp
@@ -135,5 +135,218 @@ IMalloc* GetDxcMalloc()
 }
 #endif // PLATFORM_WINDOWS
 
+
+// Go through the specialization constants in the SPIRV to match them to what was expected in the build env.
+// Evaluate all the different combinations of the spec consts to find the resources that are conditional on specific values.
+int32 ProcessSpirvSpecializationConstants(
+	const FShaderCompilerInput& Input, 
+	CrossCompiler::FShaderConductorContext& CompilerContext,
+	spv_reflect::ShaderModule& Reflection,
+	const FSpirvReflectBindings& UnspecializedBindings,
+	FShaderCompilerOutput& Output)
+{
+	auto& SpecializationConstantValues = Input.Environment.SpecializationConstantValues;
+	const int32 TotalSpecializationCount = SpecializationConstantValues.Num();
+
+	// Bail if specialization constants aren't used
+	if (TotalSpecializationCount <= 0)
+	{
+		return 0;
+	}
+
+	// Put the total count of spec constants
+	uint32 SPIRVSpecializationCount = 0;
+	Reflection.EnumerateSpecializationConstants(&SPIRVSpecializationCount, nullptr);
+	TArray<SpvReflectSpecializationConstant*> SPIRVSpecConsts;
+	SPIRVSpecConsts.SetNumZeroed(SPIRVSpecializationCount);
+	Reflection.EnumerateSpecializationConstants(&SPIRVSpecializationCount, SPIRVSpecConsts.GetData());
+
+	// Read the array of expected specialization constants from the build environment
+	checkf((uint32)TotalSpecializationCount >= SPIRVSpecializationCount,
+		TEXT("SPIRV specialization count (%u) is greater than expected (%d) by build environment."),
+		SPIRVSpecializationCount, TotalSpecializationCount);
+
+	// SPIRV-Tool passes to run to evaluate constants
+	// No need for "set-spec-const-default-value", we'll use SPIRVReflect
+	const ANSICHAR* OptimizerConfig[] = {
+
+		// Freeze the spec constants
+		"--freeze-spec-const",
+		"--fold-spec-const-op-composite",
+
+		// Remove dead code from the frozen constants
+		"--eliminate-dead-code-aggressive",
+		"--eliminate-dead-inserts",
+		"--eliminate-dead-branches",
+		"--eliminate-dead-members",
+		"--eliminate-dead-variables",
+		"--eliminate-dead-const",
+
+		// Full -O passes (note: could probably be reduced)
+		"--eliminate-dead-branches",
+		"--merge-return",
+		"--inline-entry-points-exhaustive",
+		"--eliminate-dead-functions",
+		"--eliminate-dead-code-aggressive",
+		"--private-to-local",
+		"--eliminate-local-single-block",
+		"--eliminate-local-single-store",
+		"--eliminate-dead-code-aggressive",
+		"--scalar-replacement",
+		"--convert-local-access-chains",
+		"--eliminate-local-single-block",
+		"--eliminate-local-single-store",
+		"--eliminate-dead-code-aggressive",
+		"--eliminate-local-multi-store",
+		"--eliminate-dead-code-aggressive",
+		"--ccp",
+		"--eliminate-dead-code-aggressive",
+		"--loop-unroll",
+		"--eliminate-dead-branches",
+		"--redundancy-elimination",
+		"--combine-access-chains",
+		"--simplify-instructions",
+		"--scalar-replacement",
+		"--convert-local-access-chains",
+		"--eliminate-local-single-block",
+		"--eliminate-local-single-store",
+		"--eliminate-dead-code-aggressive",
+		"--ssa-rewrite",
+		"--eliminate-dead-code-aggressive",
+		"--vector-dce",
+		"--eliminate-dead-inserts",
+		"--eliminate-dead-branches",
+		"--simplify-instructions",
+		"--if-conversion",
+		"--copy-propagate-arrays",
+		"--reduce-load-size",
+		"--eliminate-dead-code-aggressive",
+		"--merge-blocks",
+		"--redundancy-elimination",
+		"--eliminate-dead-branches",
+		"--merge-blocks",
+		"--simplify-instructions",
+		"--eliminate-dead-members",
+		"--eliminate-local-single-store",
+		"--merge-blocks",
+		"--eliminate-local-multi-store",
+		"--redundancy-elimination",
+		"--simplify-instructions",
+		"--eliminate-dead-code-aggressive",
+		"--cfg-cleanup"
+	};
+
+	// Map of resource names to hashes of spec const values to exclude the resource from
+	TMap<FString, TArray<uint32>> ExcludedResources;
+	auto DiffResourceArrays = [&](const uint32 SpecHash, const TArray<SpvReflectDescriptorBinding*>& UnspecArray, const TArray<SpvReflectDescriptorBinding*>& SpecArray)
+		{
+			for (SpvReflectDescriptorBinding* UnspecBinding : UnspecArray)
+			{
+				if (!SpecArray.FindByPredicate([UnspecBinding](SpvReflectDescriptorBinding* SpecBinding) {
+					return FCStringAnsi::Strcmp(UnspecBinding->name, SpecBinding->name) == 0;
+					}))
+				{
+					TArray<uint32>& Hashes = ExcludedResources.FindOrAdd(UnspecBinding->name);
+					Hashes.Add(SpecHash);
+				}
+			}
+		};
+
+	auto ProcessFullSpecialization = [&](const TArray<int32>& SpecValues)
+		{
+			// Set the default values to the combination we are evaluating
+			for (SpvReflectSpecializationConstant* SpirvSpecConst : SPIRVSpecConsts)
+			{
+				checkf(SpirvSpecConst->default_value_size == 4, TEXT("All specialization constants should current be 32bit wide."));
+				checkf(SpirvSpecConst->constant_id < (uint32)SpecializationConstantValues.Num(), TEXT("Specialization constants index out of bounds."));
+
+				*((int32*)SpirvSpecConst->default_value) = SpecValues[SpirvSpecConst->constant_id];
+			}
+
+			TArray<uint32> SpirvCopy(Reflection.GetCode(), Reflection.GetCodeSize() / sizeof(uint32));
+
+			if (CompilerContext.OptimizeSpirv(SpirvCopy, OptimizerConfig, UE_ARRAY_COUNT(OptimizerConfig)))
+			{
+				// Return code reflection if requested for shader analysis
+				if (Input.Environment.CompilerFlags.Contains(CFLAG_OutputAnalysisArtifacts))
+				{
+					FGenericShaderStat SpecConstReflection;
+					if (CrossCompiler::FShaderConductorContext::Disassemble(CrossCompiler::EShaderConductorIR::Spirv, SpirvCopy.GetData(), SpirvCopy.NumBytes(), SpecConstReflection))
+					{
+						FString SpecValString;
+						for (int32 SpecId : SpecValues)
+						{
+							SpecValString.Append(FString::Printf(TEXT(".%d"), SpecId));
+						}
+						SpecConstReflection.StatName = FName(SpecConstReflection.StatName.ToString() + SpecValString);
+
+						Output.ShaderStatistics.Add(MoveTemp(SpecConstReflection));
+					}
+				}
+
+				// Parse the new SPIRV (no need to copy the SPIRV again)
+				spv_reflect::ShaderModule EvalReflection(SpirvCopy.NumBytes(), SpirvCopy.GetData(), SPV_REFLECT_MODULE_FLAG_NO_COPY);
+
+				// Grab resources in SPIRV with spec constats evaluated
+				FSpirvReflectBindings EvalBindings;
+				EvalBindings.GatherDescriptorBindings(EvalReflection);
+
+				const uint32 SpecHash = GetTypeHash(SpecValues);
+				DiffResourceArrays(SpecHash, UnspecializedBindings.UniformBuffers, EvalBindings.UniformBuffers);
+				DiffResourceArrays(SpecHash, UnspecializedBindings.Samplers, EvalBindings.Samplers);
+				DiffResourceArrays(SpecHash, UnspecializedBindings.TextureSRVs, EvalBindings.TextureSRVs);
+				DiffResourceArrays(SpecHash, UnspecializedBindings.TextureUAVs, EvalBindings.TextureUAVs);
+				DiffResourceArrays(SpecHash, UnspecializedBindings.TBufferSRVs, EvalBindings.TBufferSRVs);
+				DiffResourceArrays(SpecHash, UnspecializedBindings.TBufferUAVs, EvalBindings.TBufferUAVs);
+				DiffResourceArrays(SpecHash, UnspecializedBindings.SBufferSRVs, EvalBindings.SBufferSRVs);
+				DiffResourceArrays(SpecHash, UnspecializedBindings.SBufferUAVs, EvalBindings.SBufferUAVs);
+				DiffResourceArrays(SpecHash, UnspecializedBindings.AccelerationStructures, EvalBindings.AccelerationStructures);
+			}
+		};
+
+	// Go through every possible combination of the possible values for the spec contants
+	TArray<int32> SpecValues;
+	SpecValues.SetNumZeroed(TotalSpecializationCount);
+	const TFunction<void(int32)> ProcessDimension = [&](const int32 DimensionIndex)
+		{
+			for (int32 SpecIndex = 0; SpecIndex < SpecializationConstantValues[DimensionIndex].Num(); ++SpecIndex)
+			{
+				SpecValues[DimensionIndex] = SpecializationConstantValues[DimensionIndex][SpecIndex];
+
+				if (DimensionIndex < SpecializationConstantValues.Num() - 1)
+				{
+					ProcessDimension(DimensionIndex + 1);
+				}
+				else
+				{
+					ProcessFullSpecialization(SpecValues);
+				}
+			}
+		};
+
+	// Start processing at the first spec const
+	ProcessDimension(0);
+
+	// Create a new ShaderModule
+	// Iterate over every possible combination
+	//   Set the values using Reflection.m_module.spec_constants->default_value
+	//   Run the FoldSpecConstantOpAndCompositePass pass to fold the OpSpecConstantOp
+	//   Run a subset of optimization passes to make sure unused resources are removed
+	//   Run GatherSpirvReflectionBindings and compare the resources
+	//     Mark the unused ones in the Parameter map
+
+	auto& ParameterMap = Output.ParameterMap.ParameterMap;
+	for (auto& Pair : ExcludedResources)
+	{
+		if (FParameterAllocation* Param = ParameterMap.Find(Pair.Key))
+		{
+			TArray<uint32>& Hashes = Output.SpecializationConstantExclusions.FindOrAdd(Param->BaseIndex);
+			Hashes = MoveTemp(Pair.Value);
+		}
+	}
+
+	return TotalSpecializationCount;
+}
+
 } // UE::ShaderCompilerCommon
 
diff --git a/Engine/Source/Developer/ShaderCompilerCommon/Internal/ShaderCompilerCommonInternal.h b/Engine/Source/Developer/ShaderCompilerCommon/Internal/ShaderCompilerCommonInternal.h
--- a/Engine/Source/Developer/ShaderCompilerCommon/Internal/ShaderCompilerCommonInternal.h
+++ b/Engine/Source/Developer/ShaderCompilerCommon/Internal/ShaderCompilerCommonInternal.h
@@ -5,6 +5,8 @@
 #include "CoreMinimal.h"
 #include "ShaderCompilerCore.h"
 #include "HlslParserInternal.h"
+#include "SpirvReflectCommon.h"
+#include "ShaderConductorContext.h"
 
 #if PLATFORM_WINDOWS
 #include "Windows/AllowWindowsPlatformTypes.h"
@@ -38,5 +40,12 @@ namespace UE::ShaderCompilerCommon
 	// Get an IMalloc to overload DXC's allocator with our own in Windows
 	extern SHADERCOMPILERCOMMON_API IMalloc* GetDxcMalloc();
 #endif
+
+	extern SHADERCOMPILERCOMMON_API int32 ProcessSpirvSpecializationConstants(
+		const FShaderCompilerInput& Input,
+		CrossCompiler::FShaderConductorContext& CompilerContext,
+		spv_reflect::ShaderModule& Reflection,
+		const FSpirvReflectBindings& UnspecializedBindings,
+		FShaderCompilerOutput& Output);
 }
 
diff --git a/Engine/Source/Runtime/Engine/Private/ShaderCompiler/ShaderCompiler.cpp b/Engine/Source/Runtime/Engine/Private/ShaderCompiler/ShaderCompiler.cpp
--- a/Engine/Source/Runtime/Engine/Private/ShaderCompiler/ShaderCompiler.cpp
+++ b/Engine/Source/Runtime/Engine/Private/ShaderCompiler/ShaderCompiler.cpp
@@ -4171,6 +4171,7 @@ void GlobalBeginCompileShader(
 	SET_SHADER_DEFINE(Input.Environment, PLATFORM_SUPPORTS_BARYCENTRICS_SEMANTIC, FDataDrivenShaderPlatformInfo::GetSupportsBarycentricsSemantic(EShaderPlatform(Target.Platform)) != ERHIFeatureSupport::Unsupported);
 	SET_SHADER_DEFINE(Input.Environment, PLATFORM_SUPPORTS_BINDLESS, FDataDrivenShaderPlatformInfo::GetSupportsBindless(ShaderPlatform));
 	SET_SHADER_DEFINE(Input.Environment, PLATFORM_SUPPORTS_RAYTRACING_CLUSTER_OPS, FDataDrivenShaderPlatformInfo::GetSupportsRayTracingClusterOps(ShaderPlatform));
+	SET_SHADER_DEFINE(Input.Environment, PLATFORM_SUPPORTS_SPECIALIZATION_CONSTANTS, FDataDrivenShaderPlatformInfo::GetSupportsSpecializationConstants(ShaderPlatform));
 
 	SET_SHADER_DEFINE(Input.Environment, UE_BINDLESS_ENABLED, Input.bBindlessEnabled);
 
@@ -4883,12 +4884,11 @@ void VerifyGlobalShaders(EShaderPlatform Platform, const ITargetPlatform* Target
 			continue;
 		}
 
+		const bool bOutdated = OutdatedShaderTypes && OutdatedShaderTypes->Contains(GlobalShaderType);
+
 		int32 PermutationCountToCompile = 0;
-		for (int32 PermutationId = 0; PermutationId < GlobalShaderType->GetPermutationCount(); PermutationId++)
+		GlobalShaderType->ForEachCompiledPermutation(Platform, PermutationFlags, [&](int32 PermutationId)
 		{
-			if (GlobalShaderType->ShouldCompilePermutation(Platform, PermutationId, PermutationFlags))
-			{
-				bool bOutdated = OutdatedShaderTypes && OutdatedShaderTypes->Contains(GlobalShaderType);
 				TShaderRef<FShader> GlobalShader = GlobalShaderMap->GetShader(GlobalShaderType, PermutationId);
 				if (bOutdated || !GlobalShader.IsValid())
 				{
@@ -4920,8 +4920,8 @@ void VerifyGlobalShaders(EShaderPlatform Platform, const ITargetPlatform* Target
 #endif // WITH_EDITOR
 					}
 				}
-			}
-		}
+			return true;
+		});
 
 		int32 PermutationCountLimit = 832;	// Nanite culling as of today (2022-01-11) can go up to 832 permutations
 		if (Substrate::IsSubstrateEnabled())
@@ -5113,12 +5113,15 @@ void PrecacheComputePipelineStatesForGlobalShaders(ERHIFeatureLevel::Type Featur
 							constexpr bool bRequired = false;
 							FRHIComputeShader* RHIComputeShader = static_cast<FRHIComputeShader*>(GlobalShader.GetRHIShaderBase(SF_Compute, bRequired));
 							if (RHIComputeShader)
-							{								
+							{
 								FPSOPrecacheRequestParams PSORequestParams;
 								PSORequestParams.InMemoryMode = EPSOPrecacheInMemoryMode::Uncached;
 								PSORequestParams.PrecachePriority = EPSOPrecachePriority::Medium;
 								PSORequestParams.Name = GlobalShader.GetType()->GetName();
-								PipelineStateCache::PrecacheComputePipelineState(RHIComputeShader, PSORequestParams);
+
+								check(!RHIComputeShader->HasSpecializationConstants() || !GlobalShader.GetSpecializationConstants().IsEmpty());
+								FComputePipelineStateInitializer ComputeInitializer(RHIComputeShader, GlobalShader.GetSpecializationConstants());
+								PipelineStateCache::PrecacheComputePipelineState(ComputeInitializer, PSORequestParams);
 							}
 						}
 					));
@@ -5140,12 +5143,15 @@ void PrecacheComputePipelineStatesForGlobalShaders(ERHIFeatureLevel::Type Featur
 							bool bRequired = false;
 							FRHIComputeShader* RHIComputeShader = static_cast<FRHIComputeShader*>(GlobalShader.GetRHIShaderBase(SF_Compute, bRequired));
 							if (RHIComputeShader)
-							{								
+							{
 								FPSOPrecacheRequestParams PSORequestParams;
 								PSORequestParams.InMemoryMode = EPSOPrecacheInMemoryMode::Uncached;
 								PSORequestParams.PrecachePriority = EPSOPrecachePriority::Medium;
 								PSORequestParams.Name = GlobalShader.GetType()->GetName();
-								PipelineStateCache::PrecacheComputePipelineState(RHIComputeShader, PSORequestParams);
+
+								check(!RHIComputeShader->HasSpecializationConstants() || !GlobalShader.GetSpecializationConstants().IsEmpty());
+								FComputePipelineStateInitializer ComputeInitializer(RHIComputeShader, GlobalShader.GetSpecializationConstants());
+								PipelineStateCache::PrecacheComputePipelineState(ComputeInitializer, PSORequestParams);
 							}
 						}
 					});
diff --git a/Engine/Source/Runtime/Engine/Public/PSOPrecache.h b/Engine/Source/Runtime/Engine/Public/PSOPrecache.h
--- a/Engine/Source/Runtime/Engine/Public/PSOPrecache.h
+++ b/Engine/Source/Runtime/Engine/Public/PSOPrecache.h
@@ -274,7 +274,7 @@ struct FPSOPrecacheData
 		}
 		else
 		{
-			ComputeInitializer = FComputePipelineStateInitializer(InComputeShader.GetComputeShader());
+			ComputeInitializer = FComputePipelineStateInitializer(InComputeShader.GetComputeShader(), InComputeShader.GetSpecializationConstants());
 		}
 	}
 
diff --git a/Engine/Source/Runtime/RHI/Private/DataDrivenShaderPlatformInfo.cpp b/Engine/Source/Runtime/RHI/Private/DataDrivenShaderPlatformInfo.cpp
--- a/Engine/Source/Runtime/RHI/Private/DataDrivenShaderPlatformInfo.cpp
+++ b/Engine/Source/Runtime/RHI/Private/DataDrivenShaderPlatformInfo.cpp
@@ -322,6 +322,7 @@ void FGenericDataDrivenShaderPlatformInfo::ParseDataDrivenShaderInfo(const FConf
 	GET_SECTION_BOOL_HELPER(bSupportsWorkGraphsTier1_1);
 	GET_SECTION_BOOL_HELPER(bSupportsDLSSShaders);
 	GET_SECTION_BOOL_HELPER(bSupportsAdaptiveGBuffer);
+	GET_SECTION_BOOL_HELPER(bSupportsSpecializationConstants);
 #undef GET_SECTION_BOOL_HELPER
 #undef GET_SECTION_INT_HELPER
 #undef GET_SECTION_SUPPORT_HELPER
@@ -565,6 +566,7 @@ void FGenericDataDrivenShaderPlatformInfo::UpdatePreviewPlatforms()
 				PREVIEW_DISABLE_IF_RUNTIME_UNSUPPORTED(bSupportsMeshShadersTier1);
 				PREVIEW_DISABLE_IF_RUNTIME_UNSUPPORTED(bStripNaniteVertexShaders);
 				PREVIEW_DISABLE_IF_RUNTIME_UNSUPPORTED(bSupportsMobileMultiView);
+				PREVIEW_DISABLE_IF_RUNTIME_UNSUPPORTED(bSupportsSpecializationConstants);
 
 				// Settings that need to match the runtime
 				PREVIEW_USE_RUNTIME_VALUE(bSupportsGPUScene);
diff --git a/Engine/Source/Runtime/RHI/Private/DynamicRHI.cpp b/Engine/Source/Runtime/RHI/Private/DynamicRHI.cpp
--- a/Engine/Source/Runtime/RHI/Private/DynamicRHI.cpp
+++ b/Engine/Source/Runtime/RHI/Private/DynamicRHI.cpp
@@ -592,12 +592,16 @@ uint64 FDynamicRHI::RHIComputeStatePrecachePSOHash(const FGraphicsPipelineStateI
 	{
 		uint32 VertexDeclaration;
 		uint32 VertexShader;
+		uint32 VertexShaderSpecConsts;
 		uint32 PixelShader;
+		uint32 PixelShaderSpecConsts;
 #if PLATFORM_SUPPORTS_GEOMETRY_SHADERS
 		uint32 GeometryShader;
+		uint32 GeometryShaderSpecConsts;
 #endif // PLATFORM_SUPPORTS_GEOMETRY_SHADERS
 #if PLATFORM_SUPPORTS_MESH_SHADERS
 		uint32 MeshShader;
+		uint32 MeshShaderSpecConsts;
 #endif // PLATFORM_SUPPORTS_MESH_SHADERS
 		uint32 BlendState;
 		uint32 RasterizerState;
@@ -608,17 +612,21 @@ uint64 FDynamicRHI::RHIComputeStatePrecachePSOHash(const FGraphicsPipelineStateI
 		uint32 bDepthBounds : 1;
 		uint32 Unused : 15;
 	} HashKey;
-
+	 
 	FMemory::Memzero(&HashKey, sizeof(FHashKey));
 
 	HashKey.VertexDeclaration = Initializer.BoundShaderState.VertexDeclarationRHI ? Initializer.BoundShaderState.VertexDeclarationRHI->GetPrecachePSOHash() : 0;
 	HashKey.VertexShader = Initializer.BoundShaderState.GetVertexShader() ? GetTypeHash(Initializer.BoundShaderState.GetVertexShader()->GetHash()) : 0;
+	HashKey.VertexShaderSpecConsts = Initializer.BoundShaderState.VertexShaderSpecializationConstants.Num() ? GetArrayHash(Initializer.BoundShaderState.VertexShaderSpecializationConstants.GetData(), Initializer.BoundShaderState.VertexShaderSpecializationConstants.NumBytes()) : 0;
 	HashKey.PixelShader = Initializer.BoundShaderState.GetPixelShader() ? GetTypeHash(Initializer.BoundShaderState.GetPixelShader()->GetHash()) : 0;
+	HashKey.PixelShaderSpecConsts = Initializer.BoundShaderState.PixelShaderSpecializationConstants.Num() ? GetArrayHash(Initializer.BoundShaderState.PixelShaderSpecializationConstants.GetData(), Initializer.BoundShaderState.PixelShaderSpecializationConstants.NumBytes()) : 0;
 #if PLATFORM_SUPPORTS_GEOMETRY_SHADERS
 	HashKey.GeometryShader = Initializer.BoundShaderState.GetGeometryShader() ? GetTypeHash(Initializer.BoundShaderState.GetGeometryShader()->GetHash()) : 0;
+	HashKey.GeometryShaderSpecConsts = Initializer.BoundShaderState.GeometryShaderSpecializationConstants.Num() ? GetArrayHash(Initializer.BoundShaderState.GeometryShaderSpecializationConstants.GetData(), Initializer.BoundShaderState.GeometryShaderSpecializationConstants.NumBytes()) : 0;
 #endif
 #if PLATFORM_SUPPORTS_MESH_SHADERS
 	HashKey.MeshShader = Initializer.BoundShaderState.GetMeshShader() ? GetTypeHash(Initializer.BoundShaderState.GetMeshShader()->GetHash()) : 0;
+	HashKey.MeshShaderSpecConsts = Initializer.BoundShaderState.MeshShaderSpecializationConstants.Num() ? GetArrayHash(Initializer.BoundShaderState.MeshShaderSpecializationConstants.GetData(), Initializer.BoundShaderState.MeshShaderSpecializationConstants.NumBytes()) : 0;
 #endif
 
 	FBlendStateInitializerRHI BlendStateInitializerRHI;
@@ -750,6 +758,21 @@ bool FDynamicRHI::RHIMatchPrecachePSOInitializers(const FGraphicsPipelineStateIn
 		return false;
 	}
 
+	// check the spec consts
+	if (LHS.BoundShaderState.VertexShaderSpecializationConstants != RHS.BoundShaderState.VertexShaderSpecializationConstants
+		|| LHS.BoundShaderState.PixelShaderSpecializationConstants != RHS.BoundShaderState.PixelShaderSpecializationConstants
+#if PLATFORM_SUPPORTS_MESH_SHADERS
+		|| LHS.BoundShaderState.MeshShaderSpecializationConstants != RHS.BoundShaderState.MeshShaderSpecializationConstants
+		|| LHS.BoundShaderState.AmplificationShaderSpecializationConstants != RHS.BoundShaderState.AmplificationShaderSpecializationConstants
+#endif
+#if PLATFORM_SUPPORTS_GEOMETRY_SHADERS
+		|| LHS.BoundShaderState.GeometryShaderSpecializationConstants != RHS.BoundShaderState.GeometryShaderSpecializationConstants
+#endif
+		)
+	{
+		return false;
+	}
+
 	// Full compare the of the vertex declaration
 	if (!MatchRHIState<FRHIVertexDeclaration, FVertexDeclarationElementList>(LHS.BoundShaderState.VertexDeclarationRHI, RHS.BoundShaderState.VertexDeclarationRHI))
 	{
diff --git a/Engine/Source/Runtime/RHI/Private/PipelineStateCache.cpp b/Engine/Source/Runtime/RHI/Private/PipelineStateCache.cpp
--- a/Engine/Source/Runtime/RHI/Private/PipelineStateCache.cpp
+++ b/Engine/Source/Runtime/RHI/Private/PipelineStateCache.cpp
@@ -78,16 +78,21 @@ static inline uint32 GetTypeHash(const FBoundShaderStateInput& Input)
 {
 	uint32 Hash = GetTypeHash(Input.VertexDeclarationRHI);
 	Hash = HashCombineFast(Hash, GetTypeHash(Input.VertexShaderRHI));
+	Hash = HashCombineFast(Hash, GetTypeHash(Input.VertexShaderSpecializationConstants));
 	Hash = HashCombineFast(Hash, GetTypeHash(Input.PixelShaderRHI));
+	Hash = HashCombineFast(Hash, GetTypeHash(Input.PixelShaderSpecializationConstants));
 #if PLATFORM_SUPPORTS_MESH_SHADERS
 	Hash = HashCombineFast(Hash, GetTypeHash(Input.GetMeshShader()));
+	Hash = HashCombineFast(Hash, GetTypeHash(Input.MeshShaderSpecializationConstants));
 	Hash = HashCombineFast(Hash, GetTypeHash(Input.GetAmplificationShader()));
+	Hash = HashCombineFast(Hash, GetTypeHash(Input.AmplificationShaderSpecializationConstants));
 #endif
 #if PLATFORM_SUPPORTS_WORKGRAPH_SHADERS
 	Hash = HashCombineFast(Hash, GetTypeHash(Input.GetWorkGraphShader()));
 #endif
 #if PLATFORM_SUPPORTS_GEOMETRY_SHADERS
 	Hash = HashCombineFast(Hash, GetTypeHash(Input.GetGeometryShader()));
+	Hash = HashCombineFast(Hash, GetTypeHash(Input.GeometryShaderSpecializationConstants));
 #endif
 	return Hash;
 }
@@ -100,6 +105,7 @@ inline uint32 GetTypeHash(const FExclusiveDepthStencil& Ds)
 static inline uint32 GetTypeHash(const FComputePipelineStateInitializer& Initializer)
 {
 	uint32 Hash = GetTypeHash(Initializer.ComputeShader);
+	Hash = HashCombineFast(Hash, GetTypeHash(Initializer.SpecializationConstants));
 	return Hash;
 }
 
@@ -3263,7 +3269,7 @@ public:
 
 	static FORCEINLINE uint64 PipelineStateInitializerHash(const FComputePipelineStateInitializer& Key)
 	{
-		return GetTypeHash(Key.ComputeShader->GetHash());
+		return HashCombineFast(GetTypeHash(Key.ComputeShader->GetHash()), GetTypeHash(Key.SpecializationConstants));
 	}
 };
 
@@ -4108,6 +4114,7 @@ static void InternalAsyncCompilePipelineState(const TPipelineInitializer& Initia
 
 static void InternalCreateComputePipelineState(const FComputePipelineStateInitializer& Initializer, bool bDoAsyncCompile, EPSOPrecacheResult PSOPrecacheResult, EPSOPrecachePriority PrecachePriority, bool bPrecachedPSOMarkedUsed, FComputePipelineState* CachedState, const FPSOCompilationDebugData& PSOCompilationDebugData, bool bInImmediateCmdList)
 {
+	check(!Initializer.ComputeShader->HasSpecializationConstants() || !Initializer.SpecializationConstants.IsEmpty());
 	FGraphEventRef GraphEvent = CachedState->GetCompletionEvent();
 
 	// create a compilation task, or just do it now...
diff --git a/Engine/Source/Runtime/RHI/Public/DataDrivenShaderPlatformInfo.h b/Engine/Source/Runtime/RHI/Public/DataDrivenShaderPlatformInfo.h
--- a/Engine/Source/Runtime/RHI/Public/DataDrivenShaderPlatformInfo.h
+++ b/Engine/Source/Runtime/RHI/Public/DataDrivenShaderPlatformInfo.h
@@ -134,6 +134,7 @@ class FGenericDataDrivenShaderPlatformInfo
 	uint32 bSupportsWorkGraphsTier1_1 : 1;
 	uint32 bSupportsDLSSShaders : 1;
 	uint32 bSupportsAdaptiveGBuffer : 1;
+	uint32 bSupportsSpecializationConstants : 1;
 
 	// NOTE: When adding fields, you must also add to ParseDataDrivenShaderInfo!
 	uint32 bContainsValidPlatformInfo : 1;
@@ -883,6 +884,12 @@ PRAGMA_ENABLE_DEPRECATION_WARNINGS
 		return Infos[Platform].bSupportsAdaptiveGBuffer;
 	}
 	
+	static inline const bool GetSupportsSpecializationConstants(const FStaticShaderPlatform Platform)
+	{
+		check(IsValid(Platform));
+		return Infos[Platform].bSupportsSpecializationConstants;
+	}
+
 	static inline const bool IsValid(const FStaticShaderPlatform Platform)
 	{
 		return Infos[Platform].bContainsValidPlatformInfo;
diff --git a/Engine/Source/Runtime/RHI/Public/RHIDefinitions.h b/Engine/Source/Runtime/RHI/Public/RHIDefinitions.h
--- a/Engine/Source/Runtime/RHI/Public/RHIDefinitions.h
+++ b/Engine/Source/Runtime/RHI/Public/RHIDefinitions.h
@@ -74,6 +74,13 @@ static_assert(sizeof(void*) <= SHADER_PARAMETER_POINTER_ALIGNMENT, "The alignmen
 #define HAS_GPU_STATS ((STATS || CSV_PROFILER_STATS) && (!UE_BUILD_SHIPPING))
 #endif
 
+
+/**
+ * The type used to pass the specialization values around
+ */
+using FSpecializationContainerType = TArray<int32>;
+
+
 /**
  * A type used only for printing a string for debugging/profiling.
  * Adds Number as a suffix to the printed string even if the base name includes a number, so may prints a string like: Base_1_1
diff --git a/Engine/Source/Runtime/RHI/Public/RHIResources.h b/Engine/Source/Runtime/RHI/Public/RHIResources.h
--- a/Engine/Source/Runtime/RHI/Public/RHIResources.h
+++ b/Engine/Source/Runtime/RHI/Public/RHIResources.h
@@ -848,11 +848,25 @@ public:
 	TArray<FShaderCodeValidationType> DebugUAVTypeValidationData;
 	TArray<FShaderCodeValidationUBSize> DebugUBSizeValidationData;
 
+	TMap<uint32, TArray<uint32>> DebugSpecializationExclusions;
+	bool IsExcludedFromSpecialization(uint32 BindPoint, uint32 SpecializationHash) const
+	{
+		if (bHasSpecializationConstants)
+		{
+			if (const TArray<uint32>* Array = DebugSpecializationExclusions.Find(BindPoint))
+			{
+				return Array->Find(SpecializationHash) != INDEX_NONE;
+			}
+		}
+		return false;
+	}
+
 #else
 
 	bool HasShaderName() const { return false; }
 	const TCHAR* GetShaderName() const { return TEXT("<unknown>"); }
 	FString GetUniformBufferName(uint32 Index) const { return TEXT("<unknown>"); }
+	bool IsExcludedFromSpecialization(uint32 BindPoint, uint32 SpecializationHash) const { return false; }
 
 #endif
 
@@ -862,6 +876,7 @@ public:
 		, Frequency(InFrequency)
 		, bNoDerivativeOps(false)
 		, bHasShaderBundleUsage(false)
+		, bHasSpecializationConstants(false)
 	{
 	}
 
@@ -895,11 +910,22 @@ public:
 		return bHasShaderBundleUsage;
 	}
 
+	inline void SetSpecializationConstants(bool bValue)
+	{
+		bHasSpecializationConstants = bValue;
+	}
+
+	inline bool HasSpecializationConstants() const
+	{
+		return bHasSpecializationConstants;
+	}
+
 private:
 	FSHAHash Hash;
 	EShaderFrequency Frequency;
 	uint8 bNoDerivativeOps : 1;
 	uint8 bHasShaderBundleUsage : 1;
+	uint8 bHasSpecializationConstants : 1;
 #if PLATFORM_WINDOWS
 	volatile int16 InUseByPSOCompilation = 0;
 #endif // PLATFORM_WINDOWS
@@ -4429,6 +4455,8 @@ struct FBoundShaderStateInput
 	void SetMeshShader(FRHIMeshShader* InMeshShader) { MeshShaderRHI = InMeshShader; }
 	FRHIAmplificationShader* GetAmplificationShader() const { return AmplificationShaderRHI; }
 	void SetAmplificationShader(FRHIAmplificationShader* InAmplificationShader) { AmplificationShaderRHI = InAmplificationShader; }
+	FSpecializationContainerType MeshShaderSpecializationConstants;
+	FSpecializationContainerType AmplificationShaderSpecializationConstants;
 #else
 	constexpr FRHIMeshShader* GetMeshShader() const { return nullptr; }
 	void SetMeshShader(FRHIMeshShader*) {}
@@ -4439,6 +4467,7 @@ struct FBoundShaderStateInput
 #if PLATFORM_SUPPORTS_GEOMETRY_SHADERS
 	FRHIGeometryShader* GetGeometryShader() const { return GeometryShaderRHI; }
 	void SetGeometryShader(FRHIGeometryShader* InGeometryShader) { GeometryShaderRHI = InGeometryShader; }
+	FSpecializationContainerType GeometryShaderSpecializationConstants;
 #else
 	constexpr FRHIGeometryShader* GetGeometryShader() const { return nullptr; }
 	void SetGeometryShader(FRHIGeometryShader*) {}
@@ -4455,6 +4484,8 @@ struct FBoundShaderStateInput
 	FRHIVertexDeclaration* VertexDeclarationRHI = nullptr;
 	FRHIVertexShader* VertexShaderRHI = nullptr;
 	FRHIPixelShader* PixelShaderRHI = nullptr;
+	FSpecializationContainerType VertexShaderSpecializationConstants;
+	FSpecializationContainerType PixelShaderSpecializationConstants;
 private:
 #if PLATFORM_SUPPORTS_MESH_SHADERS
 	FRHIMeshShader* MeshShaderRHI = nullptr;
@@ -4772,23 +4803,37 @@ public:
 	FComputePipelineStateInitializer(FRHIComputeShader* InComputeShader)
 		: ComputeShader(InComputeShader)
 		, Flags(0)
+	{
+		checkf(!InComputeShader || !InComputeShader->HasSpecializationConstants(), TEXT("Shader requires specialization constants to be passed in using FComputePipelineStateInitializer."));
+	}
+
+	FComputePipelineStateInitializer(
+		FRHIComputeShader* InComputeShader,
+		const FSpecializationContainerType& InSpecializationConstants)
+		: ComputeShader(InComputeShader)
+		, SpecializationConstants(InSpecializationConstants)
+		, Flags(0)
 	{
 	}
 
 	FComputePipelineStateInitializer(
 		FRHIComputeShader* InComputeShader,
+		const FSpecializationContainerType& InSpecializationConstants,
 		uint8 InFlags)
 		: ComputeShader(InComputeShader)
+		, SpecializationConstants(InSpecializationConstants)
 		, Flags(InFlags)
 	{
 	}
 
 	bool operator==(const FComputePipelineStateInitializer& Other) const
 	{
-		return ComputeShader == Other.ComputeShader;
+		return (ComputeShader == Other.ComputeShader) && 
+			(SpecializationConstants == Other.SpecializationConstants);
 	}
 
 	FRHIComputeShader* ComputeShader;
+	FSpecializationContainerType SpecializationConstants;
 
 	// Note: these flags do NOT affect compilation of this PSO.
 	// The resulting object is invariant with respect to whatever is set here, they are
diff --git a/Engine/Source/Runtime/RHICore/Private/RHICoreShader.cpp b/Engine/Source/Runtime/RHICore/Private/RHICoreShader.cpp
--- a/Engine/Source/Runtime/RHICore/Private/RHICoreShader.cpp
+++ b/Engine/Source/Runtime/RHICore/Private/RHICoreShader.cpp
@@ -76,6 +76,19 @@ void UE::RHICore::SetupShaderCodeValidationData(FRHIShader* RHIShader, FShaderCo
 			RHIShader->DebugUBSizeValidationData.Append(ShaderCodeValidationExtension.ShaderCodeValidationUBSize);
 		}
 	}
+
+	if (RHIShader)
+	{
+		int32 SpecializationConstantValidationSize = 0;
+		const uint8* SpecializationConstantValidationData = ShaderCodeReader.FindOptionalDataAndSize(EShaderOptionalDataKey::SpecializationConstant, SpecializationConstantValidationSize);
+		if (SpecializationConstantValidationData && SpecializationConstantValidationSize > 0)
+		{
+			FBufferReader ArValidationExtensionData((void*)SpecializationConstantValidationData, SpecializationConstantValidationSize, false);
+			FSpecializationConstantData SpecializationConstantData;
+			ArValidationExtensionData << SpecializationConstantData;
+			RHIShader->DebugSpecializationExclusions = SpecializationConstantData.SpecializationConstantExclusions;
+		}
+	}
 #endif
 }
 
diff --git a/Engine/Source/Runtime/RenderCore/Private/GlobalShader.cpp b/Engine/Source/Runtime/RenderCore/Private/GlobalShader.cpp
--- a/Engine/Source/Runtime/RenderCore/Private/GlobalShader.cpp
+++ b/Engine/Source/Runtime/RenderCore/Private/GlobalShader.cpp
@@ -519,6 +519,48 @@ bool FGlobalShaderType::ShouldCompilePipeline(const FShaderPipelineType* ShaderP
 	return true;
 }
 
+bool FGlobalShaderType::ForEachCompiledPermutation(EShaderPlatform Platform, EShaderPermutationFlags PermutationFlags, TFunction<bool(int32)> Function)
+{
+	bool bSuccess = true;
+	if (FDataDrivenShaderPlatformInfo::GetSupportsSpecializationConstants(Platform) && (GetSpecializationCount() > 0))
+	{
+		// If the shader type includes specialization, collapse the permutations into valid IDs
+		TSet<int32> UnspecializedIDs;
+		UnspecializedIDs.Reserve(GetPermutationCount());
+		for (int32 PermutationId = 0; PermutationId < GetPermutationCount(); PermutationId++)
+		{
+			if (ShouldCompilePermutation(Platform, PermutationId, PermutationFlags))
+			{
+				UnspecializedIDs.Add(GetUnspecializedId(PermutationId));
+			}
+		}
+		for (int32 PermutationId : UnspecializedIDs)
+		{
+			if (!Function(PermutationId))
+			{
+				bSuccess = false;
+				break;
+			}
+		}
+	}
+	else
+	{
+		for (int32 PermutationId = 0; PermutationId < GetPermutationCount(); PermutationId++)
+		{
+			if (ShouldCompilePermutation(Platform, PermutationId, PermutationFlags))
+			{
+				if (!Function(PermutationId))
+				{
+					bSuccess = false;
+					break;
+				}
+			}
+		}
+	}
+	return bSuccess;
+}
+
+
 FGlobalShader::FGlobalShader(const ShaderMetaType::CompiledShaderInitializerType& Initializer)
 :	FShader(Initializer)
 {}
@@ -573,10 +615,10 @@ bool FGlobalShaderMapSection::Serialize(FShaderSerializeContext& Ctx)
 	return Super::Serialize(Ctx);
 }
 
-TShaderRef<FShader> FGlobalShaderMapSection::GetShader(FShaderType* ShaderType, int32 PermutationId) const
+TShaderRef<FShader> FGlobalShaderMapSection::GetShader(EShaderPlatform ShaderPlatform, FShaderType* ShaderType, int32 PermutationId) const
 {
-	FShader* Shader = GetContent()->GetShader(ShaderType, PermutationId);
-	return Shader ? TShaderRef<FShader>(Shader, *this) : TShaderRef<FShader>();
+	FShader* Shader = GetContent()->GetShader(ShaderType, FDataDrivenShaderPlatformInfo::GetSupportsSpecializationConstants(ShaderPlatform) ? ShaderType->GetUnspecializedId(PermutationId) : PermutationId);
+	return Shader ? TShaderRef<FShader>(Shader, *this, ShaderType->GetSpecializationValues(PermutationId)) : TShaderRef<FShader>();
 }
 
 FShaderPipelineRef FGlobalShaderMapSection::GetShaderPipeline(const FShaderPipelineType* PipelineType) const
@@ -608,7 +650,7 @@ FGlobalShaderMap::~FGlobalShaderMap()
 TShaderRef<FShader> FGlobalShaderMap::GetShader(FShaderType* ShaderType, int32 PermutationId) const
 {
 	FGlobalShaderMapSection* const* Section = SectionMap.Find(ShaderType->GetHashedShaderFilename());
-	return Section ? (*Section)->GetShader(ShaderType, PermutationId) : TShaderRef<FShader>();
+	return Section ? (*Section)->GetShader(Platform, ShaderType, PermutationId) : TShaderRef<FShader>();
 }
 
 FShaderPipelineRef FGlobalShaderMap::GetShaderPipeline(const FShaderPipelineType* ShaderPipelineType) const
@@ -674,13 +716,13 @@ bool FGlobalShaderMap::IsComplete(const ITargetPlatform* TargetPlatform) const
 		}
 
 		int32 PermutationCountToCompile = 0;
-		for (int32 PermutationId = 0; PermutationId < GlobalShaderType->GetPermutationCount(); PermutationId++)
+		const bool bSuccess = GlobalShaderType->ForEachCompiledPermutation(Platform, PermutationFlags, [&](int32 PermutationId)
 		{
-			if (GlobalShaderType->ShouldCompilePermutation(Platform, PermutationId, PermutationFlags)
-				&& !HasShader(GlobalShaderType, PermutationId))
-			{
-				return false;
-			}
+			return HasShader(GlobalShaderType, PermutationId);
+		});
+		if (!bSuccess)
+		{
+			return false;
 		}
 	}
 
diff --git a/Engine/Source/Runtime/RenderCore/Private/Shader.cpp b/Engine/Source/Runtime/RenderCore/Private/Shader.cpp
--- a/Engine/Source/Runtime/RenderCore/Private/Shader.cpp
+++ b/Engine/Source/Runtime/RenderCore/Private/Shader.cpp
@@ -242,10 +242,13 @@ FShaderType::FShaderType(
 	const TCHAR* InFunctionName,
 	uint32 InFrequency,
 	int32 InTotalPermutationCount,
+	int32 InTotalSpecializationCount,
 	ConstructSerializedType InConstructSerializedRef,
 	ConstructCompiledType InConstructCompiledRef,
 	ShouldCompilePermutationType InShouldCompilePermutationRef,
 	ShouldPrecachePermutationType InShouldPrecachePermutationRef,
+	GetUnspecializedIdType InGetUnspecializedIdRef,
+	GetSpecializationValuesType InGetSpecializationValuesRef,
 	GetRayTracingPayloadTypeType InGetRayTracingPayloadTypeRef,
 	GetShaderBindingLayoutType InGetShaderBindingLayoutTypeRef,
 #if WITH_EDITOR
@@ -270,10 +273,13 @@ FShaderType::FShaderType(
 	, Frequency(InFrequency)
 	, TypeSize(InTypeSize)
 	, TotalPermutationCount(InTotalPermutationCount)
+	, TotalSpecializationCount(InTotalSpecializationCount)
 	, ConstructSerializedRef(InConstructSerializedRef)
 	, ConstructCompiledRef(InConstructCompiledRef)
 	, ShouldCompilePermutationRef(InShouldCompilePermutationRef)
 	, ShouldPrecachePermutationRef(InShouldPrecachePermutationRef)
+	, GetUnspecializedIdRef(InGetUnspecializedIdRef)
+	, GetSpecializationValuesRef(InGetSpecializationValuesRef)
 	, GetRayTracingPayloadTypeRef(InGetRayTracingPayloadTypeRef)
 	, GetShaderBindingLayoutTypeRef(InGetShaderBindingLayoutTypeRef)
 #if WITH_EDITOR
@@ -477,6 +483,21 @@ EShaderPermutationPrecacheRequest FShaderType::ShouldPrecachePermutation(const F
 	return ShouldCompileShaderFrequency((EShaderFrequency)Frequency, Parameters.Platform) ? (*ShouldPrecachePermutationRef)(Parameters) : EShaderPermutationPrecacheRequest::NotUsed;
 }
 
+bool FShaderType::IsSpecialization(const int32 PermutationId) const
+{
+	return PermutationId != (*GetUnspecializedIdRef)(PermutationId);
+}
+
+int32 FShaderType::GetUnspecializedId(const int32 PermutationId) const
+{
+	return (TotalSpecializationCount > 0) ? (*GetUnspecializedIdRef)(PermutationId) : PermutationId;
+}
+
+FSpecializationContainerType FShaderType::GetSpecializationValues(const int32 PermutationId) const
+{
+	return (*GetSpecializationValuesRef)(PermutationId);
+}
+
 const FShaderBindingLayout* FShaderType::GetShaderBindingLayout(const FShaderPermutationParameters& Parameters) const
 {
 	return (*GetShaderBindingLayoutTypeRef)(Parameters);
diff --git a/Engine/Source/Runtime/RenderCore/Private/ShaderCore.cpp b/Engine/Source/Runtime/RenderCore/Private/ShaderCore.cpp
--- a/Engine/Source/Runtime/RenderCore/Private/ShaderCore.cpp
+++ b/Engine/Source/Runtime/RenderCore/Private/ShaderCore.cpp
@@ -1528,6 +1528,7 @@ void FShaderCompilerEnvironment::SerializeEverythingButFiles(FArchive& Ar)
 	Ar << ResourceTableMap.Resources;
 	Ar << UniformBufferMap;
 	Ar << RHIShaderBindingLayout;
+	Ar << SpecializationConstantValues;
 	Ar << FullPrecisionInPS;
 	if (Ar.IsLoading())
 	{
@@ -1544,6 +1545,7 @@ void FShaderCompilerEnvironment::SerializeCompilationDependencies(FArchive& Ar)
 	Ar << ResourceTableMap.Resources;
 	Ar << UniformBufferMap;
 	Ar << RHIShaderBindingLayout;
+	Ar << SpecializationConstantValues;
 	Ar << FullPrecisionInPS;
 	if (Ar.IsLoading())
 	{
@@ -1604,6 +1606,20 @@ void FShaderCompilerOutput::SerializeShaderCodeValidation()
 	}
 }
 
+void FShaderCompilerOutput::SerializeSpecializationConstantData()
+{
+	if (SpecializationConstantExclusions.Num() > 0)
+	{
+		FSpecializationConstantData SpecializationConstantData;
+		SpecializationConstantData.SpecializationConstantExclusions = SpecializationConstantExclusions;
+
+		TArray<uint8> WriterBytes;
+		FMemoryWriter Writer(WriterBytes);
+		Writer << SpecializationConstantData;
+		ShaderCode.AddOptionalData(FSpecializationConstantData::Key, WriterBytes.GetData(), WriterBytes.Num());
+	}
+}
+
 void FShaderCompilerOutput::SerializeShaderDiagnosticData()
 {
 	if (ShaderDiagnosticDatas.Num() > 0)
diff --git a/Engine/Source/Runtime/RenderCore/Private/ShaderParameterStruct.cpp b/Engine/Source/Runtime/RenderCore/Private/ShaderParameterStruct.cpp
--- a/Engine/Source/Runtime/RenderCore/Private/ShaderParameterStruct.cpp
+++ b/Engine/Source/Runtime/RenderCore/Private/ShaderParameterStruct.cpp
@@ -505,7 +505,7 @@ struct FShaderParameterReader
 
 #if DO_CHECK
 
-void ValidateShaderParameters(const TShaderRef<FShader>& Shader, const FShaderParametersMetadata* ParametersMetadata, const void* ParametersData)
+void ValidateShaderParameters(const TShaderRef<FShader>& Shader, const FShaderParametersMetadata* ParametersMetadata, const void* ParametersData, const FResourceExclusionArrayView ExclusionList)
 {
 	const FShaderParameterBindings& Bindings = Shader->Bindings;
 
@@ -532,7 +532,7 @@ void ValidateShaderParameters(const TShaderRef<FShader>& Shader, const FShaderPa
 			case UBMT_SAMPLER:
 			{
 				const FRHIResource* Resource = Reader.Read<const FRHIResource*>(Parameter);
-				if (!Resource)
+				if (!Resource && !ExclusionList.Contains(Parameter.BaseIndex))
 				{
 					EmitNullShaderParameterFatalError(Shader, ParametersMetadata, Parameter.ByteOffset);
 				}
@@ -541,7 +541,7 @@ void ValidateShaderParameters(const TShaderRef<FShader>& Shader, const FShaderPa
 			case UBMT_RDG_TEXTURE:
 			{
 				const FRDGTexture* GraphTexture = Reader.Read<const FRDGTexture*>(Parameter);
-				if (!GraphTexture)
+				if (!GraphTexture && !ExclusionList.Contains(Parameter.BaseIndex))
 				{
 					EmitNullShaderParameterFatalError(Shader, ParametersMetadata, Parameter.ByteOffset);
 				}
@@ -555,7 +555,7 @@ void ValidateShaderParameters(const TShaderRef<FShader>& Shader, const FShaderPa
 			{
 				const FRDGResource* GraphResource = Reader.Read<const FRDGResource*>(Parameter);
 
-				if (!GraphResource)
+				if (!GraphResource && !ExclusionList.Contains(Parameter.BaseIndex))
 				{
 					EmitNullShaderParameterFatalError(Shader, ParametersMetadata, Parameter.ByteOffset);
 				}
@@ -617,6 +617,32 @@ void ValidateShaderParameterResourcesRHI(const void* Contents, const FRHIUniform
 
 #endif // DO_CHECK
 
+#if RHI_INCLUDE_SHADER_DEBUG_DATA
+FResourceExclusionArray GetSpecializationResourceExclusions(const TShaderRef<FShader>& Shader)
+{
+	FRHIShader* RHIShader = Shader.GetRHIShaderBase(Shader->GetFrequency());
+	if (RHIShader->HasSpecializationConstants())
+	{
+		FResourceExclusionArray ResourceExclusions;
+
+		const uint32 SpecializationHash = GetTypeHash(Shader.GetSpecializationConstants());
+		const FShaderParameterBindings& Bindings = Shader->Bindings;
+		for (const FShaderParameterBindings::FResourceParameter& Parameter : Bindings.ResourceParameters)
+		{
+			// Check if this resource should be excluded from this combination of specialization constants
+			if (RHIShader->IsExcludedFromSpecialization(Parameter.BaseIndex, SpecializationHash))
+			{
+				ResourceExclusions.Add(Parameter.BaseIndex);
+			}
+		}
+
+		return ResourceExclusions;
+	}
+
+	return FResourceExclusionArray();
+}
+#endif // RHI_INCLUDE_SHADER_DEBUG_DATA
+
 template<typename BindingParameterType>
 FRHIShaderParameterResource ExtractShaderParameterResource(FShaderParameterReader Reader, const BindingParameterType& Parameter)
 {
@@ -837,7 +863,8 @@ static void ExtractShaderParameterResources(
 	FRHIBatchedShaderParameters& BatchedParameters,
 	const FShaderParameterBindings& Bindings,
 	TConstArrayView<uint8> ParametersData,
-	EUniformBufferBindingFlags UniformBufferBindingFlags)
+	EUniformBufferBindingFlags UniformBufferBindingFlags,
+	const FResourceExclusionArrayView ResourceExclusions)
 {
 	const FShaderParameterReader Reader(ParametersData);
 
@@ -858,6 +885,11 @@ static void ExtractShaderParameterResources(
 	{
 		for (const FShaderParameterBindings::FResourceParameter& Parameter : Bindings.ResourceParameters)
 		{
+			if (ResourceExclusions.Contains(Parameter.BaseIndex))
+			{
+				continue;
+			}
+
 			const FRHIShaderParameterResource ShaderParameterResource = ExtractShaderParameterResource(Reader, Parameter);
 			BatchedParameters.AddResourceParameter(ShaderParameterResource);
 		}
@@ -889,7 +921,8 @@ void SetShaderParameters(
 	FRHIBatchedShaderParameters& BatchedParameters,
 	const FShaderParameterBindings& Bindings,
 	const FShaderParametersMetadata* ParametersMetadata,
-	const void* InParametersData)
+	const void* InParametersData,
+	const FResourceExclusionArrayView ResourceExclusions)
 {
 	TConstArrayView<uint8> FullParametersData((const uint8*)InParametersData, ParametersMetadata->GetSize());
 
@@ -898,7 +931,7 @@ void SetShaderParameters(
 		BatchedParameters.SetShaderParameter(Parameter.BufferIndex, Parameter.BaseIndex, Parameter.ByteSize, FullParametersData.GetData() + Parameter.ByteOffset);
 	}
 
-	ExtractShaderParameterResources(BatchedParameters, Bindings, FullParametersData, EUniformBufferBindingFlags::Shader);
+	ExtractShaderParameterResources(BatchedParameters, Bindings, FullParametersData, EUniformBufferBindingFlags::Shader, ResourceExclusions);
 }
 
 /** Set shader's parameters from its parameters struct. */
@@ -908,13 +941,14 @@ inline void SetShaderParametersInternal(
 	TShaderRHI* ShaderRHI,
 	const FShaderParameterBindings& Bindings,
 	const FShaderParametersMetadata* ParametersMetadata,
-	const void* InParametersData)
+	const void* InParametersData,
+	const FResourceExclusionArrayView ResourceExclusions)
 {
 	checkf(Bindings.RootParameterBufferIndex == FShaderParameterBindings::kInvalidBufferIndex, TEXT("Can't use SetShaderParameters() for root parameter buffer index."));
 	checkf(!IsRayTracingShaderFrequency(ShaderRHI->GetFrequency()), TEXT("Can't use SetShaderParameters() with RHICmdList parameter for ray tracing shaders."));
 
 	FRHIBatchedShaderParameters& ShaderParameters = RHICmdList.GetScratchShaderParameters();
-	SetShaderParameters(ShaderParameters, Bindings, ParametersMetadata, InParametersData);
+	SetShaderParameters(ShaderParameters, Bindings, ParametersMetadata, InParametersData, ResourceExclusions);
 	RHICmdList.SetBatchedShaderParameters(ShaderRHI, ShaderParameters);
 }
 
@@ -923,9 +957,10 @@ void SetShaderParameters(
 	FRHIComputeShader* ShaderRHI,
 	const FShaderParameterBindings& Bindings,
 	const FShaderParametersMetadata* ParametersMetadata,
-	const void* ParametersData)
+	const void* ParametersData,
+	const FResourceExclusionArrayView ResourceExclusions)
 {
-	SetShaderParametersInternal(RHICmdList, ShaderRHI, Bindings, ParametersMetadata, ParametersData);
+	SetShaderParametersInternal(RHICmdList, ShaderRHI, Bindings, ParametersMetadata, ParametersData, ResourceExclusions);
 }
 
 void SetShaderParameters(
@@ -933,9 +968,10 @@ void SetShaderParameters(
 	FRHIGraphicsShader* ShaderRHI,
 	const FShaderParameterBindings& Bindings,
 	const FShaderParametersMetadata* ParametersMetadata,
-	const void* ParametersData)
+	const void* ParametersData,
+	const FResourceExclusionArrayView ResourceExclusions)
 {
-	SetShaderParametersInternal(RHICmdList, ShaderRHI, Bindings, ParametersMetadata, ParametersData);
+	SetShaderParametersInternal(RHICmdList, ShaderRHI, Bindings, ParametersMetadata, ParametersData, ResourceExclusions);
 }
 
 void SetShaderParameters(
@@ -943,9 +979,10 @@ void SetShaderParameters(
 	FRHIComputeShader* ShaderRHI,
 	const FShaderParameterBindings& Bindings,
 	const FShaderParametersMetadata* ParametersMetadata,
-	const void* ParametersData)
+	const void* ParametersData,
+	const FResourceExclusionArrayView ResourceExclusions)
 {
-	SetShaderParametersInternal(RHICmdList, ShaderRHI, Bindings, ParametersMetadata, ParametersData);
+	SetShaderParametersInternal(RHICmdList, ShaderRHI, Bindings, ParametersMetadata, ParametersData, ResourceExclusions);
 }
 
 #if RHI_RAYTRACING
@@ -1127,7 +1164,8 @@ void SetRayTracingShaderParameters(
 
 #endif
 
-	ExtractShaderParameterResources(BatchedParameters, Bindings, FullParametersData, EUniformBufferBindingFlags::StaticAndShader);
+	const FResourceExclusionArrayView NoExclusions;  // Specialization constants not supported in ray tracing yet
+	ExtractShaderParameterResources(BatchedParameters, Bindings, FullParametersData, EUniformBufferBindingFlags::StaticAndShader, NoExclusions);
 
 	// Root uniform buffer.
 	if (Bindings.RootParameterBufferIndex != FShaderParameterBindings::kInvalidBufferIndex)
diff --git a/Engine/Source/Runtime/RenderCore/Private/ShaderPipelineCache.cpp b/Engine/Source/Runtime/RenderCore/Private/ShaderPipelineCache.cpp
--- a/Engine/Source/Runtime/RenderCore/Private/ShaderPipelineCache.cpp
+++ b/Engine/Source/Runtime/RenderCore/Private/ShaderPipelineCache.cpp
@@ -1118,10 +1118,12 @@ bool FShaderPipelineCacheTask::Precompile(FRHICommandListImmediate& RHICmdList,
 			
 			if (LIKELY(!GShaderPipelineCacheDoNotPrecompileComputePSO))
 			{
-				FComputeShaderRHIRef ComputeInitializer = FShaderCodeLibrary::CreateComputeShader(Platform, PSO.ComputeDesc.ComputeShader);
-				if (ComputeInitializer.IsValid())
+				FComputeShaderRHIRef ComputeShaderRef = FShaderCodeLibrary::CreateComputeShader(Platform, PSO.ComputeDesc.ComputeShader);
+				if (ComputeShaderRef.IsValid() && !ComputeShaderRef->HasSpecializationConstants())  // :todo-jn: add specialization support to CS precaching
 				{
-					FComputePipelineState* ComputeResult = PipelineStateCache::GetAndOrCreateComputePipelineState(RHICmdList, ComputeInitializer, true);
+					FComputePipelineStateInitializer ComputeInitializer(ComputeShaderRef);
+					ComputeInitializer.bFromPSOFileCache = true;
+					FComputePipelineState* ComputeResult = PipelineStateCache::GetAndOrCreateComputePipelineState(RHICmdList, ComputeInitializer);
 					bOk = ComputeResult != nullptr;
 				}
 			}
diff --git a/Engine/Source/Runtime/RenderCore/Public/GlobalShader.h b/Engine/Source/Runtime/RenderCore/Public/GlobalShader.h
--- a/Engine/Source/Runtime/RenderCore/Public/GlobalShader.h
+++ b/Engine/Source/Runtime/RenderCore/Public/GlobalShader.h
@@ -117,6 +117,8 @@ public:
 #endif // WITH_EDITOR
 
 	RENDERCORE_API static bool ShouldCompilePipeline(const FShaderPipelineType* ShaderPipelineType, EShaderPlatform Platform, EShaderPermutationFlags Flags);
+
+	RENDERCORE_API bool ForEachCompiledPermutation(EShaderPlatform Platform, EShaderPermutationFlags PermutationFlags, TFunction<bool(int32)> Function);
 };
 
 class FGlobalShaderMapContent : public FShaderMapContent
@@ -166,7 +168,13 @@ private:
 		AssignContent(new FGlobalShaderMapContent(InPlatform, InHashedSourceFilename));
 	}
 
-	TShaderRef<FShader> GetShader(FShaderType* ShaderType, int32 PermutationId = 0) const;
+	UE_DEPRECATED(5.8, "Specifying the shader platform is now required when getting a shader from the FGlobalShaderMapSection.")
+		inline TShaderRef<FShader> GetShader(FShaderType* ShaderType, int32 PermutationId = 0) const
+	{
+		return GetShader(GetShaderPlatform(), ShaderType, PermutationId);
+	}
+
+	TShaderRef<FShader> GetShader(EShaderPlatform ShaderPlatform, FShaderType* ShaderType, int32 PermutationId = 0) const;
 	FShaderPipelineRef GetShaderPipeline(const FShaderPipelineType* PipelineType) const;
 };
 
@@ -190,7 +198,7 @@ public:
 	TShaderRef<ShaderType> GetShader(int32 PermutationId) const
 	{
 		TShaderRef<FShader> Shader = GetShader(&ShaderType::GetStaticType(), PermutationId);
-		checkf(Shader.IsValid(), TEXT("Failed to find shader type %s in Platform %s"), ShaderType::GetStaticType().GetName(), *LegacyShaderPlatformToShaderFormat(Platform).ToString());
+		checkf(Shader.IsValid(), TEXT("Failed to find permutation %d of shader type %s in Platform %s"), PermutationId, ShaderType::GetStaticType().GetName(), *LegacyShaderPlatformToShaderFormat(Platform).ToString());
 
 		// Validate that the permutation has been precached
 		ensureMsgf(ShaderType::ShouldPrecachePermutation(FGlobalShaderPermutationParameters(ShaderType::GetStaticType().GetFName(), Platform, PermutationId)) != EShaderPermutationPrecacheRequest::NotUsed,
@@ -203,7 +211,7 @@ public:
 	template<typename ShaderType>
 	TShaderRef<ShaderType> GetShader(const typename ShaderType::FPermutationDomain& PermutationVector) const
 	{
-		return GetShader<ShaderType>(PermutationVector.ToDimensionValueId());
+		return GetShader<ShaderType>(UE::ShaderPermutationUtils::ConvertToUnspecializedId(GetPlatform(), PermutationVector));
 	}
 
 	bool HasShader(FShaderType* Type, int32 PermutationId) const
@@ -216,6 +224,11 @@ public:
 		return GetShaderPipeline(ShaderPipelineType).IsValid();
 	}
 
+	EShaderPlatform GetPlatform() const
+	{
+		return Platform;
+	}
+
 	RENDERCORE_API bool IsEmpty() const;
 
 	/** Whether the shadermap has all the shader types it needs (i.e. ones that returned ShouldCompilePermutation) */
@@ -403,6 +416,7 @@ inline FGlobalShaderMap* GetGlobalShaderMap(ERHIFeatureLevel::Type FeatureLevel)
 		sizeof(ShaderClass), \
 		sizeof(ShaderClass::FEditorContent), \
 		ShaderClass::GetRootParametersMetadata() \
+		ShaderClass::FPermutationDomain::SpecializationCount, \
 		)
 		*/
 
diff --git a/Engine/Source/Runtime/RenderCore/Public/RenderGraphUtils.h b/Engine/Source/Runtime/RenderCore/Public/RenderGraphUtils.h
--- a/Engine/Source/Runtime/RenderCore/Public/RenderGraphUtils.h
+++ b/Engine/Source/Runtime/RenderCore/Public/RenderGraphUtils.h
@@ -455,7 +455,7 @@ namespace FComputeShaderUtils
 			const typename TShaderClass::FParameters& Parameters)
 		{
 			FRHIComputeShader* ShaderRHI = ComputeShader.GetComputeShader();
-			SetComputePipelineState(RHICmdList, ShaderRHI);
+			SetComputePipelineState(RHICmdList, { ShaderRHI, ComputeShader.GetSpecializationConstants() });
 			SetShaderParameters(RHICmdList, ComputeShader, ShaderRHI, ParametersMetadata, Parameters);
 		}
 
diff --git a/Engine/Source/Runtime/RenderCore/Public/Shader.h b/Engine/Source/Runtime/RenderCore/Public/Shader.h
--- a/Engine/Source/Runtime/RenderCore/Public/Shader.h
+++ b/Engine/Source/Runtime/RenderCore/Public/Shader.h
@@ -999,10 +999,12 @@ class TShaderRefBase
 public:
 	TShaderRefBase() : ShaderContent(nullptr), ShaderMap(nullptr) {}
 	TShaderRefBase(ShaderType* InShader, const FShaderMapBase& InShaderMap) : ShaderContent(InShader), ShaderMap(&InShaderMap) { checkSlow(InShader); }
+	TShaderRefBase(ShaderType* InShader, const FShaderMapBase& InShaderMap, const FSpecializationContainerType& InSpecConsts)
+		: ShaderContent(InShader), ShaderMap(&InShaderMap), SpecConsts(InSpecConsts) { checkSlow(InShader); }
 	TShaderRefBase(const TShaderRefBase&) = default;
 
 	template<typename OtherShaderType, typename OtherPointerTableType>
-	TShaderRefBase(const TShaderRefBase<OtherShaderType, OtherPointerTableType>& Rhs) : ShaderContent(Rhs.GetShader()), ShaderMap(Rhs.GetShaderMap()) {}
+	TShaderRefBase(const TShaderRefBase<OtherShaderType, OtherPointerTableType>& Rhs) : ShaderContent(Rhs.GetShader()), ShaderMap(Rhs.GetShaderMap()), SpecConsts(Rhs.GetSpecializationConstants()) {}
 
 	TShaderRefBase& operator=(const TShaderRefBase&) = default;
 
@@ -1011,19 +1013,20 @@ public:
 	{
 		ShaderContent = Rhs.GetShader();
 		ShaderMap = Rhs.GetShaderMap();
+		SpecConsts = Rhs.GetSpecializationConstants();
 		return *this;
 	}
 
 	template<typename OtherShaderType, typename OtherPointerTableType>
 	static TShaderRefBase<ShaderType, PointerTableType> Cast(const TShaderRefBase<OtherShaderType, OtherPointerTableType>& Rhs)
 	{
-		return TShaderRefBase<ShaderType, PointerTableType>(static_cast<ShaderType*>(Rhs.GetShader()), Rhs.GetShaderMap());
+		return TShaderRefBase<ShaderType, PointerTableType>(static_cast<ShaderType*>(Rhs.GetShader()), Rhs.GetShaderMap(), Rhs.GetSpecializationConstants());
 	}
 
 	template<typename OtherShaderType, typename OtherPointerTableType>
 	static TShaderRefBase<ShaderType, PointerTableType> ReinterpretCast(const TShaderRefBase<OtherShaderType, OtherPointerTableType>& Rhs)
 	{
-		return TShaderRefBase<ShaderType, PointerTableType>(reinterpret_cast<ShaderType*>(Rhs.GetShader()), Rhs.GetShaderMap());
+		return TShaderRefBase<ShaderType, PointerTableType>(reinterpret_cast<ShaderType*>(Rhs.GetShader()), Rhs.GetShaderMap(), Rhs.GetSpecializationConstants());
 	}
 
 	inline bool IsValid() const { return ShaderContent != nullptr; }
@@ -1158,9 +1161,19 @@ public:
 	}
 #endif // RHI_RAYTRACING
 
+	inline const FSpecializationContainerType& GetSpecializationConstants() const
+	{
+		return SpecConsts;
+	}
+
+	inline void SetSpecializationConstants(FSpecializationContainerType& InSpecConsts)
+	{
+		SpecConsts = InSpecConsts;
+	}
+
 private:
-	TShaderRefBase(ShaderType* InShader, const FShaderMapBase* InShaderMap)
-		: ShaderContent(InShader), ShaderMap(InShaderMap)
+	TShaderRefBase(ShaderType* InShader, const FShaderMapBase* InShaderMap, const FSpecializationContainerType& InSpecConsts)
+		: ShaderContent(InShader), ShaderMap(InShaderMap), SpecConsts(InSpecConsts)
 	{
 		checkSlow((!InShader && !InShaderMap) || (InShader && InShaderMap));
 	}
@@ -1185,6 +1198,8 @@ private:
 		return !operator==(Lhs, Rhs);
 	}
 
+protected:
+	FSpecializationContainerType SpecConsts;
 };
 
 template<typename ShaderType>
@@ -1229,6 +1244,8 @@ public:
 	typedef FShader* (*ConstructCompiledType)(const FShader::CompiledShaderInitializerType& Initializer);
 	typedef bool (*ShouldCompilePermutationType)(const FShaderPermutationParameters&);
 	typedef EShaderPermutationPrecacheRequest (*ShouldPrecachePermutationType)(const FShaderPermutationParameters&);
+	typedef int32(*GetUnspecializedIdType)(const int32 PermutationId);
+	typedef FSpecializationContainerType (*GetSpecializationValuesType)(const int32 PermutationId);
 	typedef ERayTracingPayloadType(*GetRayTracingPayloadTypeType)(const int32 PermutationId);
 	typedef const FShaderBindingLayout*(*GetShaderBindingLayoutType)(const FShaderPermutationParameters&);
 #if WITH_EDITOR
@@ -1265,10 +1282,13 @@ public:
 		const TCHAR* InFunctionName,
 		uint32 InFrequency,
 		int32 TotalPermutationCount,
+		int32 TotalSpecializationCount,
 		ConstructSerializedType InConstructSerializedRef,
 		ConstructCompiledType InConstructCompiledRef,
 		ShouldCompilePermutationType InShouldCompilePermutationRef,
 		ShouldPrecachePermutationType InShouldPrecachePermutationRef,
+		GetUnspecializedIdType InGetUnspecializedIdRef,
+		GetSpecializationValuesType InGetSpecializationValuesRef,
 		GetRayTracingPayloadTypeType InGetRayTracingPayloadTypeRef,
 		GetShaderBindingLayoutType InGetShaderBindingLayoutTypeRef,
 #if WITH_EDITOR
@@ -1281,7 +1301,6 @@ public:
 #if WITH_EDITOR
 		, GetPermutationIdStringType InGetPermutationIdStringRef = nullptr
 #endif
-
 	);
 
 	RENDERCORE_API virtual ~FShaderType();
@@ -1292,7 +1311,10 @@ public:
 
 	RENDERCORE_API bool ShouldCompilePermutation(const FShaderPermutationParameters& Parameters) const;
 	RENDERCORE_API EShaderPermutationPrecacheRequest ShouldPrecachePermutation(const FShaderPermutationParameters& Parameters) const;
-	RENDERCORE_API const FShaderBindingLayout* GetShaderBindingLayout(const FShaderPermutationParameters& Parameters) const; 
+	RENDERCORE_API bool IsSpecialization(const int32 PermutationId) const;
+	RENDERCORE_API int32 GetUnspecializedId(const int32 PermutationId) const;
+	RENDERCORE_API FSpecializationContainerType GetSpecializationValues(const int32 PermutationId) const;
+	RENDERCORE_API const FShaderBindingLayout* GetShaderBindingLayout(const FShaderPermutationParameters& Parameters) const;
 
 #if WITH_EDITOR
 	RENDERCORE_API void ModifyCompilationEnvironment(const FShaderPermutationParameters& Parameters, FShaderCompilerEnvironment& OutEnvironment) const;
@@ -1472,6 +1494,11 @@ public:
 		return TotalPermutationCount;
 	}
 
+	inline int32 GetSpecializationCount() const
+	{
+		return TotalSpecializationCount;
+	}
+
 	/** Returns the meta data for the root shader parameter struct. */
 	inline const FShaderParametersMetadata* GetRootParametersMetadata() const
 	{
@@ -1506,11 +1533,14 @@ private:
 	uint32 Frequency;
 	uint32 TypeSize;
 	int32 TotalPermutationCount;
+	int32 TotalSpecializationCount;
 
 	ConstructSerializedType ConstructSerializedRef;
 	ConstructCompiledType ConstructCompiledRef;
 	ShouldCompilePermutationType ShouldCompilePermutationRef;
 	ShouldPrecachePermutationType ShouldPrecachePermutationRef;
+	GetUnspecializedIdType GetUnspecializedIdRef;
+	GetSpecializationValuesType GetSpecializationValuesRef;
 	GetRayTracingPayloadTypeType GetRayTracingPayloadTypeRef;
 	GetShaderBindingLayoutType GetShaderBindingLayoutTypeRef;
 #if WITH_EDITOR
@@ -1625,6 +1655,17 @@ struct FShaderCompiledShaderInitializerType
 	{ return ShaderClass::ShouldCompilePermutation(static_cast<const typename ShaderClass::FPermutationParameters&>(Parameters)); } \
 	static EShaderPermutationPrecacheRequest ShouldPrecachePermutationImpl(const FShaderPermutationParameters& Parameters) \
 	{ return ShaderClass::ShouldPrecachePermutation(static_cast<const typename ShaderClass::FPermutationParameters&>(Parameters)); } \
+	static int32 GetUnspecializedIdImpl(const int32 PermutationID) \
+	{ return typename ShaderClass::FPermutationDomain(PermutationID).ToDimensionUnspecializedId(); } \
+	static FSpecializationContainerType GetSpecializationValuesImpl(const int32 PermutationID) \
+	{ \
+		if constexpr (ShaderClass::FPermutationDomain::SpecializationCount > 0) {\
+			FSpecializationContainerType SpecConsts; \
+			SpecConsts.Reserve(ShaderClass::FPermutationDomain::SpecializationCount); \
+			typename ShaderClass::FPermutationDomain(PermutationID).ToSpecializationConstantArray(SpecConsts); \
+			return SpecConsts; \
+		} else { return FSpecializationContainerType(); } \
+	} \
 	SHADER_DECLARE_EDITOR_VTABLE(ShaderClass)
 
 
@@ -1681,6 +1722,8 @@ struct FShaderCompiledShaderInitializerType
 	ShaderClass::ConstructCompiledInstance, \
 	ShaderClass::ShouldCompilePermutationImpl, \
 	ShaderClass::ShouldPrecachePermutationImpl, \
+	ShaderClass::GetUnspecializedIdImpl, \
+	ShaderClass::GetSpecializationValuesImpl, \
 	ShaderClass::GetRayTracingPayloadType, \
 	ShaderClass::GetShaderBindingLayout \
 	SHADER_TYPE_EDITOR_VTABLE(ShaderClass)
@@ -1699,6 +1742,7 @@ struct FShaderCompiledShaderInitializerType
 			FunctionName, \
 			Frequency, \
 			ShaderClass::FPermutationDomain::PermutationCount, \
+			ShaderClass::FPermutationDomain::SpecializationCount, \
 			SHADER_TYPE_VTABLE(ShaderClass), \
 			sizeof(ShaderClass), \
 			ShaderClass::GetRootParametersMetadata() \
@@ -1721,6 +1765,7 @@ struct FShaderCompiledShaderInitializerType
 			FunctionName, \
 			Frequency, \
 			ShaderClass::FPermutationDomain::PermutationCount, \
+			ShaderClass::FPermutationDomain::SpecializationCount, \
 			SHADER_TYPE_VTABLE(ShaderClass), \
 			sizeof(ShaderClass), \
 			ShaderClass::GetRootParametersMetadata() \
@@ -1743,6 +1788,7 @@ struct FShaderCompiledShaderInitializerType
 			ShaderClass::GetFunctionName(), \
 			Frequency, \
 			ShaderClass::FPermutationDomain::PermutationCount, \
+			ShaderClass::FPermutationDomain::SpecializationCount, \
 			SHADER_TYPE_VTABLE(ShaderClass), \
 			sizeof(ShaderClass), \
 			ShaderClass::GetRootParametersMetadata() \
@@ -1767,6 +1813,7 @@ struct FShaderCompiledShaderInitializerType
 			ShaderClass::GetFunctionName(), \
 			Frequency, \
 			ShaderClass::FPermutationDomain::PermutationCount, \
+			ShaderClass::FPermutationDomain::SpecializationCount, \
 			SHADER_TYPE_VTABLE(ShaderClass), \
 			sizeof(ShaderClass), \
 			ShaderClass::GetRootParametersMetadata() \
@@ -1789,6 +1836,7 @@ struct FShaderCompiledShaderInitializerType
 			ShaderClass::GetFunctionName(), \
 			Frequency, \
 			ShaderClass::FPermutationDomain::PermutationCount, \
+			ShaderClass::FPermutationDomain::SpecializationCount, \
 			SHADER_TYPE_VTABLE(ShaderClass), \
 			sizeof(ShaderClass), \
 			ShaderClass::GetRootParametersMetadata() \
@@ -1808,10 +1856,13 @@ struct FShaderCompiledShaderInitializerType
 		const TCHAR* InFunctionName, \
 		uint32 InFrequency, \
 		int32 InTotalPermutationCount, \
+		int32 InTotalSpecializationCount, \
 		ConstructSerializedType InConstructSerializedRef, \
 		ConstructCompiledType InConstructCompiledRef, \
 		ShouldCompilePermutationType InShouldCompilePermutationRef, \
 		ShouldPrecachePermutationType InShouldPrecachePermutationRef, \
+		GetUnspecializedIdType InGetUnspecializedIdRef, \
+		GetSpecializationValuesType InGetSpecializationValuesRef, \
 		GetRayTracingPayloadTypeType InGetRayTracingPayloadTypeRef, \
 		GetShaderBindingLayoutType InGetShaderBindingLayoutTypeRef, \
 		ModifyCompilationEnvironmentType InModifyCompilationEnvironmentRef, \
@@ -1829,10 +1880,13 @@ struct FShaderCompiledShaderInitializerType
 			InFunctionName, \
 			InFrequency, \
 			InTotalPermutationCount, \
+			InTotalSpecializationCount, \
 			InConstructSerializedRef, \
 			InConstructCompiledRef, \
 			InShouldCompilePermutationRef, \
 			InShouldPrecachePermutationRef, \
+			InGetUnspecializedIdRef, \
+			InGetSpecializationValuesRef, \
 			InGetRayTracingPayloadTypeRef, \
 			InGetShaderBindingLayoutTypeRef, \
 			InModifyCompilationEnvironmentRef, \
@@ -1858,10 +1912,13 @@ struct FShaderCompiledShaderInitializerType
 		const TCHAR* InFunctionName, \
 		uint32 InFrequency, \
 		int32 InTotalPermutationCount, \
+		int32 InTotalSpecializationCount, \
 		ConstructSerializedType InConstructSerializedRef, \
 		ConstructCompiledType InConstructCompiledRef, \
 		ShouldCompilePermutationType InShouldCompilePermutationRef, \
 		ShouldPrecachePermutationType InShouldPrecachePermutationRef, \
+		GetUnspecializedIdType InGetUnspecializedIdRef, \
+		GetSpecializationValuesType InGetSpecializationValuesRef, \
 		GetRayTracingPayloadTypeType InGetRayTracingPayloadTypeRef, \
 		GetShaderBindingLayoutType InGetShaderBindingLayoutTypeRef, \
 		uint32 InTypeSize, \
@@ -1875,10 +1932,13 @@ struct FShaderCompiledShaderInitializerType
 			InFunctionName, \
 			InFrequency, \
 			InTotalPermutationCount, \
+			InTotalSpecializationCount, \
 			InConstructSerializedRef, \
 			InConstructCompiledRef, \
 			InShouldCompilePermutationRef, \
 			InShouldPrecachePermutationRef, \
+			InGetUnspecializedIdRef, \
+			InGetSpecializationValuesRef, \
 			InGetRayTracingPayloadTypeRef, \
 			InGetShaderBindingLayoutTypeRef, \
 			InTypeSize, \
@@ -2306,7 +2366,7 @@ public:
 	template<typename ShaderType>
 	ShaderType* GetShader( const typename ShaderType::FPermutationDomain& PermutationVector ) const
 	{
-		return GetShader<ShaderType>( PermutationVector.ToDimensionValueId() );
+		return GetShader<ShaderType>(UE::ShaderPermutationUtils::ConvertToUnspecializedId(GetShaderPlatform(), PermutationVector));
 	}
 
 	/** Finds the shader with the given type.  May return NULL. */
@@ -2606,7 +2666,9 @@ public:
 		const typename ShaderType::ShaderMapType* ShaderIndex,
 		const typename ShaderType::FPermutationDomain& PermutationVector)
 		: TShaderRef<ShaderType>(ShaderIndex->template GetShader<ShaderType>(PermutationVector.ToDimensionValueId())) // gcc3 needs the template quantifier so it knows the < is not a less-than
-	{ }
+	{
+		checkf(TShaderRef<ShaderType>::SpecConsts.Num() == ShaderType::FPermutationDomain::SpecializationCount, TEXT("Shader returned without specialization constant values."));
+	}
 };
 
 /** A reference to an optional shader, initialized with a shader type from a shader map if it is available or nullptr if it is not. */
diff --git a/Engine/Source/Runtime/RenderCore/Public/ShaderCompilerCore.h b/Engine/Source/Runtime/RenderCore/Public/ShaderCompilerCore.h
--- a/Engine/Source/Runtime/RenderCore/Public/ShaderCompilerCore.h
+++ b/Engine/Source/Runtime/RenderCore/Public/ShaderCompilerCore.h
@@ -540,6 +540,8 @@ struct FShaderCompilerOutput
 	TArray<FShaderCodeValidationType> ParametersUAVTypeToValidate;
 	TArray<FShaderCodeValidationUBSize> ParametersUBSizeToValidate;
 
+	TMap<uint32, TArray<uint32>> SpecializationConstantExclusions;
+
 	TArray<FShaderDiagnosticData> ShaderDiagnosticDatas;
 
 	/** Use this field to store the shader source code if it's modified as part of the shader format's compilation process. This should only be set when 
@@ -565,6 +567,9 @@ struct FShaderCompilerOutput
 	/** Add optional data in ShaderCode to perform additional shader input validation at runtime*/
 	RENDERCORE_API void SerializeShaderCodeValidation();
 
+	/** Add specialization constant related data*/
+	RENDERCORE_API void SerializeSpecializationConstantData();
+
 	/** Add optional diagnostic data in ShaderCode to perform assert translation at runtime*/
 	RENDERCORE_API void SerializeShaderDiagnosticData();
 
diff --git a/Engine/Source/Runtime/RenderCore/Public/ShaderCore.h b/Engine/Source/Runtime/RenderCore/Public/ShaderCore.h
--- a/Engine/Source/Runtime/RenderCore/Public/ShaderCore.h
+++ b/Engine/Source/Runtime/RenderCore/Public/ShaderCore.h
@@ -559,6 +559,9 @@ struct FShaderCompilerEnvironment
 	// Serialized version of the shader binding layout which can be used during platform specific shader code generation and serialization
 	FRHIShaderBindingLayout RHIShaderBindingLayout;
 
+	// Array of possible values for each specialization constants
+	TArray<TArray<int32>> SpecializationConstantValues;
+
 	const ITargetPlatform* TargetPlatform = nullptr;
 
 	// Used for mobile platforms to allow per shader/material precision modes
@@ -567,7 +570,7 @@ struct FShaderCompilerEnvironment
 	/** Default constructor. */
 	RENDERCORE_API FShaderCompilerEnvironment();
 
-	/** Constructor used when enviroment is constructed temporarily purely for the purpose of hashing for inclusion in DDC keys. */
+	/** Constructor used when environment is constructed temporarily purely for the purpose of hashing for inclusion in DDC keys. */
 	RENDERCORE_API FShaderCompilerEnvironment(FMemoryHasherBlake3& Hasher);
 
 	/**
@@ -753,6 +756,7 @@ enum class EShaderOptionalDataKey : uint8
 	ShaderBindingLayout  = uint8('s'),
 	EntryPoint           = uint8('e'),
 	GlobalUniformBufferConfiguration = uint8('G'),
+	SpecializationConstant = uint8('k'),
 };
 
 enum class EShaderResourceUsageFlags : uint8
@@ -905,6 +909,24 @@ struct FShaderCodeVendorExtension
 };
 
 
+struct FSpecializationConstantData
+{
+	// for FindOptionalData() and AddOptionalData()
+	static constexpr EShaderOptionalDataKey Key = EShaderOptionalDataKey::SpecializationConstant;
+	static constexpr uint16 StaticVersion = 0;
+
+	TMap<uint32, TArray<uint32>> SpecializationConstantExclusions;
+	uint16 Version = StaticVersion;
+
+	friend FArchive& operator<<(FArchive& Ar, FSpecializationConstantData& Extension)
+	{
+		Ar << Extension.Version;
+		Ar << Extension.SpecializationConstantExclusions;
+		return Ar;
+	}
+};
+
+
 inline FArchive& operator<<(FArchive& Ar, FShaderCodeValidationStride& ShaderCodeValidationStride)
 {
 	return Ar << ShaderCodeValidationStride.BindPoint << ShaderCodeValidationStride.Stride;
diff --git a/Engine/Source/Runtime/RenderCore/Public/ShaderParameterStruct.h b/Engine/Source/Runtime/RenderCore/Public/ShaderParameterStruct.h
--- a/Engine/Source/Runtime/RenderCore/Public/ShaderParameterStruct.h
+++ b/Engine/Source/Runtime/RenderCore/Public/ShaderParameterStruct.h
@@ -118,17 +118,32 @@ inline void ValidateShaderParameterResourcesRHI(const void* Contents, const FRHI
 /** Raise fatal error when a required shader parameter has not been set. */
 extern RENDERCORE_API void EmitNullShaderParameterFatalError(const TShaderRef<FShader>& Shader, const FShaderParametersMetadata* ParametersMetadata, uint16 MemberOffset);
 
+/** Generate a list of excluded resources from the specialization constants. */
+using FResourceExclusionArray = TArray<uint16>;
+using FResourceExclusionArrayView = TArrayView<uint16>;
+#if RHI_INCLUDE_SHADER_DEBUG_DATA
+extern RENDERCORE_API FResourceExclusionArray GetSpecializationResourceExclusions(const TShaderRef<FShader>& Shader);
+#else
+inline FResourceExclusionArray GetSpecializationResourceExclusions(const TShaderRef<FShader>& Shader) { return FResourceExclusionArray();  }
+#endif
+
 /** Validates that all resource parameters of a shader are set. */
 #if DO_CHECK
-extern RENDERCORE_API void ValidateShaderParameters(const TShaderRef<FShader>& Shader, const FShaderParametersMetadata* ParametersMetadata, const void* ParametersData);
+extern RENDERCORE_API void ValidateShaderParameters(const TShaderRef<FShader>& Shader, const FShaderParametersMetadata* ParametersMetadata, const void* ParametersData, const FResourceExclusionArrayView ExclusionList);
 #else
-inline void ValidateShaderParameters(const TShaderRef<FShader>& Shader, const FShaderParametersMetadata* ParametersMetadata, const void* ParametersData) {}
+inline void ValidateShaderParameters(const TShaderRef<FShader>& Shader, const FShaderParametersMetadata* ParametersMetadata, const void* ParametersData, const FResourceExclusionArrayView ExclusionList) {}
 #endif
 
+inline void ValidateShaderParameters(const TShaderRef<FShader>& Shader, const FShaderParametersMetadata* ParametersMetadata, const void* ParametersData)
+{
+	ValidateShaderParameters(Shader, ParametersMetadata, ParametersData, GetSpecializationResourceExclusions(Shader));
+}
+
+
 template<typename TShaderClass>
 inline void ValidateShaderParameters(const TShaderRef<TShaderClass>& Shader, const typename TShaderClass::FParameters& Parameters)
 {
-	return ValidateShaderParameters(Shader, TShaderClass::FParameters::FTypeInfo::GetStructMetadata(), &Parameters);
+	ValidateShaderParameters(Shader, TShaderClass::FParameters::FTypeInfo::GetStructMetadata(), &Parameters);
 }
 
 template<typename TParameterType>
@@ -248,7 +263,18 @@ RENDERCORE_API void SetShaderParameters(
 	FRHIBatchedShaderParameters& BatchedParameters,
 	const FShaderParameterBindings& Bindings,
 	const FShaderParametersMetadata* ParametersMetadata,
-	const void* ParametersData);
+	const void* ParametersData,
+	const FResourceExclusionArrayView ResourceExclusions);
+
+inline void SetShaderParameters(
+	FRHIBatchedShaderParameters& BatchedParameters,
+	const FShaderParameterBindings& Bindings,
+	const FShaderParametersMetadata* ParametersMetadata,
+	const void* ParametersData)
+{
+	const FResourceExclusionArrayView NoResourceExclusions;
+	SetShaderParameters(BatchedParameters, Bindings, ParametersMetadata, ParametersData, NoResourceExclusions);
+}
 
 template<typename TShaderClass>
 inline void SetShaderParameters(
@@ -283,21 +309,45 @@ RENDERCORE_API void SetShaderParameters(
 	FRHIComputeShader* ShaderRHI,
 	const FShaderParameterBindings& Bindings,
 	const FShaderParametersMetadata* ParametersMetadata,
-	const void* ParametersData);
+	const void* ParametersData,
+	const FResourceExclusionArrayView ResourceExclusions);
+
+//UE_DEPRECATED(5.8, "Calls to SetShaderParameters now require ResourceExclusions to be specified.")
+inline void SetShaderParameters(FRHIComputeCommandList& RHICmdList, FRHIComputeShader* ShaderRHI, const FShaderParameterBindings& Bindings, const FShaderParametersMetadata* ParametersMetadata, const void* ParametersData)
+{
+	checkf(!ShaderRHI->HasSpecializationConstants(), TEXT("SetShaderParameters called with a shader using specialization constants, a ResourceExclusion array needs to be specified."));
+	SetShaderParameters(RHICmdList, ShaderRHI, Bindings, ParametersMetadata, ParametersData, FResourceExclusionArrayView());
+}
 
 RENDERCORE_API void SetShaderParameters(
 	FRHICommandList& RHICmdList,
 	FRHIGraphicsShader* ShaderRHI,
 	const FShaderParameterBindings& Bindings,
 	const FShaderParametersMetadata* ParametersMetadata,
-	const void* ParametersData);
+	const void* ParametersData,
+	const FResourceExclusionArrayView ResourceExclusions);
+
+//UE_DEPRECATED(5.8, "Calls to SetShaderParameters now require ResourceExclusions to be specified.")
+inline void SetShaderParameters(FRHICommandList& RHICmdList, FRHIGraphicsShader* ShaderRHI, const FShaderParameterBindings& Bindings, const FShaderParametersMetadata* ParametersMetadata, const void* ParametersData)
+{
+	checkf(!ShaderRHI->HasSpecializationConstants(), TEXT("SetShaderParameters called with a shader using specialization constants, a ResourceExclusion array needs to be specified."));
+	SetShaderParameters(RHICmdList, ShaderRHI, Bindings, ParametersMetadata, ParametersData, FResourceExclusionArrayView());
+}
 
 RENDERCORE_API void SetShaderParameters(
 	FRHICommandList& RHICmdList,
 	FRHIComputeShader* ShaderRHI,
 	const FShaderParameterBindings& Bindings,
 	const FShaderParametersMetadata* ParametersMetadata,
-	const void* ParametersData);
+	const void* ParametersData,
+	const FResourceExclusionArrayView ResourceExclusions);
+
+//UE_DEPRECATED(5.8, "Calls to SetShaderParameters now require ResourceExclusions to be specified.")
+inline void SetShaderParameters(FRHICommandList& RHICmdList, FRHIComputeShader* ShaderRHI, const FShaderParameterBindings& Bindings, const FShaderParametersMetadata* ParametersMetadata, const void* ParametersData)
+{
+	checkf(!ShaderRHI->HasSpecializationConstants(), TEXT("SetShaderParameters called with a shader using specialization constants, a ResourceExclusion array needs to be specified."));
+	SetShaderParameters(RHICmdList, ShaderRHI, Bindings, ParametersMetadata, ParametersData, FResourceExclusionArrayView());
+}
 
 template<typename TRHICmdList, typename TShaderClass, typename TShaderRHI>
 inline void SetShaderParameters(
@@ -307,11 +357,13 @@ inline void SetShaderParameters(
 	const FShaderParametersMetadata* ParametersMetadata,
 	const typename TShaderClass::FParameters& Parameters)
 {
-	ValidateShaderParameters(Shader, ParametersMetadata, &Parameters);
+	FResourceExclusionArray ResourceExclusions = GetSpecializationResourceExclusions(Shader);
+
+	ValidateShaderParameters(Shader, ParametersMetadata, &Parameters, ResourceExclusions);
 
 	// TODO(RDG): Once all shader sets their parameter through this, can refactor RHI so all shader parameters get sets through a single RHI function call.
 
-	SetShaderParameters(RHICmdList, ShaderRHI, Shader->Bindings, ParametersMetadata, &Parameters);
+	SetShaderParameters(RHICmdList, ShaderRHI, Shader->Bindings, ParametersMetadata, &Parameters, ResourceExclusions);
 }
 
 template<typename TRHICmdList, typename TShaderClass, typename TShaderRHI>
diff --git a/Engine/Source/Runtime/RenderCore/Public/ShaderPermutation.h b/Engine/Source/Runtime/RenderCore/Public/ShaderPermutation.h
--- a/Engine/Source/Runtime/RenderCore/Public/ShaderPermutation.h
+++ b/Engine/Source/Runtime/RenderCore/Public/ShaderPermutation.h
@@ -14,6 +14,8 @@ inline constexpr bool bLogPermutations = false;
 
 struct FShaderCompilerEnvironment;
 
+
+
 // Flags that can specialize shader permutations compiled for specific platforms
 enum class EShaderPermutationFlags : uint32
 {
@@ -44,6 +46,7 @@ struct FShaderPermutationParameters
 };
 
 /** Defines at compile time a boolean permutation dimension. */
+template <bool TIsSpecialization>
 struct FShaderPermutationBool
 {
 	/** Setup the dimension's type in permutation domain as boolean. */
@@ -57,6 +60,9 @@ struct FShaderPermutationBool
 	 */
 	static constexpr bool IsMultiDimensional = false;
 
+	/** Specialization constant data. */
+	static constexpr bool IsSpecialization = TIsSpecialization;
+
 
 	/** Converts dimension boolean value to dimension's value id. */
 	static int32 ToDimensionValueId(Type E)
@@ -76,11 +82,19 @@ struct FShaderPermutationBool
 		checkf(PermutationId == 0 || PermutationId == 1, TEXT("Invalid shader permutation dimension id %i."), PermutationId);
 		return PermutationId == 1;
 	}
+
+	/** Returns all the possible values of a specialization. */
+	static void GetPossibleValues(TArray<int32>& OutValues)
+	{
+		checkf(TIsSpecialization, TEXT("Getting possible values should only be called for specializations!"));
+		OutValues.Add(0);
+		OutValues.Add(1);
+	}
 };
 
 
 /** Defines at compile time a permutation dimension made of int32 from 0 to N -1. */
-template <typename TType, int32 TDimensionSize, int32 TFirstValue=0>
+template <bool TIsSpecialization, typename TType, int32 TDimensionSize, int32 TFirstValue=0>
 struct TShaderPermutationInt
 {
 	/** Setup the dimension's type in permutation domain as integer. */
@@ -94,6 +108,9 @@ struct TShaderPermutationInt
 	 */
 	static constexpr bool IsMultiDimensional = false;
 	
+	/** Specialization constant data. */
+	static constexpr bool IsSpecialization = TIsSpecialization;
+
 	/** Min and max values. */
 	static constexpr Type MinValue = static_cast<Type>(TFirstValue);
 	static constexpr Type MaxValue = static_cast<Type>(TFirstValue + TDimensionSize - 1);
@@ -102,7 +119,7 @@ struct TShaderPermutationInt
 	/** Converts dimension's integer value to dimension's value id. */
 	static int32 ToDimensionValueId(Type E)
 	{
-		int32 PermutationId = static_cast<int32>(E) - TFirstValue;
+		const int32 PermutationId = static_cast<int32>(E) - TFirstValue;
 		checkf(PermutationId < PermutationCount && PermutationId >= 0, TEXT("Unknown shader permutation dimension value id %i."), PermutationId);
 		return PermutationId;
 	}
@@ -119,11 +136,21 @@ struct TShaderPermutationInt
 		checkf(PermutationId < PermutationCount && PermutationId >= 0, TEXT("Invalid shader permutation dimension value id %i."), PermutationId);
 		return static_cast<Type>(PermutationId + TFirstValue);
 	}
+
+	/** Returns all the possible values of a specialization. */
+	static void GetPossibleValues(TArray<int32>& OutValues)
+	{
+		checkf(TIsSpecialization, TEXT("Getting possible values should only be called for specializations!"));
+		for (int32 Value = MinValue; Value <= MaxValue; ++Value)
+		{
+			OutValues.Add(Value);
+		}
+	}
 };
 
 
 /** Defines at compile time a permutation dimension made of specific int32. */
-template <int32... Ts>
+template <bool TIsSpecialization, int32... Ts>
 struct TShaderPermutationSparseInt
 {
 	/** Setup the dimension's type in permutation domain as integer. */
@@ -137,6 +164,8 @@ struct TShaderPermutationSparseInt
 	 */
 	static constexpr bool IsMultiDimensional = false;
 
+	/** Specialization constant data. */
+	static constexpr bool IsSpecialization = TIsSpecialization;
 
 	/** Converts dimension's integer value to dimension's value id, bu in this case fail because the dimension value was wrong. */
 	static int32 ToDimensionValueId(Type E)
@@ -151,22 +180,29 @@ struct TShaderPermutationSparseInt
 		checkf(false, TEXT("Invalid shader permutation dimension id %i."), PermutationId);
 		return Type(0);
 	}
+
+	/** Returns all the possible values of a specialization. */
+	static void GetPossibleValues(TArray<int32>& OutValues)
+	{
+	}
 };
 
-template <int32 TUniqueValue, int32... Ts>
-struct TShaderPermutationSparseInt<TUniqueValue, Ts...>
+template <bool TIsSpecialization, int32 TUniqueValue, int32... Ts>
+struct TShaderPermutationSparseInt<TIsSpecialization, TUniqueValue, Ts...>
 {
 	/** Setup the dimension's type in permutation domain as integer. */
 	using Type = int32;
 
 	/** Setup the dimension's number of permutation. */
-	static constexpr int32 PermutationCount = TShaderPermutationSparseInt<Ts...>::PermutationCount + 1;
+	static constexpr int32 PermutationCount = TShaderPermutationSparseInt<TIsSpecialization, Ts...>::PermutationCount + 1;
 	
 	/** Setup the dimension as non multi-dimensional, so that the ModifyCompilationEnvironement's
 	 * define can conventily be set up in SHADER_PERMUTATION_SPARSE_INT.
 	 */
 	static constexpr bool IsMultiDimensional = false;
 
+	/** Specialization constant data. */
+	static constexpr bool IsSpecialization = TIsSpecialization;
 
 	/** Converts dimension's integer value to dimension's value id. */
 	static int32 ToDimensionValueId(Type E)
@@ -175,7 +211,7 @@ struct TShaderPermutationSparseInt<TUniqueValue, Ts...>
 		{
 			return PermutationCount - 1;
 		}
-		return TShaderPermutationSparseInt<Ts...>::ToDimensionValueId(E);
+		return TShaderPermutationSparseInt<TIsSpecialization, Ts...>::ToDimensionValueId(E);
 	}
 
 	/** Pass down a int32 to FShaderCompilerEnvironment::SetDefine(). */
@@ -191,7 +227,15 @@ struct TShaderPermutationSparseInt<TUniqueValue, Ts...>
 		{
 			return TUniqueValue;
 		}
-		return TShaderPermutationSparseInt<Ts...>::FromDimensionValueId(PermutationId);
+		return TShaderPermutationSparseInt<TIsSpecialization, Ts...>::FromDimensionValueId(PermutationId);
+	}
+
+	/** Returns all the possible values of a specialization. */
+	static void GetPossibleValues(TArray<int32>& OutValues)
+	{
+		checkf(TIsSpecialization, TEXT("Getting possible values should only be called for specializations!"));
+		OutValues.Add(TUniqueValue);
+		return TShaderPermutationSparseInt<TIsSpecialization, Ts...>::GetPossibleValues(OutValues);
 	}
 };
 
@@ -238,6 +282,8 @@ struct TShaderPermutationDomain
 	/** Total number of permutation within the domain is one if no dimension at all. */
 	static constexpr int32 PermutationCount = 1;
 
+	/** Total number of specialization within the domain is one if no dimension at all. */
+	static constexpr int32 SpecializationCount = 0;
 
 	/** Constructors. */
 	TShaderPermutationDomain<Ts...>() {}
@@ -272,6 +318,11 @@ struct TShaderPermutationDomain
 	/** Modify compilation environment. */
 	void ModifyCompilationEnvironment(FShaderCompilerEnvironment& OutEnvironment) const {}
 
+	/** Returns if the permutation is a specialization (requries a specialization constant array to be set on the shader). */
+	bool IsSpecialization() const { return false; }
+
+	/** Generate a specialization constant array. */
+	void ToSpecializationConstantArray(FSpecializationContainerType& OutArray) const {}
 
 	/** Converts domain permutation vector to domain's value id. */
 	static int32 ToDimensionValueId(const Type& PermutationVector)
@@ -284,6 +335,15 @@ struct TShaderPermutationDomain
 		return ToDimensionValueId(*this);
 	}
 
+	static int32 ToDimensionUnspecializedId(const Type& PermutationVector)
+	{
+		return 0;
+	}
+
+	int32 ToDimensionUnspecializedId() const
+	{
+		return ToDimensionUnspecializedId(*this);
+	}
 
 	/** Returns the permutation domain from the unique ID. */
 	static Type FromDimensionValueId(const int32 PermutationId)
@@ -310,16 +370,38 @@ public:
 	template<typename TPermutationVector, typename TDimension>
 	static void ModifyCompilationEnvironment(const TPermutationVector& PermutationVector, FShaderCompilerEnvironment& OutEnvironment)
 	{
-
 		if constexpr (bLogPermutations)
 		{
 			UE_LOG(LogShaders, Verbose, TEXT("		%s = %d"), TDimension::DefineName, TDimension::ToDefineValue(PermutationVector.DimensionValue));
 		}
 
 		OutEnvironment.SetDefine(TDimension::DefineName, TDimension::ToDefineValue(PermutationVector.DimensionValue));
+
+		if constexpr (TDimension::IsSpecialization)
+		{
+			const int32 SpecializationConstantIndex = OutEnvironment.SpecializationConstantValues.Num();
+			OutEnvironment.SetDefine(TDimension::SpecializationIndexName, SpecializationConstantIndex);
+
+			TArray<int32>& Values = OutEnvironment.SpecializationConstantValues.AddDefaulted_GetRef();
+			Values.Reserve(TDimension::PermutationCount);
+			TDimension::GetPossibleValues(Values);
+		}
+
 		return PermutationVector.Tail.ModifyCompilationEnvironment(OutEnvironment);
 	}
 
+
+	template<typename TPermutationVector, typename TDimension>
+	static void ToSpecializationConstantArray(const TPermutationVector& PermutationVector, FSpecializationContainerType& OutArray)
+	{
+		if constexpr (TDimension::IsSpecialization)
+		{
+			OutArray.Add(TDimension::ToDefineValue(PermutationVector.DimensionValue));
+		}
+
+		return PermutationVector.Tail.ToSpecializationConstantArray(OutArray);
+	}
+
 	template<typename TPermutationVector, typename TDimensionToGet>
 	static const typename TDimensionToGet::Type& GetDimension(const TPermutationVector& PermutationVector)
 	{
@@ -331,7 +413,6 @@ public:
 	{
 		return PermutationVector.Tail.template Set<TDimensionToSet>(Value);
 	}
-
 };
 
 template<>
@@ -346,6 +427,13 @@ public:
 		return PermutationVector.Tail.ModifyCompilationEnvironment(OutEnvironment);
 	}
 
+	template<typename TPermutationVector, typename TDimension>
+	static void ToSpecializationConstantArray(const TPermutationVector& PermutationVector, FSpecializationContainerType& OutArray)
+	{
+		PermutationVector.DimensionValue.ToSpecializationConstantArray(OutArray);
+		return PermutationVector.Tail.ToSpecializationConstantArray(OutArray);
+	}
+
 	template<typename TPermutationVector, typename TDimensionToGet>
 	static const typename TDimensionToGet::Type& GetDimension(const TPermutationVector& PermutationVector)
 	{
@@ -357,7 +445,6 @@ public:
 	{
 		PermutationVector.DimensionValue = Value;
 	}
-
 };
 
 
@@ -378,6 +465,19 @@ struct TShaderPermutationDomain<TDimension, Ts...>
 	/** Total number of permutation within the domain. */
 	static constexpr int32 PermutationCount = Super::PermutationCount * TDimension::PermutationCount;
 
+	/** Total number of specializations within the domain. */
+	static constexpr int32 GetDimensionSpecCount()
+	{
+		if constexpr (TDimension::IsMultiDimensional)
+		{
+			return TDimension::SpecializationCount;
+		}
+		else
+		{
+			return TDimension::IsSpecialization ? 1 : 0;
+		}
+	}
+	static constexpr int32 SpecializationCount = Super::SpecializationCount + GetDimensionSpecCount();
 
 	/** Constructors. */
 	TShaderPermutationDomain<TDimension, Ts...>()
@@ -422,6 +522,20 @@ struct TShaderPermutationDomain<TDimension, Ts...>
 		TShaderPermutationDomainSpetialization<TDimension::IsMultiDimensional>::template ModifyCompilationEnvironment<Type, TDimension>(*this, OutEnvironment);
 	}
 
+	/** Returns true if any of the specialization values are in use (different than 0).
+	* All shaders created are created with the default value of 0 for every specialization constant,
+	* as soon as a value differs from 0 the permutation is considered to be a specialization.
+	*/
+	bool IsSpecialization() const
+	{
+		return ToDimensionValueId(*this) != ToDimensionUnspecializedId(*this);
+	}
+
+	/** Fill an array with all the specialization constant values. */
+	void ToSpecializationConstantArray(FSpecializationContainerType& OutArray) const
+	{
+		TShaderPermutationDomainSpetialization<TDimension::IsMultiDimensional>::template ToSpecializationConstantArray<Type, TDimension>(*this, OutArray);
+	}
 
 	/** Converts domain permutation vector to domain's value id. */
 	static int32 ToDimensionValueId(const Type& PermutationVector)
@@ -434,6 +548,32 @@ struct TShaderPermutationDomain<TDimension, Ts...>
 		return TDimension::ToDimensionValueId(DimensionValue) + TDimension::PermutationCount * Tail.ToDimensionValueId();
 	}
 
+	/** Same as ToDimensionValueId but only collects values for Dimensions that are not specializations. 
+	* Used to map a specialization to an actual shader permutation.
+	*/
+	static int32 ToDimensionUnspecializedId(const Type& PermutationVector)
+	{
+		return PermutationVector.ToDimensionUnspecializedId();
+	}
+
+	int32 ToDimensionUnspecializedId() const
+	{
+		if constexpr (TDimension::IsMultiDimensional)
+		{
+			return TDimension::ToDimensionUnspecializedId(DimensionValue) + TDimension::PermutationCount * Tail.ToDimensionUnspecializedId();
+		}
+		else
+		{
+			if constexpr (TDimension::IsSpecialization)
+			{
+				return TDimension::PermutationCount * Tail.ToDimensionUnspecializedId();
+			}
+			else
+			{
+				return TDimension::ToDimensionValueId(DimensionValue) + TDimension::PermutationCount * Tail.ToDimensionUnspecializedId();
+			}
+		}
+	}
 
 	/** Returns the permutation domain from the unique ID. */
 	static Type FromDimensionValueId(const int32 PermutationId)
@@ -468,44 +608,69 @@ using FShaderPermutationNone = TShaderPermutationDomain<>;
 
 
 // Internal implementation of non multi-dimensional shader permutation dimension.
-#define DECLARE_SHADER_PERMUTATION_IMPL(InDefineName,PermutationMetaType,...) \
-	public PermutationMetaType<__VA_ARGS__> { \
+#define DECLARE_SHADER_PERMUTATION_IMPL(IsSpecialization,InDefineName,PermutationMetaType,...) \
+	public PermutationMetaType<IsSpecialization,__VA_ARGS__> { \
 	public: \
-		static constexpr const TCHAR* DefineName = TEXT(InDefineName); \
+		static constexpr const TCHAR* DefineName = IsSpecialization ? TEXT(InDefineName "_DEFAULT") : TEXT(InDefineName); \
+		static constexpr const TCHAR* SpecializationIndexName = IsSpecialization ? TEXT(InDefineName "_INDEX") : nullptr; \
 	}
 
 
+
+// Make things clearer in macro below
+#define SHADER_SPECIALIZATION	true
+#define SHADER_PERMUTATION		false
+
+
 /** Implements a boolean shader permutation dimensions. Meant to be used like so:
  *
  * class FMyShaderDim : SHADER_PERMUTATION_BOOL("MY_SHADER_DEFINE_NAME");
  */
 #define SHADER_PERMUTATION_BOOL(InDefineName) \
-	public FShaderPermutationBool { \
+	public FShaderPermutationBool<SHADER_PERMUTATION> { \
 	public: \
 		static constexpr const TCHAR* DefineName = TEXT(InDefineName); \
+		static constexpr const TCHAR* SpecializationIndexName = nullptr; \
 	}
-	
+
+#define SHADER_SPECIALIZATION_BOOL(InDefineName) \
+	public FShaderPermutationBool<SHADER_SPECIALIZATION> { \
+	public: \
+		static constexpr const TCHAR* DefineName = TEXT(InDefineName "_DEFAULT"); \
+		static constexpr const TCHAR* SpecializationIndexName = TEXT(InDefineName "_INDEX"); \
+	}
+
+
 /** Implements an integer shader permutation dimensions with N permutation values from [[0; N[[. Meant to be used like so:
  *
  * class FMyShaderDim : SHADER_PERMUTATION_INT("MY_SHADER_DEFINE_NAME", N);
  */
 #define SHADER_PERMUTATION_INT(InDefineName, Count) \
-	DECLARE_SHADER_PERMUTATION_IMPL(InDefineName, TShaderPermutationInt, int32, Count)
-	
+	DECLARE_SHADER_PERMUTATION_IMPL(SHADER_PERMUTATION, InDefineName, TShaderPermutationInt, int32, Count)
+
+#define SHADER_SPECIALIZATION_INT(InDefineName, Count) \
+	DECLARE_SHADER_PERMUTATION_IMPL(SHADER_SPECIALIZATION, InDefineName, TShaderPermutationInt, int32, Count)
+
 /** Implements an integer shader permutation dimensions with N permutation values from [[X; X+N[[. Meant to be used like so:
  *
  * class FMyShaderDim : SHADER_PERMUTATION_RANGE_INT("MY_SHADER_DEFINE_NAME", X, N);
  */
 #define SHADER_PERMUTATION_RANGE_INT(InDefineName, Start, Count) \
-	DECLARE_SHADER_PERMUTATION_IMPL(InDefineName, TShaderPermutationInt, int32, Count, Start)
+	DECLARE_SHADER_PERMUTATION_IMPL(SHADER_PERMUTATION, InDefineName, TShaderPermutationInt, int32, Count, Start)
+
+#define SHADER_SPECIALIZATION_RANGE_INT(InDefineName, Start, Count) \
+	DECLARE_SHADER_PERMUTATION_IMPL(SHADER_SPECIALIZATION, InDefineName, TShaderPermutationInt, int32, Count, Start)
 
 /** Implements an integer shader permutation dimensions with non contiguous permutation values. Meant to be used like so:
  *
  * class FMyShaderDim : SHADER_PERMUTATION_SPARSE_INT("MY_SHADER_DEFINE_NAME", 1, 2, 4, 8);
  */
 #define SHADER_PERMUTATION_SPARSE_INT(InDefineName,...) \
-	DECLARE_SHADER_PERMUTATION_IMPL(InDefineName, TShaderPermutationSparseInt, __VA_ARGS__)
-	
+	DECLARE_SHADER_PERMUTATION_IMPL(SHADER_PERMUTATION, InDefineName, TShaderPermutationSparseInt, __VA_ARGS__)
+
+#define SHADER_SPECIALIZATION_SPARSE_INT(InDefineName,...) \
+	DECLARE_SHADER_PERMUTATION_IMPL(SHADER_SPECIALIZATION, InDefineName, TShaderPermutationSparseInt, __VA_ARGS__)
+
 /** Implements an shader permutation dimensions with an enum class assumed to have contiguous integer values. Meant to be used like so:
  *
  * enum class EMyEnum
@@ -519,4 +684,11 @@ using FShaderPermutationNone = TShaderPermutationDomain<>;
  * class FMyShaderDim : SHADER_PERMUTATION_ENUM_CLASS("MY_SHADER_DEFINE_NAME", EMyEnum);
  */
 #define SHADER_PERMUTATION_ENUM_CLASS(InDefineName, EnumName) \
-	DECLARE_SHADER_PERMUTATION_IMPL(InDefineName, TShaderPermutationInt, EnumName, static_cast<int32>(EnumName::MAX))
+	DECLARE_SHADER_PERMUTATION_IMPL(SHADER_PERMUTATION, InDefineName, TShaderPermutationInt, EnumName, static_cast<int32>(EnumName::MAX))
+
+#define SHADER_SPECIALIZATION_ENUM_CLASS(InDefineName, EnumName) \
+	DECLARE_SHADER_PERMUTATION_IMPL(SHADER_SPECIALIZATION, InDefineName, TShaderPermutationInt, EnumName, static_cast<int32>(EnumName::MAX))
+
+
+
+
diff --git a/Engine/Source/Runtime/RenderCore/Public/ShaderPermutationUtils.h b/Engine/Source/Runtime/RenderCore/Public/ShaderPermutationUtils.h
--- a/Engine/Source/Runtime/RenderCore/Public/ShaderPermutationUtils.h
+++ b/Engine/Source/Runtime/RenderCore/Public/ShaderPermutationUtils.h
@@ -98,22 +98,33 @@ namespace UE::ShaderPermutationUtils
 		{
 			FormatPermutationDomain(E, OutString, bFullNames, Prefix);
 		}
-		else if constexpr (std::is_base_of_v<FShaderPermutationBool, TDimension>)
+		else if constexpr (std::is_base_of_v<FShaderPermutationBool<SHADER_SPECIALIZATION>, TDimension>)
 		{
 			if (bFullNames)
 			{
-				OutString.Appendf(TEXT("%s (%s)"), TDimension::DefineName, !FShaderPermutationBool::ToDefineValue(E) ? TEXT("false") : TEXT("true"));
+				OutString.Appendf(TEXT("%s (%s) *spec*"), TDimension::DefineName, !FShaderPermutationBool<SHADER_SPECIALIZATION>::ToDefineValue(E) ? TEXT("false") : TEXT("true"));
 			}
 			else
 			{
-				OutString.Appendf(TEXT("%d"), static_cast<int32>(FShaderPermutationBool::ToDefineValue(E)));
+				OutString.Appendf(TEXT("%d"), static_cast<int32>(FShaderPermutationBool<SHADER_SPECIALIZATION>::ToDefineValue(E)));
+			}
+		}
+		else if constexpr (std::is_base_of_v<FShaderPermutationBool<SHADER_PERMUTATION>, TDimension>)
+		{
+			if (bFullNames)
+			{
+				OutString.Appendf(TEXT("%s (%s)"), TDimension::DefineName, !FShaderPermutationBool<SHADER_PERMUTATION>::ToDefineValue(E) ? TEXT("false") : TEXT("true"));
+			}
+			else
+			{
+				OutString.Appendf(TEXT("%d"), static_cast<int32>(FShaderPermutationBool<SHADER_PERMUTATION>::ToDefineValue(E)));
 			}
 		}
 		else
 		{
 			if (bFullNames)
 			{
-				OutString.Appendf(TEXT("%s (%d)"), TDimension::DefineName, static_cast<int32>(TDimension::ToDefineValue(E)));
+				OutString.Appendf(TEXT("%s (%d)%s"), TDimension::DefineName, static_cast<int32>(TDimension::ToDefineValue(E)), TDimension::IsSpecialization ? TEXT(" *spec*") : TEXT(""));
 			}
 			else
 			{
@@ -134,4 +145,17 @@ namespace UE::ShaderPermutationUtils
 		}
 		return false;
 	}
+
+	template <typename DomainType>
+	int32 ConvertToUnspecializedId(const EShaderPlatform Platform, const DomainType& PermutationVector)
+	{
+		if (FDataDrivenShaderPlatformInfo::GetSupportsSpecializationConstants(Platform))
+		{
+			return DomainType::ToDimensionUnspecializedId(PermutationVector);
+		}
+		else
+		{
+			return DomainType::ToDimensionValueId(PermutationVector);
+		}
+	}
 }
