Disable velocity rendering depth encoding for all mobile platforms

#Synced-CL 50642587

diff --git a/Engine/Source/Runtime/RenderCore/Private/RenderUtils.cpp b/Engine/Source/Runtime/RenderCore/Private/RenderUtils.cpp
--- a/Engine/Source/Runtime/RenderCore/Private/RenderUtils.cpp
+++ b/Engine/Source/Runtime/RenderCore/Private/RenderUtils.cpp
@@ -1802,7 +1802,7 @@ bool ForceSimpleSkyDiffuse(const FStaticShaderPlatform Platform)
 
 bool VelocityEncodeDepth(const FStaticShaderPlatform Platform)
 {
-	return (GVelocityEncodeDepthPlatformMask[(int)Platform]);
+	return !IsMobilePlatform(Platform) && (GVelocityEncodeDepthPlatformMask[(int)Platform]);
 }
 
 bool VelocityEncodeHasPixelAnimation(const FStaticShaderPlatform Platform)
