Move bProduced to its own memory location to avoid data race.

#Synced-CL 50642529

diff --git a/Engine/Source/Runtime/RenderCore/Private/RenderGraphResources.cpp b/Engine/Source/Runtime/RenderCore/Private/RenderGraphResources.cpp
--- a/Engine/Source/Runtime/RenderCore/Private/RenderGraphResources.cpp
+++ b/Engine/Source/Runtime/RenderCore/Private/RenderGraphResources.cpp
@@ -45,7 +45,6 @@ FRDGViewableResource::FRDGViewableResource(const TCHAR* InName, const ERDGViewab
 	, Type(InType)
 	, bExternal(0)
 	, bExtracted(0)
-	, bProduced(0)
 	, bTransient(0)
 	, bForceNonTransient(0)
 	, bSkipLastTransition(0)
diff --git a/Engine/Source/Runtime/RenderCore/Public/RenderGraphResources.h b/Engine/Source/Runtime/RenderCore/Public/RenderGraphResources.h
--- a/Engine/Source/Runtime/RenderCore/Public/RenderGraphResources.h
+++ b/Engine/Source/Runtime/RenderCore/Public/RenderGraphResources.h
@@ -375,9 +375,6 @@ protected:
 	/** Whether this is an extracted resource. */
 	uint8 bExtracted : 1;
 
-	/** Whether any sub-resource has been used for write by a pass. */
-	uint8 bProduced : 1;
-
 	/** Whether this resource is allocated through the transient resource allocator. */
 	uint8 bTransient : 1;
 
@@ -405,6 +402,9 @@ protected:
 	uint8 bClobbered : 1;
 #endif
 
+	/** Whether any sub-resource has been used for write by a pass. Separate from the bitfield because it is assigned by the async setup queue. */
+	bool bProduced = false;
+
 	/** Whether this resource is allowed to be both transient and extracted. */
 	ETransientExtractionHint TransientExtractionHint;
 
