Fixed unsynchronized increment of stat counter.

#Synced-CL 50642529

diff --git a/Engine/Source/Runtime/RenderCore/Private/RenderGraphPass.cpp b/Engine/Source/Runtime/RenderCore/Private/RenderGraphPass.cpp
--- a/Engine/Source/Runtime/RenderCore/Private/RenderGraphPass.cpp
+++ b/Engine/Source/Runtime/RenderCore/Private/RenderGraphPass.cpp
@@ -274,7 +274,7 @@ void FRDGBarrierBatchBegin::Submit(FRHIComputeCommandList& RHICmdList, ERHIPipel
 	}
 
 #if RDG_STATS
-	GRDGStatTransitionBatchCount++;
+	GRDGStatTransitionBatchCount.fetch_add(1, std::memory_order_relaxed);
 #endif
 }
 
diff --git a/Engine/Source/Runtime/RenderCore/Private/RenderGraphPrivate.cpp b/Engine/Source/Runtime/RenderCore/Private/RenderGraphPrivate.cpp
--- a/Engine/Source/Runtime/RenderCore/Private/RenderGraphPrivate.cpp
+++ b/Engine/Source/Runtime/RenderCore/Private/RenderGraphPrivate.cpp
@@ -446,7 +446,7 @@ int32 GRDGStatTransientTextureCount = 0;
 int32 GRDGStatTransientBufferCount = 0;
 int32 GRDGStatTransitionCount = 0;
 int32 GRDGStatAliasingCount = 0;
-int32 GRDGStatTransitionBatchCount = 0;
+std::atomic_int32_t GRDGStatTransitionBatchCount{ 0 };
 int32 GRDGStatMemoryWatermark = 0;
 #endif
 
diff --git a/Engine/Source/Runtime/RenderCore/Private/RenderGraphPrivate.h b/Engine/Source/Runtime/RenderCore/Private/RenderGraphPrivate.h
--- a/Engine/Source/Runtime/RenderCore/Private/RenderGraphPrivate.h
+++ b/Engine/Source/Runtime/RenderCore/Private/RenderGraphPrivate.h
@@ -138,7 +138,7 @@ extern int32 GRDGStatTransientTextureCount;
 extern int32 GRDGStatTransientBufferCount;
 extern int32 GRDGStatTransitionCount;
 extern int32 GRDGStatAliasingCount;
-extern int32 GRDGStatTransitionBatchCount;
+extern std::atomic_int32_t GRDGStatTransitionBatchCount;
 extern int32 GRDGStatMemoryWatermark;
 #endif
 
