[Shaders] Improve debug output of `-DumpShaderLibraryHashes` commandline argument.

Outputs the ShaderGroupIoHashes with their associated shaders and shadermaps:
```
IoChunk 617b8b2f55e462712ef15909 (2 shaders)
9B6BD915665B790F0B8BA016341221F11BBB8DA4 (Referenced By Shadermaps: 5DB05BDEC32B998EAC2A992867051027EA9A71C3)
977355DF9590EE15C744B41BF05C474E6CB175EB (Referenced By Shadermaps: 5DB05BDEC32B998EAC2A992867051027EA9A71C3)
```
It also outputs the shadermaps with their associated shaders in the bottom half of the debug output file:
```
ShaderMap 5DB05BDEC32B998EAC2A992867051027EA9A71C3 (15 shaders)
9B6BD915665B790F0B8BA016341221F11BBB8DA4
977355DF9590EE15C744B41BF05C474E6CB175EB
...
```

#rnx

#Synced-CL 50638327

diff --git a/Engine/Source/Runtime/RenderCore/Private/ShaderLibrary/ShaderCodeArchive.cpp b/Engine/Source/Runtime/RenderCore/Private/ShaderLibrary/ShaderCodeArchive.cpp
--- a/Engine/Source/Runtime/RenderCore/Private/ShaderLibrary/ShaderCodeArchive.cpp
+++ b/Engine/Source/Runtime/RenderCore/Private/ShaderLibrary/ShaderCodeArchive.cpp
@@ -2294,35 +2294,74 @@ static void DumpShaderCodeArchiveHashes(const FIoStoreShaderCodeArchiveHeader& I
 {
 	FString DebugOutput;
 
+	// Build shader-to-shadermap links for faster access
+	TMap<FSHAHash, TSet<FSHAHash>> ShaderToShaderMaps;
+	for (int32 ShaderMapIndex = 0; ShaderMapIndex < InLibraryHeader.ShaderMapEntries.Num(); ++ShaderMapIndex)
+	{
+		const FIoStoreShaderMapEntry& ShaderMapEntry = InLibraryHeader.ShaderMapEntries[ShaderMapIndex];
+		for (uint32 ShaderIndex = 0; ShaderIndex < ShaderMapEntry.NumShaders; ++ShaderIndex)
+		{
+			const uint32 MappedShaderIndex = InLibraryHeader.ShaderIndices[ShaderMapEntry.ShaderIndicesOffset + ShaderIndex];
+			TSet<FSHAHash>& ShaderMapHahses = ShaderToShaderMaps.FindOrAdd(InLibraryHeader.ShaderHashes[MappedShaderIndex]);
+			ShaderMapHahses.Add(InLibraryHeader.ShaderMapHashes[ShaderMapIndex]);
+		}
+	}
+
+	// Dump hashes of all IoChunks with their shaders and associated shadermaps
 	DebugOutput.Appendf(TEXT("--------------- ShaderGroupIoHashes: ---------------\n"));
 	for (int32 GroupIndex = 0; GroupIndex < InLibraryHeader.ShaderGroupIoHashes.Num(); ++GroupIndex)
 	{
 		const FIoChunkId& ShaderGroupChunkId = InLibraryHeader.ShaderGroupIoHashes[GroupIndex];
 		const FIoStoreShaderGroupEntry& GroupEntry = InLibraryHeader.ShaderGroupEntries[GroupIndex];
-		DebugOutput.Appendf(TEXT("%s (%d shaders)\n"), *LexToString(ShaderGroupChunkId), GroupEntry.NumShaders);
+		DebugOutput.Appendf(TEXT("IoChunk %s (%d shaders)\n"), *LexToString(ShaderGroupChunkId), GroupEntry.NumShaders);
 
 		for (uint32 ShaderIndex = 0; ShaderIndex < GroupEntry.NumShaders; ++ShaderIndex)
 		{
-			const int32 MappedShaderIndex = InLibraryHeader.ShaderIndices[GroupEntry.ShaderIndicesOffset + ShaderIndex];
-			DebugOutput.Appendf(TEXT("  %s\n"), *InLibraryHeader.ShaderHashes[MappedShaderIndex].ToString());
+			const uint32 MappedShaderIndex = InLibraryHeader.ShaderIndices[GroupEntry.ShaderIndicesOffset + ShaderIndex];
+			const FSHAHash& ShaderHash = InLibraryHeader.ShaderHashes[MappedShaderIndex];
+			DebugOutput.Appendf(TEXT("  %s"), *LexToString(ShaderHash));
+
+			// Append all shadermaps that reference this shader
+			if (const TSet<FSHAHash>* FoundShaderMaps = ShaderToShaderMaps.Find(ShaderHash))
+			{
+				DebugOutput.Append(TEXT(" (Referenced By Shadermaps: "));
+				bool bIsFirstEntry = true;
+				for (const FSHAHash& ShaderMapHash : *FoundShaderMaps)
+				{
+					if (!bIsFirstEntry)
+					{
+						DebugOutput.Append(TEXT(", "));
+					}
+					DebugOutput.Append(LexToString(ShaderMapHash));
+					bIsFirstEntry = false;
+				}
+				DebugOutput.Append(TEXT(")"));
+			}
+			else
+			{
+				DebugOutput.Append(TEXT(" (No Shadermap References)"));
+			}
+
+			DebugOutput.Append(TEXT("\n"));
 		}
 		DebugOutput.Append(TEXT("\n"));
 	}
 	DebugOutput.Append(TEXT("\n\n\n"));
 
-	DebugOutput.Appendf(TEXT("--------------- ShaderHashes: ---------------\n"));
-	for (const FSHAHash& ShaderHash : InLibraryHeader.ShaderHashes)
-	{
-		DebugOutput.Appendf(TEXT("%s\n"), *ShaderHash.ToString());
-	}
-	DebugOutput.Append(TEXT("\n\n\n"));
-
+	// Dump hashes of all shadermaps with their shaders
 	DebugOutput.Appendf(TEXT("--------------- ShaderMapHashes: ---------------\n"));
-	for (const FSHAHash& ShaderMapHash : InLibraryHeader.ShaderMapHashes)
+	for (int32 ShaderMapIndex = 0; ShaderMapIndex < InLibraryHeader.ShaderMapEntries.Num(); ++ShaderMapIndex)
 	{
-		DebugOutput.Appendf(TEXT("%s\n"), *ShaderMapHash.ToString());
+		const FIoStoreShaderMapEntry& ShaderMapEntry = InLibraryHeader.ShaderMapEntries[ShaderMapIndex];
+		DebugOutput.Appendf(TEXT("ShaderMap %s (%d shaders)\n"), *LexToString(InLibraryHeader.ShaderMapHashes[ShaderMapIndex]), ShaderMapEntry.NumShaders);
+		for (uint32 ShaderIndex = 0; ShaderIndex < ShaderMapEntry.NumShaders; ++ShaderIndex)
+		{
+			const uint32 MappedShaderIndex = InLibraryHeader.ShaderIndices[ShaderMapEntry.ShaderIndicesOffset + ShaderIndex];
+			const FSHAHash& ShaderHash = InLibraryHeader.ShaderHashes[MappedShaderIndex];
+			DebugOutput.Appendf(TEXT("  %s\n"), *LexToString(ShaderHash));
+		}
+		DebugOutput.Append(TEXT("\n"));
 	}
-	DebugOutput.Append(TEXT("\n\n\n"));
 
 	FTCHARToUTF8 DebugOutputUtf8Converter(*DebugOutput);
 	Ar.Serialize(const_cast<UTF8CHAR*>(reinterpret_cast<const UTF8CHAR*>(DebugOutputUtf8Converter.Get())), DebugOutputUtf8Converter.Length());
