[Shader/IAX] Sort shader library components by their priority to ensure we lookup shaders in non-IAD containers first.

With shader chunking, some shaders can be grouped into different IoChunks that are included in an IAD and non-IAD container respectively.
Any system that does not handle deferred shader loading (unlike IAD assets) relies on the shader code library to load that shader from the non-IAD container.
Therefore, we assign a priority to the shader library components to ensure we always lookup shaders in the non-IAD containers first.
This priority of components is isolated to the internals of the shader code library and does not interact with the IoDispatcher.
It also uses insertion sort to maintain the same order as before besides the new arrangement based on priority.

Optionally, a higher verbosity of LogShaderLibrary will also log the active set of components in their sorted order whenever a new container is mounted:
```
LogShaderLibrary: Verbose: OpenShaderCode('Lyra_Chunk2'): Components=(0, 1, 2)
```

#rnx
#Synced-CL 50638327

diff --git a/Engine/Source/Runtime/RenderCore/Private/ShaderLibrary/ShaderCodeLibrary.cpp b/Engine/Source/Runtime/RenderCore/Private/ShaderLibrary/ShaderCodeLibrary.cpp
--- a/Engine/Source/Runtime/RenderCore/Private/ShaderLibrary/ShaderCodeLibrary.cpp
+++ b/Engine/Source/Runtime/RenderCore/Private/ShaderLibrary/ShaderCodeLibrary.cpp
@@ -85,6 +85,8 @@ namespace UE::ShaderLibrary::Private
 
 	static TSet<FString> PluginsToIgnoreOnMount;
 
+	using FShaderLibraryInstancePriorityPair = TPair<int32, TUniquePtr<FShaderLibraryInstance>>;
+
 	// [RCL] TODO 2020-11-20: Separate runtime and editor-only code (tracked as UE-103486)
 	/** Descriptor used to pass the pak file information to the library as we cannot store IPakFile ref */
 	struct FMountedPakFileInfo
@@ -157,7 +159,7 @@ namespace UE::ShaderLibrary::Private
 		FRWLock ComponentsMutex;
 
 		/** Even putting aside chunking, each named library can be potentially comprised of multiple files */
-		TArray<TUniquePtr<FShaderLibraryInstance>> Components;
+		TArray<FShaderLibraryInstancePriorityPair> Components;
 
 		FNamedShaderLibrary(const FString& InLogicalName, const EShaderPlatform InShaderPlatform, const FString& InBaseDirectory)
 			: LogicalName(InLogicalName)
@@ -179,9 +181,9 @@ namespace UE::ShaderLibrary::Private
 			FRWScopeLock ReadLock(ComponentsMutex, SLT_ReadOnly);
 			TSet<int32> PresentChunks;
 
-			for (TUniquePtr<FShaderLibraryInstance>& Component : Components)
+			for (const FShaderLibraryInstancePriorityPair& Component : Components)
 			{
-				int32 ChunkId = Component->GetChunkId();
+				int32 ChunkId = Component.Value->GetChunkId();
 				if (LIKELY(ChunkId != INDEX_NONE))
 				{
 					PresentChunks.Add(ChunkId);
@@ -191,41 +193,54 @@ namespace UE::ShaderLibrary::Private
 			return MoveTemp(PresentChunks);
 		}
 
+		/** Helper struct to carry the chunk library name, ID, and priority assigned to the library instance internals. */
+		struct FShaderLibraryChunkInfo
+		{
+			FString ChunkLibraryName;
+			int32 ChunkId = INDEX_NONE;
+			int32 Priority = 0;
+		};
+
 		// Transforms the chunk Id of the pak file for a named shader library component,
 		// e.g. pakchunk100 and pakchunk100iad must not point to the same ID 100, so the latter one will be transformed to a higher number.
-		TTuple<FString, int32> GetLibraryNameAndChunkId(const FMountedPakFileInfo& MountInfo) const
+		FShaderLibraryChunkInfo GetLibraryChunkInfo(const FMountedPakFileInfo& MountInfo) const
 		{
 			constexpr int32 kChunkIdSuffixOffsetRange = 100000;
 			constexpr int32 kChunkIdSuffixBitmask = 0xFF;
 
-			FString ChunkLibraryName = GetShaderLibraryNameForChunk(LogicalName, MountInfo.ChunkId);
-			int32 ChunkId = MountInfo.ChunkId;
+			FShaderLibraryChunkInfo ChunkInfo;
+			ChunkInfo.ChunkLibraryName = GetShaderLibraryNameForChunk(LogicalName, MountInfo.ChunkId);
+			ChunkInfo.ChunkId = MountInfo.ChunkId;
 
 			// If the pakchunk ID contains a suffix such as "iad", the shader library must have been created as a proxy target file and will also contain that suffix.
 			FString ChunkIdSuffix;
 			uint64 ChunkIdSuffixHash = 0;
-			if (FShaderLibraryNameInfo::ParsePakChunkId(FPaths::GetBaseFilename(MountInfo.PakFilename), ChunkId, ChunkIdSuffix) && !ChunkIdSuffix.IsEmpty())
+			if (FShaderLibraryNameInfo::ParsePakChunkId(FPaths::GetBaseFilename(MountInfo.PakFilename), ChunkInfo.ChunkId, ChunkIdSuffix) && !ChunkIdSuffix.IsEmpty())
 			{
-				ChunkLibraryName.Append(ChunkIdSuffix);
+				// Mangle chunk ID used for internal redundancy checks, so that IAD and non-IAD chunks with the same ID don't collide (e.g. pakchunk100 and pakchunk100iad)
+				ChunkInfo.ChunkLibraryName.Append(ChunkIdSuffix);
 				ChunkIdSuffixHash = CityHash64(reinterpret_cast<const char*>(*ChunkIdSuffix), ChunkIdSuffix.Len() * sizeof(TCHAR));
-				ChunkId += (ChunkIdSuffixHash & kChunkIdSuffixBitmask) * kChunkIdSuffixOffsetRange;
+				ChunkInfo.ChunkId += (ChunkIdSuffixHash & kChunkIdSuffixBitmask) * kChunkIdSuffixOffsetRange;
+
+				// Lower priority of such secondary chunks
+				ChunkInfo.Priority = -1;
 			}
 
-			return { MoveTemp(ChunkLibraryName), ChunkId };
+			return MoveTemp(ChunkInfo);
 		}
 
 		void OnPakFileMounted(const FMountedPakFileInfo& MountInfo, const FString& Directory)
 		{
-			auto [ChunkLibraryName, ChunkId] = GetLibraryNameAndChunkId(MountInfo);
-			if (!GetPresentChunks().Contains(ChunkId))
+			FShaderLibraryChunkInfo ChunkInfo = GetLibraryChunkInfo(MountInfo);
+			if (!GetPresentChunks().Contains(ChunkInfo.ChunkId))
 			{
 				// Parts of shader library might be in the UFS (e.g. .metallibs), hence we need to look for them in the
 				// appropriate directory.
-				OpenShaderCode(Directory, ChunkLibraryName, ChunkId);
+				OpenShaderCode(Directory, ChunkInfo.ChunkLibraryName, ChunkInfo.ChunkId, ChunkInfo.Priority);
 			}
 		}
 
-		bool OpenShaderCode(const FString& ShaderCodeDir, FString const& Library, int32 ChunkId);
+		bool OpenShaderCode(const FString& ShaderCodeDir, FString const& Library, int32 ChunkId, int32 Priority = 0);
 		FShaderLibraryInstance* FindShaderLibraryForShaderMap(const FSHAHash& Hash, int32& OutShaderMapIndex);
 		FShaderLibraryInstance* FindShaderLibraryForShader(const FSHAHash& Hash, int32& OutShaderIndex);
 		uint32 GetShaderCount();
@@ -938,11 +953,11 @@ public:
 		for (const TTuple<FString, TUniquePtr<UE::ShaderLibrary::Private::FNamedShaderLibrary>>& NamedLibraryPair : NamedLibrariesStack)
 		{
 			const FNamedShaderLibrary* Library = NamedLibraryPair.Value.Get();
-			for (const TUniquePtr<FShaderLibraryInstance>& LibraryInstance : Library->Components)
+			for (const FShaderLibraryInstancePriorityPair& LibraryInstance : Library->Components)
 			{
-				for (int32 LibraryResourceIndex = 0; LibraryResourceIndex < LibraryInstance->GetNumResources(); ++LibraryResourceIndex)
+				for (int32 LibraryResourceIndex = 0; LibraryResourceIndex < LibraryInstance.Value->GetNumResources(); ++LibraryResourceIndex)
 				{
-					if (LibraryInstance->GetResource(LibraryResourceIndex).GetReference() == InResource)
+					if (LibraryInstance.Value->GetResource(LibraryResourceIndex).GetReference() == InResource)
 					{
 						OutLibrariesWithRefToResource.Add(NamedLibraryPair.Value->LogicalName);
 					}
@@ -2798,17 +2813,20 @@ void FShaderCodeLibrary::OpenPluginShaderLibrary(IPlugin& Plugin, bool bMonolith
 // FNamedShaderLibrary methods
 
 // At runtime, open shader code collection for specified shader platform. Returns true if new code was opened
-bool UE::ShaderLibrary::Private::FNamedShaderLibrary::OpenShaderCode(const FString& ShaderCodeDir, FString const& Library, int32 ChunkId)
+bool UE::ShaderLibrary::Private::FNamedShaderLibrary::OpenShaderCode(const FString& ShaderCodeDir, FString const& Library, int32 ChunkId, int32 Priority)
 {
 	LLM_SCOPE(ELLMTag::Shaders);
 	// check if any of the components has this content
 	{
 		FRWScopeLock ReadLock(ComponentsMutex, SLT_Write);
-		for (TUniquePtr<FShaderLibraryInstance>& Component : Components)
+		for (const FShaderLibraryInstancePriorityPair& Component : Components)
 		{
-			if (Component->HasContentFrom(ShaderCodeDir, Library))
+			if (Component.Value->HasContentFrom(ShaderCodeDir, Library))
 			{
 				// in this context, "false" means "no new library was opened"
+				ensureMsgf(Component.Key == Priority,
+					TEXT("Priority mismatch of shader library '%s' between new request (%d) and existing instance (%d)"),
+					*Library, Priority, Component.Key);
 				return false;
 			}
 		}
@@ -2825,9 +2843,9 @@ bool UE::ShaderLibrary::Private::FNamedShaderLibrary::OpenShaderCode(const FStri
 	FRWScopeLock WriteLock(ComponentsMutex, SLT_Write);
 
 	// re-check that no one has added the same library while we were creating it. If so, delete ours
-	for (TUniquePtr<FShaderLibraryInstance>& Component : Components)
+	for (const FShaderLibraryInstancePriorityPair& Component : Components)
 	{
-		if (Component->HasContentFrom(ShaderCodeDir, Library))
+		if (Component.Value->HasContentFrom(ShaderCodeDir, Library))
 		{
 			// in this context, "false" means "no new library was opened"
 			delete LibraryInstance;
@@ -2844,7 +2862,37 @@ bool UE::ShaderLibrary::Private::FNamedShaderLibrary::OpenShaderCode(const FStri
 		UE_LOG(LogShaderLibrary, Display, TEXT("Cooked Context: Using Shared Shader Library %s"), *Library);
 	}
 
-	Components.Emplace(LibraryInstance);
+	// Insertion sort of new component by their priority.
+	// This is mandatory because some shaders can be grouped across different IoChunks, some of which can be packaged into an IAD and a non-IAD container respectively.
+	// When such shaders are requested from systems that don't handle deferred shader downloads (unlike IAD asset),
+	// we need to ensure that shader is loaded from the non-IAD container, which is why they are assigned a higher priority.
+	const int32 InsertIndex = Components.IndexOfByPredicate([Priority](const FShaderLibraryInstancePriorityPair& Entry) -> bool { return Priority > Entry.Key; });
+	if (InsertIndex == INDEX_NONE)
+	{
+		Components.Emplace(FShaderLibraryInstancePriorityPair(Priority, LibraryInstance));
+	}
+	else
+	{
+		Components.Insert(FShaderLibraryInstancePriorityPair(Priority, LibraryInstance), InsertIndex);
+	}
+
+	// Log new set of component IDs if verbose logging is enabled
+	auto PrintComponentChunkIds = [](const TConstArrayView<FShaderLibraryInstancePriorityPair>& InComponents) -> FString
+		{
+			TStringBuilder<1024> OutString;
+			for (const FShaderLibraryInstancePriorityPair& Component : InComponents)
+			{
+				if (OutString.Len() > 0)
+				{
+					OutString.Append(TEXT(", "));
+				}
+				OutString.Appendf(TEXT("%d"), Component.Value->GetChunkId());
+			}
+			return FString(OutString);
+		};
+
+	UE_LOG(LogShaderLibrary, Verbose, TEXT("OpenShaderCode('%s'): Components=(%s)"), *Library, *PrintComponentChunkIds(Components));
+
 	return true;
 }
 
@@ -2853,13 +2901,13 @@ FShaderLibraryInstance* UE::ShaderLibrary::Private::FNamedShaderLibrary::FindSha
 	FRWScopeLock ReadLock(ComponentsMutex, SLT_ReadOnly);
 
 	// Search in library opened order
-	for (const TUniquePtr<FShaderLibraryInstance>& Instance : Components)
+	for (const FShaderLibraryInstancePriorityPair& Instance : Components)
 	{
-		const int32 ShaderMapIndex = Instance->Library->FindShaderMapIndex(Hash);
+		const int32 ShaderMapIndex = Instance.Value->Library->FindShaderMapIndex(Hash);
 		if (ShaderMapIndex != INDEX_NONE)
 		{
 			OutShaderMapIndex = ShaderMapIndex;
-			return Instance.Get();
+			return Instance.Value.Get();
 		}
 	}
 	return nullptr;
@@ -2870,13 +2918,13 @@ FShaderLibraryInstance* UE::ShaderLibrary::Private::FNamedShaderLibrary::FindSha
 	FRWScopeLock ReadLock(ComponentsMutex, SLT_ReadOnly);
 
 	// Search in library opened order
-	for (const TUniquePtr<FShaderLibraryInstance>& Instance : Components)
+	for (const FShaderLibraryInstancePriorityPair& Instance : Components)
 	{
-		const int32 ShaderIndex = Instance->Library->FindShaderIndex(Hash);
+		const int32 ShaderIndex = Instance.Value->Library->FindShaderIndex(Hash);
 		if (ShaderIndex != INDEX_NONE)
 		{
 			OutShaderIndex = ShaderIndex;
-			return Instance.Get();
+			return Instance.Value.Get();
 		}
 	}
 	return nullptr;
@@ -2887,9 +2935,9 @@ uint32 UE::ShaderLibrary::Private::FNamedShaderLibrary::GetShaderCount(void)
 	FRWScopeLock ReadLock(ComponentsMutex, SLT_ReadOnly);
 
 	int ShaderCount = 0;
-	for (const TUniquePtr<FShaderLibraryInstance>& Instance : Components)
+	for (const FShaderLibraryInstancePriorityPair& Instance : Components)
 	{
-		ShaderCount += Instance->Library->GetNumShaders();
+		ShaderCount += Instance.Value->Library->GetNumShaders();
 	}
 	return ShaderCount;
 }
@@ -2900,10 +2948,10 @@ void UE::ShaderLibrary::Private::FNamedShaderLibrary::DumpLibraryContents(const
 	FRWScopeLock ReadLock(ComponentsMutex, SLT_ReadOnly);
 
 	int32 ComponentIdx = 0;
-	for (const TUniquePtr<FShaderLibraryInstance>& Instance : Components)
+	for (const FShaderLibraryInstancePriorityPair& Instance : Components)
 	{
 		UE_LOG(LogShaderLibrary, Display, TEXT("%sComponent %d: Native=%s Shaders: %d Name: %s"),
-			*Prefix, ComponentIdx, Instance->Library->IsNativeLibrary() ? TEXT("yes") : TEXT("no"), Instance->GetNumShaders(), *Instance->Library->GetName() );
+			*Prefix, ComponentIdx, Instance.Value->Library->IsNativeLibrary() ? TEXT("yes") : TEXT("no"), Instance.Value->GetNumShaders(), *Instance.Value->Library->GetName() );
 		++ComponentIdx;
 	}
 }
@@ -2918,8 +2966,9 @@ FAutoConsoleCommandWithArgsAndOutputDevice GListShaderLibrariesCmd(
 			auto Iter = [&](const FString& LogicalName, UE::ShaderLibrary::Private::FNamedShaderLibrary& NamedShaderLibrary)
 			{
 				FRWScopeLock ComponentReadLock(NamedShaderLibrary.ComponentsMutex, SLT_ReadOnly);
-				for (const TUniquePtr<FShaderLibraryInstance>& ShaderLibrary : NamedShaderLibrary.Components)
+				for (const FShaderLibraryInstancePriorityPair& ShaderLibraryPair : NamedShaderLibrary.Components)
 				{
+					FShaderLibraryInstance* ShaderLibrary = ShaderLibraryPair.Value.Get();
 					const FString& Name = ShaderLibrary->Library->GetName();
 					const FString OwnerName = ShaderLibrary->Library->GetOwnerName().ToString();
 					uint32 Id = ShaderLibrary->Library->GetId();
