Make GRDGAllowRHIAccess atomic to silence TSAN. Also moved pass bit to other async setup queue union where it's accessed instead.

#Synced-CL 50783182

diff --git a/Engine/Source/Runtime/RenderCore/Private/RenderGraphBuilder.cpp b/Engine/Source/Runtime/RenderCore/Private/RenderGraphBuilder.cpp
--- a/Engine/Source/Runtime/RenderCore/Private/RenderGraphBuilder.cpp
+++ b/Engine/Source/Runtime/RenderCore/Private/RenderGraphBuilder.cpp
@@ -626,10 +626,10 @@ FRDGBuilder::FRDGBuilder(FRHICommandListImmediate& InRHICmdList, FRDGEventName I
 
 #if RDG_ENABLE_DEBUG
 	UserValidation.SetParallelExecuteEnabled(ParallelExecute.TaskMode != ERDGPassTaskMode::Inline);
-	if (GRDGAllowRHIAccessAsync != bParallelExecuteAllowedAsync)
+	if (GRDGAllowRHIAccessAsync.load(std::memory_order_relaxed) != bParallelExecuteAllowedAsync)
 	{
 		WaitForAsyncExecuteTask();
-		GRDGAllowRHIAccessAsync = bParallelExecuteAllowedAsync;
+		GRDGAllowRHIAccessAsync.store(bParallelExecuteAllowedAsync, std::memory_order_relaxed);
 	}
 #endif
 }
@@ -869,9 +869,9 @@ FRHIUniformBuffer* FRDGBuilder::ConvertToExternalUniformBuffer(FRDGUniformBuffer
 		if (!UniformBuffer->GetRHIUnchecked())
 		{
 			// It's safe to reset the access to false because validation won't allow this call during execution.
-			IF_RDG_ENABLE_DEBUG(GRDGAllowRHIAccess = true);
+			IF_RDG_ENABLE_DEBUG(GRDGAllowRHIAccess.store(true, std::memory_order_relaxed));
 			UniformBuffer->InitRHI();
-			IF_RDG_ENABLE_DEBUG(GRDGAllowRHIAccess = false);
+			IF_RDG_ENABLE_DEBUG(GRDGAllowRHIAccess.store(false, std::memory_order_relaxed));
 		}
 	}
 	return UniformBuffer->GetRHIUnchecked();
@@ -1788,7 +1788,7 @@ void FRDGBuilder::Execute()
 	UE::Tasks::FTask CreateViewsTask;
 
 	IF_RDG_ENABLE_DEBUG(UserValidation.ValidateExecuteBegin());
-	IF_RDG_ENABLE_DEBUG(GRDGAllowRHIAccess = true);
+	IF_RDG_ENABLE_DEBUG(GRDGAllowRHIAccess.store(true, std::memory_order_relaxed));
 
 	FCollectResourceContext CollectResourceContext;
 
@@ -2025,7 +2025,7 @@ void FRDGBuilder::Execute()
 	EndFlushResourcesRHI();
 	WaitForParallelSetupTasks(ERDGSetupTaskWaitPoint::Execute);
 
-	IF_RDG_ENABLE_DEBUG(GRDGAllowRHIAccess = ParallelExecute.IsEnabled());
+	IF_RDG_ENABLE_DEBUG(GRDGAllowRHIAccess.store(ParallelExecute.IsEnabled(), std::memory_order_relaxed));
 	IF_RDG_ENABLE_TRACE(Trace.OutputGraphBegin());
 
 	bCompiling = false;
@@ -2210,7 +2210,7 @@ void FRDGBuilder::Execute()
 	IF_RDG_ENABLE_TRACE(Trace.OutputGraphEnd(*this));
 
 	IF_RDG_ENABLE_DEBUG(UserValidation.ValidateExecuteEnd());
-	IF_RDG_ENABLE_DEBUG(GRDGAllowRHIAccess = false);
+	IF_RDG_ENABLE_DEBUG(GRDGAllowRHIAccess.store(false, std::memory_order_relaxed));
 
 #if RDG_STATS
 	GRDGStatBufferCount += Buffers.Num();
diff --git a/Engine/Source/Runtime/RenderCore/Private/RenderGraphPrivate.cpp b/Engine/Source/Runtime/RenderCore/Private/RenderGraphPrivate.cpp
--- a/Engine/Source/Runtime/RenderCore/Private/RenderGraphPrivate.cpp
+++ b/Engine/Source/Runtime/RenderCore/Private/RenderGraphPrivate.cpp
@@ -203,8 +203,8 @@ uint8 GetClobberStencil()
 	return 123;
 }
 
-bool GRDGAllowRHIAccess = false;
-bool GRDGAllowRHIAccessAsync = false;
+std::atomic<bool> GRDGAllowRHIAccess{false};
+std::atomic<bool> GRDGAllowRHIAccessAsync{false};
 
 #endif
 
diff --git a/Engine/Source/Runtime/RenderCore/Private/RenderGraphPrivate.h b/Engine/Source/Runtime/RenderCore/Private/RenderGraphPrivate.h
--- a/Engine/Source/Runtime/RenderCore/Private/RenderGraphPrivate.h
+++ b/Engine/Source/Runtime/RenderCore/Private/RenderGraphPrivate.h
@@ -33,22 +33,22 @@ extern int32 GRDGDebugDisableTransientResources;
 extern int32 GRDGTransitionLog;
 extern int32 GRDGImmediateMode;
 extern int32 GRDGOverlapUAVs;
-extern bool  GRDGAllowRHIAccess;
-extern bool  GRDGAllowRHIAccessAsync;
+extern std::atomic<bool> GRDGAllowRHIAccess;
+extern std::atomic<bool> GRDGAllowRHIAccessAsync;
 
 class FRDGAllowRHIAccessScope
 {
 public:
 	FRDGAllowRHIAccessScope()
 	{
-		check(!GRDGAllowRHIAccess);
-		GRDGAllowRHIAccess = true;
+		check(!GRDGAllowRHIAccess.load(std::memory_order_relaxed));
+		GRDGAllowRHIAccess.store(true, std::memory_order_relaxed);
 	}
 
 	~FRDGAllowRHIAccessScope()
 	{
-		check(GRDGAllowRHIAccess);
-		GRDGAllowRHIAccess = false;
+		check(GRDGAllowRHIAccess.load(std::memory_order_relaxed));
+		GRDGAllowRHIAccess.store(false, std::memory_order_relaxed);
 	}
 };
 
diff --git a/Engine/Source/Runtime/RenderCore/Private/RenderGraphValidation.cpp b/Engine/Source/Runtime/RenderCore/Private/RenderGraphValidation.cpp
--- a/Engine/Source/Runtime/RenderCore/Private/RenderGraphValidation.cpp
+++ b/Engine/Source/Runtime/RenderCore/Private/RenderGraphValidation.cpp
@@ -40,7 +40,7 @@ void FRDGResource::ValidateRHIAccess() const
 		return;
 	}
 
-	checkf(bAllowRHIAccess || GRDGAllowRHIAccess || GRDGAllowRHIAccessAsync,
+	checkf(bAllowRHIAccess || GRDGAllowRHIAccess.load(std::memory_order_relaxed) || GRDGAllowRHIAccessAsync.load(std::memory_order_relaxed),
 		TEXT("Accessing the RHI resource of %s at this time is not allowed. If you hit this check in pass, ")
 		TEXT("that is due to this resource not being referenced in the parameters of your pass."),
 		Name);
diff --git a/Engine/Source/Runtime/RenderCore/Public/RenderGraphPass.h b/Engine/Source/Runtime/RenderCore/Public/RenderGraphPass.h
--- a/Engine/Source/Runtime/RenderCore/Public/RenderGraphPass.h
+++ b/Engine/Source/Runtime/RenderCore/Public/RenderGraphPass.h
@@ -409,9 +409,6 @@ protected:
 			uint16 bGraphicsFork : 1;
 			uint16 bGraphicsJoin : 1;
 
-			/** Whether the pass only writes to resources in its render pass. */
-			uint16 bRenderPassOnlyWrites : 1;
-
 			/** Whether this pass is a sentinel (prologue / epilogue) pass. */
 			uint16 bSentinel : 1;
 
@@ -440,6 +437,9 @@ protected:
 
 			/** Whether this pass is used for external access transitions. */
 			uint8 bExternalAccessPass : 1;
+
+			/** Whether the pass only writes to resources in its render pass. */
+			uint8 bRenderPassOnlyWrites : 1;
 		};
 		uint8 PacketBits_AsyncSetupQueue = 0;
 	};
