Fixed render target pool hash collision by comparing descriptors.

#Synced-CL 50790641

diff --git a/Engine/Source/Runtime/RenderCore/Private/RenderTargetPool.cpp b/Engine/Source/Runtime/RenderCore/Private/RenderTargetPool.cpp
--- a/Engine/Source/Runtime/RenderCore/Private/RenderTargetPool.cpp
+++ b/Engine/Source/Runtime/RenderCore/Private/RenderTargetPool.cpp
@@ -191,8 +191,13 @@ FPooledRenderTarget* FRenderTargetPool::TryFindRenderTarget(const FRHITextureCre
 		{
 			FPooledRenderTarget* Element = PooledRenderTargets[Index];
 
-			checkf(Element, TEXT("Hash was not cleared from the list."));
-			checkf(Translate(Element->GetDesc()) == Desc, TEXT("Invalid hash or collision when attempting to allocate %s"), Element->GetDesc().DebugName);
+			checkf(Element && Element->GetRHI(), TEXT("Hash was not cleared from the list."));
+
+			// Compare the descriptors to avoid hash collisions.
+			if (Element->GetRHI()->GetDesc() != Desc)
+			{
+				continue;
+			}
 
 			if (Element->IsFree() && Predicate(Element))
 			{
@@ -283,15 +288,12 @@ bool FRenderTargetPool::FindFreeElement(FRHICommandListBase& RHICmdList, const F
 	{
 		Current = (FPooledRenderTarget*)Out.GetReference();
 
-		if (Translate(Out->GetDesc()) == Desc)
+		if (Out->GetRHI() && Out->GetRHI()->GetDesc() == Desc)
 		{
 			// we can reuse the same, but the debug name might have changed
 			Current->Desc.DebugName = Name;
 #if RHI_USE_RESOURCE_DEBUG_NAME
-			if (Current->GetRHI())
-			{
-				RHICmdList.BindDebugLabelName(Current->GetRHI(), Name);
-			}
+			RHICmdList.BindDebugLabelName(Current->GetRHI(), Name);
 #endif
 			check(!Out->IsFree());
 			return true;
