- Add permutation vector constructor to TOptionalShaderMapRef

#Synced-CL 50731656

diff --git a/Engine/Source/Runtime/RenderCore/Public/Shader.h b/Engine/Source/Runtime/RenderCore/Public/Shader.h
--- a/Engine/Source/Runtime/RenderCore/Public/Shader.h
+++ b/Engine/Source/Runtime/RenderCore/Public/Shader.h
@@ -2679,6 +2679,13 @@ public:
 	TOptionalShaderMapRef(const typename ShaderType::ShaderMapType* ShaderIndex):
 		TShaderRef<ShaderType>(TShaderRef<ShaderType>::Cast(ShaderIndex->GetShader(&ShaderType::GetStaticType()))) // gcc3 needs the template quantifier so it knows the < is not a less-than
 	{}
+
+	TOptionalShaderMapRef(
+		const typename ShaderType::ShaderMapType* ShaderIndex,
+		const typename ShaderType::FPermutationDomain& PermutationVector)
+		: TShaderRef<ShaderType>(ShaderIndex->template GetShader<ShaderType>(PermutationVector.ToDimensionValueId())) // gcc3 needs the template quantifier so it knows the < is not a less-than
+	{
+	}
 };
 
 /**
