RHI: Don't generate the specialization values when spec consts aren't supported.

#rnx
#Synced-CL 50785456

diff --git a/Engine/Source/Runtime/RenderCore/Private/GlobalShader.cpp b/Engine/Source/Runtime/RenderCore/Private/GlobalShader.cpp
--- a/Engine/Source/Runtime/RenderCore/Private/GlobalShader.cpp
+++ b/Engine/Source/Runtime/RenderCore/Private/GlobalShader.cpp
@@ -617,8 +617,9 @@ bool FGlobalShaderMapSection::Serialize(FShaderSerializeContext& Ctx)
 
 TShaderRef<FShader> FGlobalShaderMapSection::GetShader(EShaderPlatform ShaderPlatform, FShaderType* ShaderType, int32 PermutationId) const
 {
-	FShader* Shader = GetContent()->GetShader(ShaderType, FDataDrivenShaderPlatformInfo::GetSupportsSpecializationConstants(ShaderPlatform) ? ShaderType->GetUnspecializedId(PermutationId) : PermutationId);
-	return Shader ? TShaderRef<FShader>(Shader, *this, ShaderType->GetSpecializationValues(PermutationId)) : TShaderRef<FShader>();
+	const bool bSupportsSpecConsts = FDataDrivenShaderPlatformInfo::GetSupportsSpecializationConstants(ShaderPlatform);
+	FShader* Shader = GetContent()->GetShader(ShaderType, bSupportsSpecConsts ? ShaderType->GetUnspecializedId(PermutationId) : PermutationId);
+	return Shader ? TShaderRef<FShader>(Shader, *this, bSupportsSpecConsts ? ShaderType->GetSpecializationValues(PermutationId) : FSpecializationContainerType()) : TShaderRef<FShader>();
 }
 
 FShaderPipelineRef FGlobalShaderMapSection::GetShaderPipeline(const FShaderPipelineType* PipelineType) const
