Deprecate FVTProducerDescription::bNotifyCompleted
There is no core usage of this and for scenes with many virtual textures there is a CPU cost to iterate them all in FVirtualTextureProducerCollection::NotifyRequestsCompleted().
Add r.VT.EnableNotifyCompleted cvar which disables notify callbacks by default.
The notify system will be removed in some future version.

#Synced-CL 50901855

diff --git a/Engine/Source/Runtime/RenderCore/Public/VirtualTexturing.h b/Engine/Source/Runtime/RenderCore/Public/VirtualTexturing.h
--- a/Engine/Source/Runtime/RenderCore/Public/VirtualTexturing.h
+++ b/Engine/Source/Runtime/RenderCore/Public/VirtualTexturing.h
@@ -146,7 +146,12 @@ struct FAllocatedVTDescription
 
 struct FVTProducerDescription
 {
+	// Disable bNotifyCompleted deprecation warnings here.
+	PRAGMA_DISABLE_DEPRECATION_WARNINGS
 	RENDERCORE_API FVTProducerDescription();
+	FVTProducerDescription(const FVTProducerDescription&) = default;
+	FVTProducerDescription& operator=(const FVTProducerDescription&) = default;
+	PRAGMA_ENABLE_DEPRECATION_WARNINGS
 
 	FName Name; /** Will be name of UTexture for streaming VTs, mostly here for debugging */
 	uint32 FullNameHash;
@@ -154,7 +159,9 @@ struct FVTProducerDescription
 	bool bPersistentHighestMip = true;
 	bool bContinuousUpdate = false;
 	bool bRequiresSinglePhysicalPool = false;
-	bool bNotifyCompleted = false; /** Producer will receive OnRequestsCompleted() callbacks every frame when enabled. */
+	
+	UE_DEPRECATED(5.8, "OnRequestsCompleted() callbacks are deprecated for future removal. They are currently disabled by default with r.VT.EnableNotifyCompleted.")
+	bool bNotifyCompleted = false;
 	
 	uint32 TileSize = 0u;
 	uint32 TileBorderSize = 0u;
@@ -352,7 +359,8 @@ public:
 	/** Dump any type specific debug info. */
 	virtual void DumpToConsole(bool verbose) {}
 
-	/** Called on every virtual texture system update once all requests are completed, if bNotifyCompleted is enabled. */
+	/** Called on every virtual texture system update once all requests are completed, if bNotifyCompleted is enabled and r.VT.EnableNotifyCompleted=1. */
+	UE_DEPRECATED(5.8, "OnRequestsCompleted() callbacks are deprecated for future removal.")
 	virtual void OnRequestsCompleted() {}
 };
 
diff --git a/Engine/Source/Runtime/Renderer/Private/VT/VirtualTextureProducer.cpp b/Engine/Source/Runtime/Renderer/Private/VT/VirtualTextureProducer.cpp
--- a/Engine/Source/Runtime/Renderer/Private/VT/VirtualTextureProducer.cpp
+++ b/Engine/Source/Runtime/Renderer/Private/VT/VirtualTextureProducer.cpp
@@ -5,6 +5,14 @@
 #include "VT/VirtualTexturePhysicalSpace.h"
 #include "VT/VirtualTextureSystem.h"
 
+static TAutoConsoleVariable<int32> CVarVTEnableNotifyCompleted(
+	TEXT("r.VT.EnableNotifyCompleted"),
+	0,
+	TEXT("Enable calls to OnRequestsCompleted().\n")
+	TEXT("Note that support for OnRequestsCompleted() callbacks is deprecated in 5.8 and will be removed in a later version."),
+	ECVF_RenderThreadSafe
+);
+
 FVirtualTextureProducer::~FVirtualTextureProducer()
 {
 }
@@ -248,14 +256,19 @@ uint32 FVirtualTextureProducerCollection::RemoveAllCallbacks(const void* Baton)
 
 void FVirtualTextureProducerCollection::NotifyRequestsCompleted()
 {
-	TRACE_CPUPROFILER_EVENT_SCOPE(FVirtualTextureSystem::NotifyRequestsComplete);
-
-	for (FProducerEntry& Entry : Producers)
+	if (CVarVTEnableNotifyCompleted.GetValueOnAnyThread() != 0)
 	{
-		if (Entry.Producer.Description.bNotifyCompleted)
+		TRACE_CPUPROFILER_EVENT_SCOPE(FVirtualTextureSystem::NotifyRequestsComplete);
+
+		for (FProducerEntry& Entry : Producers)
 		{
-			Entry.Producer.GetVirtualTexture()->OnRequestsCompleted();
+PRAGMA_DISABLE_DEPRECATION_WARNINGS
+			if (Entry.Producer.Description.bNotifyCompleted)
+			{
+				Entry.Producer.GetVirtualTexture()->OnRequestsCompleted();
+			}
 		}
+PRAGMA_ENABLE_DEPRECATION_WARNINGS
 	}
 }
 
diff --git a/Engine/Source/Runtime/Renderer/Private/VT/VirtualTextureProducer.h b/Engine/Source/Runtime/Renderer/Private/VT/VirtualTextureProducer.h
--- a/Engine/Source/Runtime/Renderer/Private/VT/VirtualTextureProducer.h
+++ b/Engine/Source/Runtime/Renderer/Private/VT/VirtualTextureProducer.h
@@ -71,7 +71,10 @@ public:
 	void CallPendingCallbacks();
 	bool HasPendingCallbacks() const;
 
-	/** Notify producers marked as "continous notify" that all requests have been completed. */
+	/** 
+	 * Notify producers that all requests have been completed. 
+	 * Only notifies when r.VT.EnableNotifyCompleted=1 and for producers with bNotifyCompleted=true.
+	 */
 	void NotifyRequestsCompleted();
 
 	/**
