[Shaders] Don't use debug path for intermediate shader files that changes when running on build machines.

Intermediate files are not meant to be uploaded as artifacts, only debug information.
So use the fixed path ProjectSavedDir + "Shaders" again for intermediate shader output during packaging.

- Remove CVar `r.ShaderCodeLibrary.UploadArtifacts`; Always upload these artifacts when they are produced on a build machine.
This only produces 1-3 text files per target platform, reporting the hashes of all shadermaps that have moved between pakchunks during shader chunk re-assignment.
- Output shader debug information during packaging under Saved/ShaderDebugInfo/ to clearly separate it from intermediate data in Saved/Shaders/.

#rnx

#Synced-CL 50901530

#Synced-CL 50901530

diff --git a/Engine/Source/Developer/IoStoreUtilities/Private/IoStoreShaderLibraryProcessor.cpp b/Engine/Source/Developer/IoStoreUtilities/Private/IoStoreShaderLibraryProcessor.cpp
--- a/Engine/Source/Developer/IoStoreUtilities/Private/IoStoreShaderLibraryProcessor.cpp
+++ b/Engine/Source/Developer/IoStoreUtilities/Private/IoStoreShaderLibraryProcessor.cpp
@@ -899,7 +899,7 @@ void FIoStoreShaderLibraryProcessor::ProcessShaderLibraries(const FIoStoreArgume
 		{
 			LogChunkOverrideStats(FName(FormatAndPlatformName));
 
-			const FString DebugLibFolder = FSerializedShaderArchive::MakeDebugDirectory() / FormatAndPlatformName;
+			const FString DebugLibFolder = FSerializedShaderArchive::MakeDebugDirectory();
 			IFileManager::Get().MakeDirectory(*DebugLibFolder, true);
 			const FString DebugFilename = DebugLibFolder / FString::Printf(TEXT("ShaderChunkOverride-%s.txt"), *FormatAndPlatformName);
 			if (TUniquePtr<FArchive> DebugFileWriter = TUniquePtr<FArchive>(IFileManager::Get().CreateFileWriter(*DebugFilename)))
diff --git a/Engine/Source/Runtime/RenderCore/Private/ShaderLibrary/EditorShaderCodeArchive.cpp b/Engine/Source/Runtime/RenderCore/Private/ShaderLibrary/EditorShaderCodeArchive.cpp
--- a/Engine/Source/Runtime/RenderCore/Private/ShaderLibrary/EditorShaderCodeArchive.cpp
+++ b/Engine/Source/Runtime/RenderCore/Private/ShaderLibrary/EditorShaderCodeArchive.cpp
@@ -794,7 +794,7 @@ void FEditorShaderCodeArchive::AddExistingShaderCodeLibrary(FString const& Outpu
 	static const FRegexPattern ShaderArchivePattern(ShaderArchivePatternStr);
 
 	const FString FormatNameStr = FormatAndPlatformName.ToString();
-	const FString ShaderIntermediateLocation = GetShaderCodeLibraryDebugInfoPath() / FormatNameStr;
+	const FString ShaderIntermediateLocation = GetShaderCodeIntermediatePath() / FormatNameStr;
 
 	TArray<FString> ShaderFiles;
 	ShaderFindFiles(ShaderFiles, *ShaderIntermediateLocation, *ShaderExtension);
@@ -943,14 +943,14 @@ bool FEditorShaderCodeArchive::SaveToDisk(const FString& OutputDir, const FStrin
 	// Shader library
 	if (bSuccess && SerializedShaders.GetNumShaderMaps() > 0)
 	{
-		// Write to a intermediate file
-		const FString DebugInfoPath = GetShaderCodeLibraryDebugInfoPath();
+		// Write to intermediate files
+		const FString IntermediatePathBase = GetShaderCodeIntermediatePath() / FormatAndPlatformName.ToString();
 
 		// Save the ShaderMaps, Shaders, and the Shaders' byte code
 		constexpr bool bSaveOnlyAssetInfo = false; // A currently unused setting that we might want to bring back.
 		if (!bSaveOnlyAssetInfo)
 		{
-			const FString IntermediateFormatPath = GetShaderCodeFilename(DebugInfoPath / FormatAndPlatformName.ToString(), LibraryName, FormatAndPlatformName, Target);
+			const FString IntermediateFormatPath = GetShaderCodeFilename(IntermediatePathBase, LibraryName, FormatAndPlatformName, Target);
 			if (TUniquePtr<FArchive> FileWriter = TUniquePtr<FArchive>(IFileManager::Get().CreateFileWriter(*IntermediateFormatPath, FILEWRITE_NoFail)))
 			{
 				SerializedShaders.Finalize();
@@ -971,7 +971,7 @@ bool FEditorShaderCodeArchive::SaveToDisk(const FString& OutputDir, const FStrin
 				CopyFilesToOutputDir(GetCodeArchiveFilename, IntermediateFormatPath, OutputFilenames);
 			}
 
-			const FString TypeInfoIntermediatePath = GetShaderTypeInfoFilename(DebugInfoPath / FormatAndPlatformName.ToString(), LibraryName, FormatAndPlatformName, Target);
+			const FString TypeInfoIntermediatePath = GetShaderTypeInfoFilename(IntermediatePathBase, LibraryName, FormatAndPlatformName, Target);
 			if (TUniquePtr<FArchive> ShaderTypeFileWriter = TUniquePtr<FArchive>(IFileManager::Get().CreateFileWriter(*TypeInfoIntermediatePath, FILEWRITE_NoFail)))
 			{
 				*ShaderTypeFileWriter << SerializedShaders.ShaderTypes;
@@ -984,7 +984,7 @@ bool FEditorShaderCodeArchive::SaveToDisk(const FString& OutputDir, const FStrin
 
 		// Save asset info
 		{
-			const FString AssetInfoIntermediatePath = GetShaderAssetInfoFilename(DebugInfoPath / FormatAndPlatformName.ToString(), LibraryName, FormatAndPlatformName, Target);
+			const FString AssetInfoIntermediatePath = GetShaderAssetInfoFilename(IntermediatePathBase, LibraryName, FormatAndPlatformName, Target);
 			if (TUniquePtr<FArchive> AssetInfoWriter = TUniquePtr<FArchive>(IFileManager::Get().CreateFileWriter(*AssetInfoIntermediatePath, FILEWRITE_NoFail)))
 			{
 				SerializedShaders.SaveAssetInfo(*AssetInfoWriter);
@@ -1014,8 +1014,8 @@ bool FEditorShaderCodeArchive::PackageNativeShaderLibrary(const FString& ShaderC
 		return false;
 	}
 
-	FString IntermediateFormatPath = GetShaderDebugFolder(GetShaderCodeLibraryDebugInfoPath() / FormatAndPlatformName.ToString(), LibraryName, FormatAndPlatformName);
-	FString TempPath = IntermediateFormatPath / TEXT("NativeLibrary");
+	const FString IntermediateFormatPath = GetShaderDebugFolder(GetShaderCodeIntermediatePath() / FormatAndPlatformName.ToString(), LibraryName, FormatAndPlatformName);
+	const FString TempPath = IntermediateFormatPath / TEXT("NativeLibrary");
 
 	IFileManager::Get().MakeDirectory(*TempPath, true);
 	IFileManager::Get().MakeDirectory(*ShaderCodeDir, true);
@@ -1183,7 +1183,7 @@ void FEditorShaderCodeArchive::DumpStatsAndDebugInfo()
 			UE_LOG(LogShaderLibrary, Display, TEXT("No shader usage info is provided"));
 		}
 
-		const FString DebugLibFolder = FSerializedShaderArchive::MakeDebugDirectory(LibraryName, FormatAndPlatformName);
+		const FString DebugLibFolder = FSerializedShaderArchive::MakeDebugDirectory(LibraryName, FormatAndPlatformName, *UE::ShaderLibrary::Private::GetShaderCodeIntermediatePath());
 		{
 			const FString DumpFilename = DebugLibFolder / TEXT("Dump.txt");
 			if (TUniquePtr<FArchive> DumpAr = TUniquePtr<FArchive>(IFileManager::Get().CreateFileWriter(*DumpFilename)))
diff --git a/Engine/Source/Runtime/RenderCore/Private/ShaderLibrary/EditorShaderStableInfo.cpp b/Engine/Source/Runtime/RenderCore/Private/ShaderLibrary/EditorShaderStableInfo.cpp
--- a/Engine/Source/Runtime/RenderCore/Private/ShaderLibrary/EditorShaderStableInfo.cpp
+++ b/Engine/Source/Runtime/RenderCore/Private/ShaderLibrary/EditorShaderStableInfo.cpp
@@ -114,7 +114,7 @@ void FEditorShaderStableInfo::AddExistingShaderCodeLibrary(FString const& Output
 
 	bool bLibraryExistsInSavedShadersDir = false;
 	{
-		const FString ShaderIntermediateLocation = GetShaderCodeLibraryDebugInfoPath() / FormatName.ToString();
+		const FString ShaderIntermediateLocation = GetShaderCodeIntermediatePath() / FormatName.ToString();
 		TArray<FString> ShaderFiles;
 		ShaderFindFiles(ShaderFiles, *ShaderIntermediateLocation, *StableExtension);
 		FString ExpectedFileNameText = LibraryName + TEXT("-") + FormatName.ToString() + TEXT(".");
@@ -156,7 +156,7 @@ bool FEditorShaderStableInfo::SaveToDisk(FString const& OutputDir, FString& OutS
 	if (bSuccess)
 	{
 		// Write to a intermediate file
-		FString IntermediateFormatPath = GetStableInfoArchiveFilename(GetShaderCodeLibraryDebugInfoPath() / FormatName.ToString(), LibraryName, FormatName);
+		FString IntermediateFormatPath = GetStableInfoArchiveFilename(GetShaderCodeIntermediatePath() / FormatName.ToString(), LibraryName, FormatName);
 
 		// Write directly to the file
 		{
diff --git a/Engine/Source/Runtime/RenderCore/Private/ShaderLibrary/ShaderCodeArchive.cpp b/Engine/Source/Runtime/RenderCore/Private/ShaderLibrary/ShaderCodeArchive.cpp
--- a/Engine/Source/Runtime/RenderCore/Private/ShaderLibrary/ShaderCodeArchive.cpp
+++ b/Engine/Source/Runtime/RenderCore/Private/ShaderLibrary/ShaderCodeArchive.cpp
@@ -17,6 +17,7 @@
 #include "Serialization/MemoryReader.h"
 #include "Serialization/MemoryWriter.h"
 #include "ShaderCodeLibrary.h"
+#include "ShaderCompilerCore.h"
 #include "ShaderLibrary/ShaderCodeLibraryUtilities.h"
 #include "ShaderCore.h"
 #include "Stats/Stats.h"
@@ -727,10 +728,10 @@ static void MakeDebugDirectoryTreeGuarded(const TCHAR* InPath)
 	}
 }
 
-FString FSerializedShaderArchive::MakeDebugDirectory(const FString LibraryName, const FName& FormatAndPlatformName)
+FString FSerializedShaderArchive::MakeDebugDirectory(const FString LibraryName, const FName& FormatAndPlatformName, const TCHAR* BasePath)
 {
 	using namespace UE::ShaderLibrary::Private;
-	FString DebugLibFolder = GetShaderCodeLibraryDebugInfoPath();
+	FString DebugLibFolder = BasePath != nullptr ? BasePath : GetShaderDebugInfoPath();
 	if (!LibraryName.IsEmpty() && FormatAndPlatformName != NAME_None)
 	{
 		DebugLibFolder = GetShaderDebugFolder(DebugLibFolder / FormatAndPlatformName.ToString(), LibraryName, FormatAndPlatformName);
@@ -2415,7 +2416,7 @@ FIoStoreShaderCodeArchive* FIoStoreShaderCodeArchive::Create(EShaderPlatform InP
 					!Library->Header.ShaderHashes.IsEmpty() ||
 					!Library->Header.ShaderMapHashes.IsEmpty())
 				{
-					const FString DumpPath = FPaths::Combine(FPaths::ProjectSavedDir(), TEXT("Shaders"));
+					const FString DumpPath = GetShaderDebugInfoPath();
 					if (!FPaths::DirectoryExists(DumpPath))
 					{
 						MakeDebugDirectoryTreeGuarded(*DumpPath);
diff --git a/Engine/Source/Runtime/RenderCore/Private/ShaderLibrary/ShaderCodeLibrary.cpp b/Engine/Source/Runtime/RenderCore/Private/ShaderLibrary/ShaderCodeLibrary.cpp
--- a/Engine/Source/Runtime/RenderCore/Private/ShaderLibrary/ShaderCodeLibrary.cpp
+++ b/Engine/Source/Runtime/RenderCore/Private/ShaderLibrary/ShaderCodeLibrary.cpp
@@ -1388,11 +1388,11 @@ public:
 
 	void CleanDirectories(TArray<FName> const& ShaderFormats)
 	{
+		const FString IntermediatePathBase = GetShaderCodeIntermediatePath();
 		for (FName const& Format : ShaderFormats)
 		{
-			const FString DebugInfoPath = GetShaderCodeLibraryDebugInfoPath();
-			FString ShaderIntermediateLocation = DebugInfoPath / Format.ToString();
-			FString ShaderPlatformIntermediateLocation = DebugInfoPath / GetFormatAndPlatformName(Format);
+			FString ShaderIntermediateLocation = IntermediatePathBase / Format.ToString();
+			FString ShaderPlatformIntermediateLocation = IntermediatePathBase / GetFormatAndPlatformName(Format);
 			IFileManager::Get().DeleteDirectory(*ShaderIntermediateLocation, false, true);
 			IFileManager::Get().DeleteDirectory(*ShaderPlatformIntermediateLocation, false, true);
 		}
diff --git a/Engine/Source/Runtime/RenderCore/Private/ShaderLibrary/ShaderCodeLibraryUtilities.cpp b/Engine/Source/Runtime/RenderCore/Private/ShaderLibrary/ShaderCodeLibraryUtilities.cpp
--- a/Engine/Source/Runtime/RenderCore/Private/ShaderLibrary/ShaderCodeLibraryUtilities.cpp
+++ b/Engine/Source/Runtime/RenderCore/Private/ShaderLibrary/ShaderCodeLibraryUtilities.cpp
@@ -72,22 +72,13 @@ namespace UE::ShaderLibrary::Private
 		ECVF_Default
 	);
 
-	bool GUploadShaderDebugArtifacts = false;
-	static FAutoConsoleVariableRef CVarUploadShaderDebugArtifactsRef(
-		TEXT("r.ShaderCodeLibrary.UploadArtifacts"),
-		GUploadShaderDebugArtifacts,
-		TEXT("If true, uploads artifacts that were dumped to the Saved/Shaders/ folder when running on a build machine"),
-		ECVF_Default
-	);
-
 #if WITH_EDITOR
 	ICookArtifactReader* CookArtifactReader = nullptr;
 #endif // WITH_EDITOR
 
-	FString GetShaderCodeLibraryDebugInfoPath()
+	FString GetShaderCodeIntermediatePath()
 	{
-		const FString BaseDebugInfoPath = GIsBuildMachine && GUploadShaderDebugArtifacts ? GetBuildMachineArtifactBasePath() : FPaths::ProjectSavedDir();
-		return BaseDebugInfoPath / TEXT("Shaders");
+		return FPaths::Combine(FPaths::ProjectSavedDir(), TEXT("Shaders"));
 	}
 
 	FString GetCodeArchiveFilename(const FString& BaseDir, const FString& LibraryName, FName Platform, ESaveToDiskTarget Target)
diff --git a/Engine/Source/Runtime/RenderCore/Private/ShaderLibrary/ShaderCodeLibraryUtilities.h b/Engine/Source/Runtime/RenderCore/Private/ShaderLibrary/ShaderCodeLibraryUtilities.h
--- a/Engine/Source/Runtime/RenderCore/Private/ShaderLibrary/ShaderCodeLibraryUtilities.h
+++ b/Engine/Source/Runtime/RenderCore/Private/ShaderLibrary/ShaderCodeLibraryUtilities.h
@@ -48,7 +48,7 @@ namespace UE::ShaderLibrary::Private
 #endif // WITH_EDITOR
 
 	// Shader library utility functions:
-	FString GetShaderCodeLibraryDebugInfoPath();
+	FString GetShaderCodeIntermediatePath();
 	FString GetCodeArchiveFilename(const FString& BaseDir, const FString& LibraryName, FName Platform, ShaderCodeArchive::ESaveToDiskTarget Target);
 	FString GetStableInfoArchiveFilename(const FString& BaseDir, const FString& LibraryName, FName Platform);
 	FString GetShaderCodeFilename(const FString& BaseDir, const FString& LibraryName, FName Platform, ShaderCodeArchive::ESaveToDiskTarget Target);
diff --git a/Engine/Source/Runtime/RenderCore/Public/ShaderCodeArchive.h b/Engine/Source/Runtime/RenderCore/Public/ShaderCodeArchive.h
--- a/Engine/Source/Runtime/RenderCore/Public/ShaderCodeArchive.h
+++ b/Engine/Source/Runtime/RenderCore/Public/ShaderCodeArchive.h
@@ -424,7 +424,7 @@ public:
 	RENDERCORE_API void DecompressShader(int32 Index, const TArray<FSharedBuffer>& ShaderCode, TArray<uint8>& OutDecompressedShader) const;
 	RENDERCORE_API static void DecompressShaderEntryAndAppend(const FShaderCodeEntry& Entry, const uint8* CompressedShaderCode, TArray<uint8>& OutDecompressedShader);
 
-	RENDERCORE_API static FString MakeDebugDirectory(const FString LibraryName = TEXT(""), const FName& FormatAndPlatformName = NAME_None);
+	RENDERCORE_API static FString MakeDebugDirectory(const FString LibraryName = TEXT(""), const FName& FormatAndPlatformName = NAME_None, const TCHAR* BasePath = nullptr);
 	
 	RENDERCORE_API static bool SerializeHeaderVersion(FArchive& Ar);
 	RENDERCORE_API void Serialize(FArchive& Ar);
