[Shaders/IAX] Resolve shadermaps to their IoChunks from all shader library instances, not just the first one we find.

Due to shader grouping, some shaders might be packaged into IoChunk A and B which are put into containers X and Y respectively.
When resolving shadermaps to their IoChunks for IAD assets, we have to return both A and B from X and Y, not just one of them from the first library instance we find.
This fixes a crash in GetResultOrDie() with `I/O Error (Not Installed)`, when a shader is loaded from an IoChunk that hasn't been resolved yet (from library A)
while the same shader is also available in another IoChunk that has been resolved (from library B).

#rnx

#Synced-CL 50937994

diff --git a/Engine/Source/Runtime/RenderCore/Private/ShaderLibrary/ShaderCodeLibrary.cpp b/Engine/Source/Runtime/RenderCore/Private/ShaderLibrary/ShaderCodeLibrary.cpp
--- a/Engine/Source/Runtime/RenderCore/Private/ShaderLibrary/ShaderCodeLibrary.cpp
+++ b/Engine/Source/Runtime/RenderCore/Private/ShaderLibrary/ShaderCodeLibrary.cpp
@@ -241,6 +241,7 @@ namespace UE::ShaderLibrary::Private
 		}
 
 		bool OpenShaderCode(const FString& ShaderCodeDir, FString const& Library, int32 ChunkId, int32 Priority = 0);
+		void ForEachShaderLibraryWithShaderMap(const FSHAHash& Hash, TFunctionRef<void(FShaderLibraryInstance* InLibraryInstance, int32 InShaderMapIndex)> InPerShaderMapCallback);
 		FShaderLibraryInstance* FindShaderLibraryForShaderMap(const FSHAHash& Hash, int32& OutShaderMapIndex);
 		FShaderLibraryInstance* FindShaderLibraryForShader(const FSHAHash& Hash, int32& OutShaderIndex);
 		uint32 GetShaderCount();
@@ -1274,6 +1275,15 @@ public:
 		return ShaderPlatform;
 	}
 
+	void ForEachShaderLibraryWithShaderMap(const FSHAHash& Hash, TFunctionRef<void(FShaderLibraryInstance* InLibraryInstance, int32 InShaderMapIndex)> InPerShaderMapCallback)
+	{
+		FRWScopeLock ReadLock(NamedLibrariesMutex, SLT_ReadOnly);
+		for (TTuple<FString, TUniquePtr<UE::ShaderLibrary::Private::FNamedShaderLibrary>>& NamedLibraryPair : NamedLibrariesStack)
+		{
+			NamedLibraryPair.Value->ForEachShaderLibraryWithShaderMap(Hash, InPerShaderMapCallback);
+		}
+	}
+
 	FShaderLibraryInstance* FindShaderLibraryForShaderMap(const FSHAHash& Hash, int32& OutShaderMapIndex)
 	{
 		FRWScopeLock ReadLock(NamedLibrariesMutex, SLT_ReadOnly);
@@ -1987,12 +1997,11 @@ public:
 	{
 		for (const FSHAHash& ShaderMapHash : ShaderMapHashes)
 		{
-			int32 ShaderMapIndex;
-			FShaderLibraryInstance* LibraryInstance = FindShaderLibraryForShaderMap(ShaderMapHash, ShaderMapIndex);
-			if (LibraryInstance)
-			{
+			// When resolving shadermaps to all of their IoChunks, we have to iterate through all shader library instances that contain such shadermap.
+			// Otherwise, the order of mounted (and unmounted) library instances might break the availability of shaders.
+			ForEachShaderLibraryWithShaderMap(ShaderMapHash, [&IoChunkIdResolvedFunc](FShaderLibraryInstance* LibraryInstance, int32 ShaderMapIndex) {
 				LibraryInstance->ResolvePackageShaderMap(ShaderMapIndex, IoChunkIdResolvedFunc);
-			}
+			});
 		}
 	}
 
@@ -2896,11 +2905,24 @@ bool UE::ShaderLibrary::Private::FNamedShaderLibrary::OpenShaderCode(const FStri
 	return true;
 }
 
+void UE::ShaderLibrary::Private::FNamedShaderLibrary::ForEachShaderLibraryWithShaderMap(const FSHAHash& Hash, TFunctionRef<void(FShaderLibraryInstance* InLibraryInstance, int32 InShaderMapIndex)> InPerShaderMapCallback)
+{
+	FRWScopeLock ReadLock(ComponentsMutex, SLT_ReadOnly);
+
+	for (const FShaderLibraryInstancePriorityPair& Instance : Components)
+	{
+		const int32 ShaderMapIndex = Instance.Value->Library->FindShaderMapIndex(Hash);
+		if (ShaderMapIndex != INDEX_NONE)
+		{
+			InPerShaderMapCallback(Instance.Value.Get(), ShaderMapIndex);
+		}
+	}
+}
+
 FShaderLibraryInstance* UE::ShaderLibrary::Private::FNamedShaderLibrary::FindShaderLibraryForShaderMap(const FSHAHash& Hash, int32& OutShaderMapIndex)
 {
 	FRWScopeLock ReadLock(ComponentsMutex, SLT_ReadOnly);
 
-	// Search in library opened order
 	for (const FShaderLibraryInstancePriorityPair& Instance : Components)
 	{
 		const int32 ShaderMapIndex = Instance.Value->Library->FindShaderMapIndex(Hash);
@@ -2917,7 +2939,6 @@ FShaderLibraryInstance* UE::ShaderLibrary::Private::FNamedShaderLibrary::FindSha
 {
 	FRWScopeLock ReadLock(ComponentsMutex, SLT_ReadOnly);
 
-	// Search in library opened order
 	for (const FShaderLibraryInstancePriorityPair& Instance : Components)
 	{
 		const int32 ShaderIndex = Instance.Value->Library->FindShaderIndex(Hash);
