Use FAppTime::GetApproximateCurrentTime() for texture LastRenderTime
- Only approximate time is needed for texture streaming calculations.
- FAppTime::GetApproximateCurrentTime() avoids an ensure when no time context is available on the calling thread.

#Synced-CL 50971219

diff --git a/Engine/Source/Runtime/RenderCore/Public/ShaderParameterUtils.h b/Engine/Source/Runtime/RenderCore/Public/ShaderParameterUtils.h
--- a/Engine/Source/Runtime/RenderCore/Public/ShaderParameterUtils.h
+++ b/Engine/Source/Runtime/RenderCore/Public/ShaderParameterUtils.h
@@ -129,7 +129,7 @@ inline void SetTextureParameter(
 {
 	if (TextureParameter.IsBound())
 	{
-		Texture->LastRenderTime = FApp::GetCurrentTime();
+		Texture->LastRenderTime = FAppTime::GetApproximateCurrentTime();
 	}
 
 	SetTextureParameter(BatchedParameters, TextureParameter, Texture->TextureRHI);
