[Shaders] Provide fallback to default material in shader code library when on-demand shadermaps are not installed as expected.

When shader maps are requested from an on-demand container and their respective IoChunks are not installed yet, the material will crash when attempting to create its shaders with an error like the following:
```
I/O Error '(Not Installed)' (IoChunkId='963d691fa0cabe8aadfbad09', Type=ShaderCode)
```
Materials can fall back to their default material when at least one shader hash is unavailable, but on-demand containers report those as being available as long as they have known IoChunk IDs.
If a shader map is accidentally requested from a system that cannot handle deferred loading of shaders (unlike IAD assets), they might not have been installed yet and the material will fail to create the shaders.
This change adds a mechanism to verify all IoChunks of a requested shader map have already been fully installed, which must be the case for IAD assets.

This also adds some helper functions to `FIoStoreShaderCodeArchiveHeader`  to iterate through all shaders of a shadermap and group.
These are currently only used by the debug function `DumpShaderCodeArchiveHashes()` to reduce the code duplication to resolve the indirection of shader indices.

#test Simulated failing IoChunk requests and verified a default materials were showing up instead of a crash.
#rnx

#Synced-CL 50957382

diff --git a/Engine/Source/Runtime/RHI/Public/RHIShaderLibrary.h b/Engine/Source/Runtime/RHI/Public/RHIShaderLibrary.h
--- a/Engine/Source/Runtime/RHI/Public/RHIShaderLibrary.h
+++ b/Engine/Source/Runtime/RHI/Public/RHIShaderLibrary.h
@@ -38,6 +38,15 @@ public:
 	virtual bool PreloadShaderMap(int32 ShaderMapIndex, FCoreDelegates::FAttachShaderReadRequestFunc AttachShaderReadRequestFunc) { return false; }
 	virtual void ReleasePreloadedShader(int32 ShaderIndex) {}
 
+	/**
+	 * Returns true if all chunks of the specified shadermap resource are ready for use.
+	 * This is only relevant for the IoStore shader library for containers that are loaded on demand.
+	 * We expect those chunks to be loaded by the on-demand content managers.
+	 * This function ensures the shadermap chunks are available when a material tries to load such a shadermap resource.
+	 * Otherwise, the material can fallback to the default material.
+	 */
+	virtual bool IsShaderMapResolved(int32 ShaderMapIndex) const { return true; };
+
 	virtual void AddRefPreloadedShaderGroup(int32 ShaderGroupIndex) {}
 	virtual void ReleasePreloadedShaderGroup(int32 ShaderGroupIndex) {}
 	virtual int32 GetGroupIndexForShader(int32 ShaderIndex) const { return INDEX_NONE; }
diff --git a/Engine/Source/Runtime/RenderCore/Private/ShaderLibrary/ShaderCodeArchive.cpp b/Engine/Source/Runtime/RenderCore/Private/ShaderLibrary/ShaderCodeArchive.cpp
--- a/Engine/Source/Runtime/RenderCore/Private/ShaderLibrary/ShaderCodeArchive.cpp
+++ b/Engine/Source/Runtime/RenderCore/Private/ShaderLibrary/ShaderCodeArchive.cpp
@@ -2299,13 +2299,12 @@ static void DumpShaderCodeArchiveHashes(const FIoStoreShaderCodeArchiveHeader& I
 	TMap<FSHAHash, TSet<FSHAHash>> ShaderToShaderMaps;
 	for (int32 ShaderMapIndex = 0; ShaderMapIndex < InLibraryHeader.ShaderMapEntries.Num(); ++ShaderMapIndex)
 	{
-		const FIoStoreShaderMapEntry& ShaderMapEntry = InLibraryHeader.ShaderMapEntries[ShaderMapIndex];
-		for (uint32 ShaderIndex = 0; ShaderIndex < ShaderMapEntry.NumShaders; ++ShaderIndex)
+		const FSHAHash& ShaderMapHash = InLibraryHeader.ShaderMapHashes[ShaderMapIndex];
+		InLibraryHeader.ForEachShaderInShaderMap(ShaderMapIndex, [&](int32 ShaderIndex)
 		{
-			const uint32 MappedShaderIndex = InLibraryHeader.ShaderIndices[ShaderMapEntry.ShaderIndicesOffset + ShaderIndex];
-			TSet<FSHAHash>& ShaderMapHahses = ShaderToShaderMaps.FindOrAdd(InLibraryHeader.ShaderHashes[MappedShaderIndex]);
-			ShaderMapHahses.Add(InLibraryHeader.ShaderMapHashes[ShaderMapIndex]);
-		}
+			TSet<FSHAHash>& ShaderMapHahses = ShaderToShaderMaps.FindOrAdd(InLibraryHeader.ShaderHashes[ShaderIndex]);
+			ShaderMapHahses.Add(ShaderMapHash);
+		});
 	}
 
 	// Dump hashes of all IoChunks with their shaders and associated shadermaps
@@ -2316,10 +2315,9 @@ static void DumpShaderCodeArchiveHashes(const FIoStoreShaderCodeArchiveHeader& I
 		const FIoStoreShaderGroupEntry& GroupEntry = InLibraryHeader.ShaderGroupEntries[GroupIndex];
 		DebugOutput.Appendf(TEXT("IoChunk %s (%d shaders)\n"), *LexToString(ShaderGroupChunkId), GroupEntry.NumShaders);
 
-		for (uint32 ShaderIndex = 0; ShaderIndex < GroupEntry.NumShaders; ++ShaderIndex)
+		InLibraryHeader.ForEachShaderInGroup(GroupIndex, [&](int32 ShaderIndex)
 		{
-			const uint32 MappedShaderIndex = InLibraryHeader.ShaderIndices[GroupEntry.ShaderIndicesOffset + ShaderIndex];
-			const FSHAHash& ShaderHash = InLibraryHeader.ShaderHashes[MappedShaderIndex];
+			const FSHAHash& ShaderHash = InLibraryHeader.ShaderHashes[ShaderIndex];
 			DebugOutput.Appendf(TEXT("  %s"), *LexToString(ShaderHash));
 
 			// Append all shadermaps that reference this shader
@@ -2344,7 +2342,7 @@ static void DumpShaderCodeArchiveHashes(const FIoStoreShaderCodeArchiveHeader& I
 			}
 
 			DebugOutput.Append(TEXT("\n"));
-		}
+		});
 		DebugOutput.Append(TEXT("\n"));
 	}
 	DebugOutput.Append(TEXT("\n\n\n"));
@@ -2355,12 +2353,11 @@ static void DumpShaderCodeArchiveHashes(const FIoStoreShaderCodeArchiveHeader& I
 	{
 		const FIoStoreShaderMapEntry& ShaderMapEntry = InLibraryHeader.ShaderMapEntries[ShaderMapIndex];
 		DebugOutput.Appendf(TEXT("ShaderMap %s (%d shaders)\n"), *LexToString(InLibraryHeader.ShaderMapHashes[ShaderMapIndex]), ShaderMapEntry.NumShaders);
-		for (uint32 ShaderIndex = 0; ShaderIndex < ShaderMapEntry.NumShaders; ++ShaderIndex)
+		InLibraryHeader.ForEachShaderInShaderMap(ShaderMapIndex, [&](int32 ShaderIndex)
 		{
-			const uint32 MappedShaderIndex = InLibraryHeader.ShaderIndices[ShaderMapEntry.ShaderIndicesOffset + ShaderIndex];
-			const FSHAHash& ShaderHash = InLibraryHeader.ShaderHashes[MappedShaderIndex];
+			const FSHAHash& ShaderHash = InLibraryHeader.ShaderHashes[ShaderIndex];
 			DebugOutput.Appendf(TEXT("  %s\n"), *LexToString(ShaderHash));
-		}
+		});
 		DebugOutput.Append(TEXT("\n"));
 	}
 
@@ -2821,6 +2818,26 @@ void FIoStoreShaderCodeArchive::ReleasePreloadedShader(int32 ShaderIndex)
 	);
 }
 
+bool FIoStoreShaderCodeArchive::IsShaderMapResolved(int32 ShaderMapIndex) const
+{
+	// Check availability of each IoChunk for the specified shadermap
+	const FIoStoreShaderMapEntry& ShaderMapEntry = Header.ShaderMapEntries[ShaderMapIndex];
+	for (uint32 ShaderIndex = 0; ShaderIndex < ShaderMapEntry.NumShaders; ++ShaderIndex)
+	{
+		const uint32 MappedShaderIndex = Header.ShaderIndices[ShaderMapEntry.ShaderIndicesOffset + ShaderIndex];
+		const int32 GroupIndex = GetGroupIndexForShader(static_cast<int32>(MappedShaderIndex));
+		const FIoChunkId& GroupChunkId = Header.ShaderGroupIoHashes[GroupIndex];
+
+		if (!IoDispatcher.GetSizeForChunk(GroupChunkId).IsOk())
+		{
+			UE_LOG(LogShaderLibrary, Warning, TEXT("Shadermap '%s' (%d) not resolved in library %s because IoChunk '%s' is not available"),
+				*LexToString(Header.ShaderMapHashes[ShaderMapIndex]), ShaderMapIndex, *GetName(), *LexToString(GroupChunkId));
+			return false;
+		}
+	}
+	return true;
+}
+
 int32 FIoStoreShaderCodeArchive::FindShaderMapIndex(const FSHAHash& Hash)
 {
 	const uint32 Key = GetTypeHash(Hash);
diff --git a/Engine/Source/Runtime/RenderCore/Private/ShaderLibrary/ShaderCodeLibrary.cpp b/Engine/Source/Runtime/RenderCore/Private/ShaderLibrary/ShaderCodeLibrary.cpp
--- a/Engine/Source/Runtime/RenderCore/Private/ShaderLibrary/ShaderCodeLibrary.cpp
+++ b/Engine/Source/Runtime/RenderCore/Private/ShaderLibrary/ShaderCodeLibrary.cpp
@@ -2841,7 +2841,10 @@ bool UE::ShaderLibrary::Private::FNamedShaderLibrary::OpenShaderCode(const FStri
 		}
 	}
 
-	FShaderLibraryInstance* LibraryInstance = FShaderLibraryInstance::Create(ShaderPlatform, ShaderCodeDir, Library, ChunkId);
+	// Any shader library instance with a lower than default priority is expected to have its shadermaps resolved before requesting shaders to be created.
+	// This is the case for shader libraries from on-demand containers since all IoChunks associated with their shadermaps need to be resolved before any CreateShader() requests are permitted.
+	const bool bRequireResolvedShaderMaps = Priority < 0;
+	FShaderLibraryInstance* LibraryInstance = FShaderLibraryInstance::Create(ShaderPlatform, ShaderCodeDir, Library, ChunkId, bRequireResolvedShaderMaps);
 	if (LibraryInstance == nullptr)
 	{
 		UE_LOG(LogShaderLibrary, Verbose, TEXT("Cooked Context: No Shared Shader Library for: %s and native library not supported."), *Library);
diff --git a/Engine/Source/Runtime/RenderCore/Private/ShaderLibrary/ShaderLibraryInstance.cpp b/Engine/Source/Runtime/RenderCore/Private/ShaderLibrary/ShaderLibraryInstance.cpp
--- a/Engine/Source/Runtime/RenderCore/Private/ShaderLibrary/ShaderLibraryInstance.cpp
+++ b/Engine/Source/Runtime/RenderCore/Private/ShaderLibrary/ShaderLibraryInstance.cpp
@@ -17,7 +17,7 @@
 #include "UObject/NameTypes.h"
 
 
-FShaderLibraryInstance* FShaderLibraryInstance::Create(EShaderPlatform InShaderPlatform, const FString& ShaderCodeDir, FString const& InLibraryName, int32 InChunkId)
+FShaderLibraryInstance* FShaderLibraryInstance::Create(EShaderPlatform InShaderPlatform, const FString& ShaderCodeDir, FString const& InLibraryName, int32 InChunkId, bool bInRequireResolvedShaderMaps)
 {
 	using namespace UE::ShaderLibrary::Private;
 
@@ -78,6 +78,7 @@ FShaderLibraryInstance* FShaderLibraryInstance::Create(EShaderPlatform InShaderP
 	Instance->Library = Library;
 	Instance->ShaderCodeDirectory = ShaderCodeDirectory;
 	Instance->ChunkId = InChunkId;
+	Instance->bRequireResolvedShaderMaps = bInRequireResolvedShaderMaps;
 
 	const int32 NumResources = Library->GetNumShaderMaps();
 	Instance->Resources.AddZeroed(NumResources);
@@ -185,8 +186,20 @@ TRefCountPtr<FShaderMapResource_SharedCode> FShaderLibraryInstance::GetResource(
 	return Resources[ShaderMapIndex];
 }
 
+bool FShaderLibraryInstance::IsShaderMapResourceResolved(int32 ShaderMapIndex) const
+{
+	// If this library instance loads its shaders on-demand, then we expect that all IoChunks for this shader map have already been resolved.
+	// Otherwise, this will fail to create shaders when the chunks are requested to be installed too late.
+	return !bRequireResolvedShaderMaps || Library->IsShaderMapResolved(ShaderMapIndex);
+}
+
 TRefCountPtr<FShaderMapResource_SharedCode> FShaderLibraryInstance::AddResource(int32 ShaderMapIndex, FArchive* Ar)
 {
+	if (!IsShaderMapResourceResolved(ShaderMapIndex))
+	{
+		return nullptr;
+	}
+
 	TRefCountPtr<FShaderMapResource_SharedCode> OutResource;
 	bool bPreload = false;
 	{
@@ -329,6 +342,11 @@ void FShaderLibraryInstance::ResolvePackageShaderMap(int32 ShaderMapIndex, FCore
 
 void FShaderLibraryInstance::PreloadPackageShaderMap(int32 ShaderMapIndex, FCoreDelegates::FAttachShaderReadRequestFunc AttachShaderReadRequestFunc)
 {
+	if (!IsShaderMapResourceResolved(ShaderMapIndex))
+	{
+		return;
+	}
+
 	FRWScopeLock Locker(ResourceLock, SLT_Write);
 	FShaderMapResource_SharedCode*& Resource = Resources[ShaderMapIndex];
 	if (!Resource)
diff --git a/Engine/Source/Runtime/RenderCore/Private/ShaderLibrary/ShaderLibraryInstance.h b/Engine/Source/Runtime/RenderCore/Private/ShaderLibrary/ShaderLibraryInstance.h
--- a/Engine/Source/Runtime/RenderCore/Private/ShaderLibrary/ShaderLibraryInstance.h
+++ b/Engine/Source/Runtime/RenderCore/Private/ShaderLibrary/ShaderLibraryInstance.h
@@ -19,7 +19,7 @@
 class FShaderLibraryInstance
 {
 public:
-	static FShaderLibraryInstance* Create(EShaderPlatform InShaderPlatform, const FString& ShaderCodeDir, FString const& InLibraryName, int32 InChunkId);
+	static FShaderLibraryInstance* Create(EShaderPlatform InShaderPlatform, const FString& ShaderCodeDir, FString const& InLibraryName, int32 InChunkId, bool bInRequireResolvedShaderMaps = false);
 
 	~FShaderLibraryInstance();
 
@@ -72,6 +72,8 @@ private:
 
 	bool OnExternalReadCallback(const FGraphEventArray& Events, double RemainingTime);
 
+	bool IsShaderMapResourceResolved(int32 ShaderMapIndex) const;
+
 	enum class EPreloadingState : uint8
 	{
 		NotPreloaded = 0, // Not preloaded.
@@ -100,5 +102,9 @@ private:
 	FString ShaderCodeDirectory;
 
 	/** Chunk ID of this library, if any (INDEX_NONE if it is monolithic) */
-	int32 ChunkId;
+	int32 ChunkId = INDEX_NONE;
+
+	/** Specifies whether shadermaps must be fully resolved and ready to use when loading shadermap resources.
+	    By default, resources can be loaded even when their shadermaps have not be resolved yet. */
+	bool bRequireResolvedShaderMaps = false;
 };
diff --git a/Engine/Source/Runtime/RenderCore/Public/ShaderCodeArchive.h b/Engine/Source/Runtime/RenderCore/Public/ShaderCodeArchive.h
--- a/Engine/Source/Runtime/RenderCore/Public/ShaderCodeArchive.h
+++ b/Engine/Source/Runtime/RenderCore/Public/ShaderCodeArchive.h
@@ -734,6 +734,18 @@ struct FIoStoreShaderGroupEntry
 
 struct FIoStoreShaderCodeArchiveHeader
 {
+private:
+	template <typename TGroupEntry>
+	void ForEachShaderInternal(const TArray<TGroupEntry>& InGroupEntries, int32 InEntryIndex, TFunctionRef<void(int32 ShaderIndex)> InCallback) const
+	{
+		const TGroupEntry& GroupEntry = InGroupEntries[InEntryIndex];
+		for (uint32 IndirectShaderIndex = 0; IndirectShaderIndex < GroupEntry.NumShaders; ++IndirectShaderIndex)
+		{
+			const uint32 MappedShaderIndex = this->ShaderIndices[GroupEntry.ShaderIndicesOffset + IndirectShaderIndex];
+			InCallback(static_cast<int32>(MappedShaderIndex));
+		}
+	}
+
 public:
 
 	/** Hashes of all shadermaps in the library */
@@ -789,6 +801,18 @@ public:
 		return 0;
 	}
 
+	/** Helper function to iterate through all shaders of the specified shadermap and resolve the indirection of shader indices. */
+	FORCEINLINE void ForEachShaderInShaderMap(int32 ShaderMapIndex, TFunctionRef<void(int32 ShaderIndex)> InCallback) const
+	{
+		ForEachShaderInternal(this->ShaderMapEntries, ShaderMapIndex, InCallback);
+	}
+
+	/** Helper function to iterate through all shaders of the specified shader group and resolve the indirection of shader indices. */
+	FORCEINLINE void ForEachShaderInGroup(int32 GroupIoIndex, TFunctionRef<void(int32 ShaderIndex)> InCallback) const
+	{
+		ForEachShaderInternal(this->ShaderGroupEntries, GroupIoIndex, InCallback);
+	}
+
 	uint64 GetAllocatedSize() const
 	{
 		return sizeof(*this) +
@@ -876,6 +900,11 @@ public:
 	virtual bool PreloadShaderMap(int32 ShaderMapIndex, FGraphEventArray& OutCompletionEvents) override;
 	virtual bool PreloadShaderMap(int32 ShaderMapIndex, FCoreDelegates::FAttachShaderReadRequestFunc AttachShaderReadRequestFunc) override;
 	virtual void ReleasePreloadedShader(int32 ShaderIndex) override;
+
+	/** Returns true if all IoChunks associated with the specified shadermap have been resolved and are available for use without any further delay.
+	    Otherwise, at least one IoChunk is either unknown or has not been installed yet (when coming from an on-demand container). */
+	virtual bool IsShaderMapResolved(int32 ShaderMapIndex) const override;
+
 	virtual TRefCountPtr<FRHIShader> CreateShader(int32 Index, bool bRequired = true) override;
 	virtual void Teardown() override;
 
