[Shaders] Allow verbose logging of resolved IoChunks for each on-demand shadermap.

This enables verbose logging for each shadermap that is resolved into their IoChunks for IAD support.
Can be enabled via command line argument `-ini:Engine:[Core.Log]:LogShaderLibrary=VeryVerbose`.

This also reveals how many shadermaps for a given asset are requested to be resolved but are not used because of multiple shader formats being mounted.
On Windows for instance, there are three shader formats being mounted: PCD3D_SM6, PCD3D_SM5, PCD3D_ES3_1.
When launching with -FeatureLevelES31, the majority of shadermaps won't be loaded and they are logged like this:
```
LogShaderLibrary: VeryVerbose: Shadermap '9D813814CAF66DC93FBB2BE4882FCECB54512FDC' resolved to IoChunks: None
```
Whereas those that belong to PCD3D_ES3_1 are logged like this:
```
LogShaderLibrary: VeryVerbose: Shadermap '1FE80CEAE446D557DB4F94D16DDB12DEC765EE96' resolved to IoChunks: 5ad885fa816d51f5d4abd909, c27e38e8e6a37091b1f7ed09, 6f1422f3bae1884923404909, 9ffecd79d8797b5e5bd4fe09, d6a7c2de6a183c21d7230f09, 2db46b8db3a63f4a45e0e009, f41d1f38a70d1a1fb4377c09, bcd0c55e814e51c906a1e609
```

#rnx

#Synced-CL 51023540

diff --git a/Engine/Source/Runtime/RenderCore/Private/ShaderLibrary/ShaderCodeLibrary.cpp b/Engine/Source/Runtime/RenderCore/Private/ShaderLibrary/ShaderCodeLibrary.cpp
--- a/Engine/Source/Runtime/RenderCore/Private/ShaderLibrary/ShaderCodeLibrary.cpp
+++ b/Engine/Source/Runtime/RenderCore/Private/ShaderLibrary/ShaderCodeLibrary.cpp
@@ -863,6 +863,39 @@ class FShaderLibrariesCollection
 #endif
 	}
 
+private:
+	// Prints a comma-separated list of all IoChunkIds that the specified shadermap will be resolved to.
+	FString PrintResolvedShaderMapIoChunks(const FSHAHash& ShaderMapHash)
+	{
+		TSet<FIoChunkId> ShaderMapIoChunks;
+
+		ForEachShaderLibraryWithShaderMap(ShaderMapHash, [&ShaderMapIoChunks](FShaderLibraryInstance* LibraryInstance, int32 ShaderMapIndex) {
+			LibraryInstance->ResolvePackageShaderMap(ShaderMapIndex, [&ShaderMapIoChunks](const FIoChunkId& IoChunk) {
+				ShaderMapIoChunks.Add(IoChunk);
+			});
+		});
+
+		// Log all IoChunks to verbose log
+		TStringBuilder<1024> IoChunksString;
+		if (ShaderMapIoChunks.IsEmpty())
+		{
+			IoChunksString.Append(TEXT("None"));
+		}
+		else
+		{
+			for (const FIoChunkId& IoChunk : ShaderMapIoChunks)
+			{
+				if (IoChunksString.Len() > 0)
+				{
+					IoChunksString.Append(TEXT(", "));
+				}
+				IoChunksString.Append(LexToString(IoChunk));
+			}
+		}
+
+		return FString(IoChunksString);
+	}
+
 public:
 	static FShaderLibrariesCollection* Impl;
 
@@ -2002,6 +2035,8 @@ public:
 			ForEachShaderLibraryWithShaderMap(ShaderMapHash, [&IoChunkIdResolvedFunc](FShaderLibraryInstance* LibraryInstance, int32 ShaderMapIndex) {
 				LibraryInstance->ResolvePackageShaderMap(ShaderMapIndex, IoChunkIdResolvedFunc);
 			});
+
+			UE_LOG(LogShaderLibrary, VeryVerbose, TEXT("Shadermap '%s' resolved to IoChunks: %s"), *LexToString(ShaderMapHash), *PrintResolvedShaderMapIoChunks(ShaderMapHash));
 		}
 	}
 
