Substrate - material quality level affect used shaders.

#tests PC editor, cvar r.Substrate.MaterialQualityLevelsClosuresPerPixel="L1-M1-H2-E3", toggle material quality levels and see impact on substrate material with many slabs in editor.
#rnx

#Synced-CL 51046328

#Synced-CL 51046328

diff --git a/Engine/Source/Editor/MaterialEditor/Private/SMaterialEditorSubstrateWidget.cpp b/Engine/Source/Editor/MaterialEditor/Private/SMaterialEditorSubstrateWidget.cpp
--- a/Engine/Source/Editor/MaterialEditor/Private/SMaterialEditorSubstrateWidget.cpp
+++ b/Engine/Source/Editor/MaterialEditor/Private/SMaterialEditorSubstrateWidget.cpp
@@ -90,7 +90,7 @@ void SMaterialEditorSubstrateWidget::Construct(const FArguments& InArgs, TWeakPt
 		.ToolTipText(LOCTEXT("CheckBoxBytesPerPixelOverride", "This will force the byte per pixel count for the preview material. It cannot go higher than the current project setting."))
 		.OnCheckStateChanged(this, &SMaterialEditorSubstrateWidget::OnCheckBoxBytesPerPixelChanged);
 
-	ClosuresPerPixelOverride = Substrate::GetClosurePerPixel(SP_PCD3D_SM5);
+	ClosuresPerPixelOverride = Substrate::GetClosurePerPixel(SP_PCD3D_SM5, Substrate::IgnoreMaterialQualityLevelsForClosurePerPixel);
 	ClosuresPerPixelOverrideInput = SNew(SNumericEntryBox<uint32>)
 		.MinDesiredValueWidth(150.0f)
 		.MinValue(1)
@@ -855,8 +855,8 @@ FReply SMaterialEditorSubstrateWidget::OnButtonApplyToPreview()
 		SubstrateCompilationConfig.bFullSimplify = CheckBoxForceFullSimplification->IsChecked();
 
 		// Have a look at MaterialStats.cpp when we want to visualize a specific platform.
-		SubstrateCompilationConfig.BytesPerPixelOverride	= CheckBoxBytesPerPixelOverride->IsChecked()	? FMath::Clamp(BytesPerPixelOverride,		12,	Substrate::GetBytePerPixel(GMaxRHIShaderPlatform))		: -1;
-		SubstrateCompilationConfig.ClosuresPerPixelOverride = CheckBoxClosuresPerPixelOverride->IsChecked() ? FMath::Clamp(ClosuresPerPixelOverride,	1,	Substrate::GetClosurePerPixel(GMaxRHIShaderPlatform))	: -1;
+		SubstrateCompilationConfig.BytesPerPixelOverride	= CheckBoxBytesPerPixelOverride->IsChecked()	? FMath::Clamp(BytesPerPixelOverride,		12,	Substrate::GetBytePerPixel(GMaxRHIShaderPlatform))																: -1;
+		SubstrateCompilationConfig.ClosuresPerPixelOverride = CheckBoxClosuresPerPixelOverride->IsChecked() ? FMath::Clamp(ClosuresPerPixelOverride,	1,	Substrate::GetClosurePerPixel(GMaxRHIShaderPlatform, Substrate::IgnoreMaterialQualityLevelsForClosurePerPixel))	: -1;
 
 		MaterialInterface->SetSubstrateCompilationConfig(SubstrateCompilationConfig);
 
diff --git a/Engine/Source/Runtime/Core/Private/UObject/FortniteMainBranchObjectVersion.cpp b/Engine/Source/Runtime/Core/Private/UObject/FortniteMainBranchObjectVersion.cpp
--- a/Engine/Source/Runtime/Core/Private/UObject/FortniteMainBranchObjectVersion.cpp
+++ b/Engine/Source/Runtime/Core/Private/UObject/FortniteMainBranchObjectVersion.cpp
@@ -13,7 +13,7 @@ TMap<FGuid, FGuid> FFortniteMainBranchObjectVersion::GetSystemGuids()
 	SystemGuids.Add(DevGuids.MATERIALSHADERMAP_DERIVEDDATA_VER, FGuid("1A77CBCDFBDB4C12BF8563EDBB7D8888"));
 	SystemGuids.Add(DevGuids.NIAGARASHADERMAP_DERIVEDDATA_VER, FGuid("CC229D919A214E59B79D01AC2E5BF981"));
 	SystemGuids.Add(DevGuids.Niagara_LatestScriptCompileVersion, FGuid("FB15C2870A4E4AE8B5F58566F0E3825E"));
-	SystemGuids.Add(DevGuids.MaterialTranslationDDCVersion, FGuid("0095BE8673FE4F6DB7AC9CF5AF47F349"));
+	SystemGuids.Add(DevGuids.MaterialTranslationDDCVersion, FGuid("B79683EE9E684C13BAC2FDE8B180477F"));
 
 	// These GUIDs are in AppendRestrictedFortniteMainStreamObjectVersionGuids as they should rarely be changed.
 	//SystemGuids.Add(DevGuids.NANITE_DERIVEDDATA_VER, FGuid("ACF98184978742A592537DA223A1E6BE"));
diff --git a/Engine/Source/Runtime/Engine/Private/Materials/MaterialCachedData.cpp b/Engine/Source/Runtime/Engine/Private/Materials/MaterialCachedData.cpp
--- a/Engine/Source/Runtime/Engine/Private/Materials/MaterialCachedData.cpp
+++ b/Engine/Source/Runtime/Engine/Private/Materials/MaterialCachedData.cpp
@@ -509,7 +509,18 @@ void FMaterialCachedExpressionData::UpdateForExpressions(const FMaterialCachedEx
 	static const FGuid FirstPersonInterpolationAlphaGuid = FMaterialAttributeDefinitionMap::GetCustomAttributeID(TEXT("FirstPersonInterpolationAlpha"));
 	static const FGuid TemporalResponsivenessGuid = FMaterialAttributeDefinitionMap::GetCustomAttributeID(TEXT("TemporalResponsiveness"));
 	static const FGuid MotionVectorWorldOffsetGuid = FMaterialAttributeDefinitionMap::GetCustomAttributeID(TEXT("MotionVectorWorldOffset"));
-	
+
+	// Enabled specific quality levels based on material quality level closure count difference as compared to ignoring per material quality level constraints.
+	// See r.Substrate.MaterialQualityLevelsClosuresPerPixel.
+	if (Substrate::IsSubstrateEnabled() && !Substrate::IsSubstrateBlendableGBufferEnabled(GMaxRHIShaderPlatform))
+	{
+		uint32 ReferenceClosurePerPixelCount = Substrate::GetClosurePerPixel(GMaxRHIShaderPlatform, Substrate::IgnoreMaterialQualityLevelsForClosurePerPixel);
+		for (int32 MaterialQualityLevel = 0; MaterialQualityLevel < EMaterialQualityLevel::Num; MaterialQualityLevel++)
+		{
+			QualityLevelsUsed[MaterialQualityLevel] = ReferenceClosurePerPixelCount != Substrate::GetClosurePerPixel(GMaxRHIShaderPlatform, MaterialQualityLevel);
+		}
+	}
+
 	for (UMaterialExpression* Expression : Expressions)
 	{
 		if (!Expression)
diff --git a/Engine/Source/Runtime/Engine/Private/Materials/SubstrateTranslatorCommon.cpp b/Engine/Source/Runtime/Engine/Private/Materials/SubstrateTranslatorCommon.cpp
--- a/Engine/Source/Runtime/Engine/Private/Materials/SubstrateTranslatorCommon.cpp
+++ b/Engine/Source/Runtime/Engine/Private/Materials/SubstrateTranslatorCommon.cpp
@@ -683,7 +683,7 @@ bool FSubstrateCompilationContext::SubstrateGenerateDerivedMaterialOperatorData(
 
 	EShaderPlatform ShaderPlatform = SubstrateTranslatorData.GetShaderPlatform();
 	const uint32 SubstrateBytePerPixel = SubstrateTranslatorData.SubstrateCompilationConfig.BytesPerPixelOverride > 0 ? SubstrateTranslatorData.SubstrateCompilationConfig.BytesPerPixelOverride : Substrate::GetBytePerPixel(ShaderPlatform);
-	const uint32 SubstrateClosurePerPixel = SubstrateTranslatorData.SubstrateCompilationConfig.ClosuresPerPixelOverride > 0 ? SubstrateTranslatorData.SubstrateCompilationConfig.ClosuresPerPixelOverride : Substrate::GetClosurePerPixel(ShaderPlatform);
+	const uint32 SubstrateClosurePerPixel = SubstrateTranslatorData.SubstrateCompilationConfig.ClosuresPerPixelOverride > 0 ? SubstrateTranslatorData.SubstrateCompilationConfig.ClosuresPerPixelOverride : Substrate::GetClosurePerPixel(ShaderPlatform, GetCurrentMaterialQualityLevelChecked());
 
 	bool bFirstLoop = true;
 	bool bRequestMaterialDetailsSet = false;
@@ -2204,7 +2204,7 @@ void FSubstrateTranslatorData::GenerateMaterialCompilationOutputAndDefines(
 			{
 				// Output some debug info as comment in code and in the material stat window
 				const uint32 SubstrateBytePerPixel_Platform = Substrate::GetBytePerPixel(GetShaderPlatform());
-				const uint32 SubstrateClosurePerPixel_Platform = Substrate::GetClosurePerPixel(GetShaderPlatform());
+				const uint32 SubstrateClosurePerPixel_Platform = Substrate::GetClosurePerPixel(GetShaderPlatform(), GetCurrentMaterialQualityLevelChecked());
 				FString SubstrateMaterialContextDescription;
 
 				auto GetSubstrateCompilationContextName = [](uint32 Index)
diff --git a/Engine/Source/Runtime/Engine/Private/ShaderCompiler/ShaderCompiler.cpp b/Engine/Source/Runtime/Engine/Private/ShaderCompiler/ShaderCompiler.cpp
--- a/Engine/Source/Runtime/Engine/Private/ShaderCompiler/ShaderCompiler.cpp
+++ b/Engine/Source/Runtime/Engine/Private/ShaderCompiler/ShaderCompiler.cpp
@@ -4103,7 +4103,16 @@ void GlobalBeginCompileShader(
 			const uint32 SubstrateUintPerPixel = Substrate::GetBytePerPixel(Target.GetPlatform()) / 4u;
 			SET_SHADER_DEFINE(Input.Environment, SUBSTRATE_MATERIAL_NUM_UINTS, SubstrateUintPerPixel);
 
-			const uint32 SubstrateClosurePerPixel = Substrate::GetClosurePerPixel(Target.GetPlatform());
+			// Get the highest closure count accounting for all material quality levels for SUBSTRATE_MATERIAL_CLOSURE_COUNT.
+			// This is because it needs to account for the maximum value for all the global shaders being compiled only once. (no permutation per material quality level).
+			uint32 SubstrateClosurePerPixel = 1;
+			if (!Substrate::IsSubstrateBlendableGBufferEnabled(Target.GetPlatform()))
+			{
+				for (int8 MaterialQualityLevel = EMaterialQualityLevel::Low; MaterialQualityLevel < EMaterialQualityLevel::Num; ++MaterialQualityLevel)
+				{
+					SubstrateClosurePerPixel = FMath::Max(SubstrateClosurePerPixel, Substrate::GetClosurePerPixel(Target.GetPlatform(), MaterialQualityLevel));
+				}
+			}
 			SET_SHADER_DEFINE(Input.Environment, SUBSTRATE_MATERIAL_CLOSURE_COUNT, SubstrateClosurePerPixel);
 
 			const bool bSubstrateDBufferPass = Substrate::IsDBufferPassEnabled(Target.GetPlatform());
diff --git a/Engine/Source/Runtime/RenderCore/Private/RenderUtils.cpp b/Engine/Source/Runtime/RenderCore/Private/RenderUtils.cpp
--- a/Engine/Source/Runtime/RenderCore/Private/RenderUtils.cpp
+++ b/Engine/Source/Runtime/RenderCore/Private/RenderUtils.cpp
@@ -2117,6 +2117,12 @@ static TAutoConsoleVariable<int32> CVarSubstrateProjectClosuresPerPixel(
 	TEXT("Define the maximum number of closures evaluated per pixel for this project. Materials with a number of closures beyond this limit will be simplified until it fits the requirements. This only applies when Adaptive GBuffer is enabled. On platforms not supporting Adaptive GBuffer, the material will be simplified into a single closure."),
 	ECVF_ReadOnly | ECVF_RenderThreadSafe | ECVF_Preview);
 
+static TAutoConsoleVariable<FString> CVarSubstrateMaterialQualityLevelsClosuresPerPixel(
+	TEXT("r.Substrate.MaterialQualityLevelsClosuresPerPixel"),
+	TEXT(""),	// Empty => invalid so ignored.
+	TEXT("Define the maximum number of closures evaluated per pixel for a given material quality level. Empty by default: no impact. Otherwise L1-M1-H1-E1 to set closure count for each of the material quality level (Low, Medium, High and Epic respecively)."),
+	ECVF_ReadOnly | ECVF_Preview);
+
 static TAutoConsoleVariable<int32> CVarSubstrateProjectClosuresPerPixelOverride(
 	TEXT("r.Substrate.ProjectClosuresPerPixelOverride"),
 	0,	// By default, use the mininum between the platform's closure count and the project's closure count (i.e., ESubstrateClosureConfig::Type::UseMinProjectAndPlatforClosuresPerPixel)
@@ -2260,6 +2266,67 @@ namespace Substrate
 	}
 
 	uint32 GetClosurePerPixel(EShaderPlatform InPlatform)
+	{
+		return GetClosurePerPixel(InPlatform, 0);
+	}
+
+	struct FSubstratePerMaterialLevelQualityClosureCount
+	{
+		bool bIsValid = false;	// Only true if the cvar has been specified explicitely, i.e. not empty.
+		uint32 LowMaterialQualityMaxClosureCount	= 1;
+		uint32 MediumMaterialQualityMaxClosureCount	= 1;
+		uint32 HighMaterialQualityMaxClosureCount	= 1;
+		uint32 EpicMaterialQualityMaxClosureCount	= 1;
+	};
+
+	const FSubstratePerMaterialLevelQualityClosureCount& GetPerMaterialLevelQualityClosureCount(EShaderPlatform InPlatform)
+	{
+		// Typically there is only one platform. In editor many can be evaluated for multiple preview in a session. So a Map should be enough.
+		static TMap<EShaderPlatform, FSubstratePerMaterialLevelQualityClosureCount> CachedPlatformResults;
+
+		FSubstratePerMaterialLevelQualityClosureCount* FoundResult = CachedPlatformResults.Find(InPlatform);
+		if (!FoundResult)
+		{
+			FSubstratePerMaterialLevelQualityClosureCount& Result = CachedPlatformResults.Add(InPlatform);
+
+			static FShaderPlatformCachedIniValue<FString> CVarMaterialQualityLevelsClosuresPerPixel(TEXT("r.Substrate.MaterialQualityLevelsClosuresPerPixel"));
+			FString CVARString = CVarMaterialQualityLevelsClosuresPerPixel.Get(InPlatform);
+			if (!CVARString.IsEmpty())
+			{
+				// The format must be in style "L1-M1-H1-E1"
+				int32 ValidRead =
+#if PLATFORM_WINDOWS || PLATFORM_MICROSOFT
+					sscanf_s(
+#else
+					sscanf(
+#endif
+					TCHAR_TO_ANSI(*CVARString), "L%u-M%u-H%u-E%u",
+					&Result.LowMaterialQualityMaxClosureCount,
+					&Result.MediumMaterialQualityMaxClosureCount,
+					&Result.HighMaterialQualityMaxClosureCount,
+					&Result.EpicMaterialQualityMaxClosureCount);
+				Result.bIsValid = true;
+
+				if (ValidRead != 4) // not all input have been read.
+				{
+					Result.bIsValid = false;
+					Result.LowMaterialQualityMaxClosureCount = 1;
+					Result.MediumMaterialQualityMaxClosureCount = 1;
+					Result.HighMaterialQualityMaxClosureCount = 1;
+					Result.EpicMaterialQualityMaxClosureCount = 1;
+				}
+
+				Result.LowMaterialQualityMaxClosureCount = FMath::Max(1u, Result.LowMaterialQualityMaxClosureCount);
+				Result.MediumMaterialQualityMaxClosureCount = FMath::Max(1u, Result.MediumMaterialQualityMaxClosureCount);
+				Result.HighMaterialQualityMaxClosureCount = FMath::Max(1u, Result.HighMaterialQualityMaxClosureCount);
+				Result.EpicMaterialQualityMaxClosureCount = FMath::Max(1u, Result.EpicMaterialQualityMaxClosureCount);
+			}
+			return Result;
+		}
+		return *FoundResult;
+	}
+
+	uint32 GetClosurePerPixel(EShaderPlatform InPlatform, uint8 InMaterialQualityLevel)
 	{
 		if (Substrate::IsSubstrateBlendableGBufferEnabled(InPlatform))
 		{
@@ -2270,6 +2337,28 @@ namespace Substrate
 
 		// Take min between the platform's settings and the project's settings
 		const int32 ProjectClosurePerPixel = FMath::Clamp(CVarSubstrateProjectClosuresPerPixel.GetValueOnAnyThread(), 1, SUBSTRATE_MAX_CLOSURE_COUNT);
+
+		const FSubstratePerMaterialLevelQualityClosureCount& PerMaterialLevelQualityClosureCount = GetPerMaterialLevelQualityClosureCount(InPlatform);
+		if (PerMaterialLevelQualityClosureCount.bIsValid && InMaterialQualityLevel != IgnoreMaterialQualityLevelsForClosurePerPixel)
+		{
+			check(0 <= InMaterialQualityLevel && InMaterialQualityLevel <= 3);
+			uint32 SelectedMaterialQualityMaxClosureCount = PerMaterialLevelQualityClosureCount.LowMaterialQualityMaxClosureCount;
+			switch (InMaterialQualityLevel)
+			{
+				//	case 0: // Low
+			case 2: // Medium
+				SelectedMaterialQualityMaxClosureCount = PerMaterialLevelQualityClosureCount.MediumMaterialQualityMaxClosureCount;
+				break;
+			case 1: // High
+				SelectedMaterialQualityMaxClosureCount = PerMaterialLevelQualityClosureCount.HighMaterialQualityMaxClosureCount;
+				break;
+			case 3: // Epic
+				SelectedMaterialQualityMaxClosureCount = PerMaterialLevelQualityClosureCount.EpicMaterialQualityMaxClosureCount;
+				break;
+			}
+			// The per material quality level value can only restricts the platform configured capacity.
+			OutClosurePerPixel = FMath::Min(SelectedMaterialQualityMaxClosureCount, OutClosurePerPixel);
+		}
 		OutClosurePerPixel = FMath::Min<uint32>(ProjectClosurePerPixel, OutClosurePerPixel);
 
 		// Force max closure count if enabled
@@ -2423,8 +2512,32 @@ namespace Substrate
 	bool IsOpaqueRoughRefractionEnabled(EShaderPlatform InPlatform)
 	{
 		static FShaderPlatformCachedIniValue<int32> CVar(TEXT("r.Substrate.OpaqueMaterialRoughRefraction"));
-		const bool bManyClosures = GetClosurePerPixel(InPlatform) > 1; // Opaque Rough refraction can only happen when there is at least 2 closures.
-		return IsSubstrateEnabled() && CVar.Get(InPlatform) > 0 && !IsSubstrateBlendableGBufferEnabled(InPlatform) && bManyClosures;
+
+		const bool bIsSubstrateBlendableGBufferEnabled = !IsSubstrateBlendableGBufferEnabled(InPlatform);
+
+		// By default sample closures from cvars using highest material quality level. 
+		// If even the highest level does not have many closures then we assume lower levels won't neither.
+		bool bManyClosures = GetClosurePerPixel(InPlatform, Substrate::IgnoreMaterialQualityLevelsForClosurePerPixel) > 1;
+
+		// If we have per material quality level closure count constraints, we should check them all and assume bManyClosures==true only when any of them is >1 since material/global shaders are all compile for only one case.
+		const FSubstratePerMaterialLevelQualityClosureCount& PerMaterialLevelQualityClosureCount = GetPerMaterialLevelQualityClosureCount(InPlatform);
+		if (bManyClosures && !bIsSubstrateBlendableGBufferEnabled && PerMaterialLevelQualityClosureCount.bIsValid)
+		{
+			if (PerMaterialLevelQualityClosureCount.LowMaterialQualityMaxClosureCount > 1
+				|| PerMaterialLevelQualityClosureCount.MediumMaterialQualityMaxClosureCount > 1
+				|| PerMaterialLevelQualityClosureCount.HighMaterialQualityMaxClosureCount > 1
+				|| PerMaterialLevelQualityClosureCount.EpicMaterialQualityMaxClosureCount > 1)
+			{
+				bManyClosures = true;
+			}
+			else
+			{
+				// All material quality level will constrain the closure count to 1 so set bManyClosures to false.
+				bManyClosures = false;
+			}
+		}
+
+		return IsSubstrateEnabled() && CVar.Get(InPlatform) > 0 && !bIsSubstrateBlendableGBufferEnabled && bManyClosures;
 	}
 
 	bool IsAdvancedVisualizationEnabled()
diff --git a/Engine/Source/Runtime/RenderCore/Private/Shader.cpp b/Engine/Source/Runtime/RenderCore/Private/Shader.cpp
--- a/Engine/Source/Runtime/RenderCore/Private/Shader.cpp
+++ b/Engine/Source/Runtime/RenderCore/Private/Shader.cpp
@@ -2509,7 +2509,11 @@ void ShaderMapAppendKey(EShaderPlatform Platform, FShaderKeyGenerator& KeyGen)
 		{
 			KeyGen.AppendSeparator();
 			KeyGen.Append(TEXT("CLOSURE"));
-			KeyGen.Append(Substrate::GetClosurePerPixel(Platform));
+			// We append all closure count to the key from low to epic material quality levels for the current platfom.
+			KeyGen.Append(Substrate::GetClosurePerPixel(Platform, 0));
+			KeyGen.Append(Substrate::GetClosurePerPixel(Platform, 1));
+			KeyGen.Append(Substrate::GetClosurePerPixel(Platform, 2));
+			KeyGen.Append(Substrate::GetClosurePerPixel(Platform, 3));
 		}
 
 		if (Substrate::IsDBufferPassEnabled(Platform))
diff --git a/Engine/Source/Runtime/RenderCore/Public/RenderUtils.h b/Engine/Source/Runtime/RenderCore/Public/RenderUtils.h
--- a/Engine/Source/Runtime/RenderCore/Public/RenderUtils.h
+++ b/Engine/Source/Runtime/RenderCore/Public/RenderUtils.h
@@ -682,8 +682,13 @@ namespace Substrate
 	RENDERCORE_API uint32 GetRayTracingMaterialPayloadSizeInBytes(bool bFullySimplifiedMaterial);
 
 	RENDERCORE_API uint32 GetBytePerPixel(EShaderPlatform InPlatform);
+
+	UE_DEPRECATED(5.8, "GetClosurePerPixel(EShaderPlatform InPlatform, uint8 InMaterialQualityLevel)")
 	RENDERCORE_API uint32 GetClosurePerPixel(EShaderPlatform InPlatform);
 
+	const uint8 IgnoreMaterialQualityLevelsForClosurePerPixel = 255;
+	RENDERCORE_API uint32 GetClosurePerPixel(EShaderPlatform InPlatform, uint8 InMaterialQualityLevel);
+
 	RENDERCORE_API uint32 GetNormalQuality();
 
 	UE_DEPRECATED(5.8, "This function is no longer used. It will always return 1.")
diff --git a/Engine/Source/Runtime/Renderer/Private/Substrate/Substrate.cpp b/Engine/Source/Runtime/Renderer/Private/Substrate/Substrate.cpp
--- a/Engine/Source/Runtime/Renderer/Private/Substrate/Substrate.cpp
+++ b/Engine/Source/Runtime/Renderer/Private/Substrate/Substrate.cpp
@@ -184,7 +184,7 @@ uint32 GetSubstrateMaxClosureCount(const FViewInfo& View)
 		}
 		else
 		{
-			Out = FMath::Clamp(uint32(GetClosurePerPixel(View.GetShaderPlatform())), 1u, SUBSTRATE_MAX_CLOSURE_COUNT);
+			Out = FMath::Clamp(uint32(GetClosurePerPixel(View.GetShaderPlatform(), GetCurrentMaterialQualityLevelChecked())), 1u, SUBSTRATE_MAX_CLOSURE_COUNT);
 		}
 	}
 	return Out;
@@ -528,7 +528,7 @@ void InitialiseSubstrateFrameSceneData(FRDGBuilder& GraphBuilder, FScene& Scene,
 
 		// Material buffer allocation can use different modes:
 		const uint32 PlatformSettingsBytesPerPixel = GetBytePerPixel(ShaderPlatform);
-		const uint32 PlatformSettingsClosurePerPixel = GetClosurePerPixel(ShaderPlatform);
+		const uint32 PlatformSettingsClosurePerPixel = GetClosurePerPixel(ShaderPlatform, GetCurrentMaterialQualityLevelChecked());
 		uint32 CurrentMaxBytesPerPixel = 0;
 		uint32 CurrentMaxClosurePerPixel = 0;
 		switch (GetMaterialBufferAllocationMode())
diff --git a/Engine/Source/Runtime/Renderer/Private/Substrate/SubstrateVisualize.cpp b/Engine/Source/Runtime/Renderer/Private/Substrate/SubstrateVisualize.cpp
--- a/Engine/Source/Runtime/Renderer/Private/Substrate/SubstrateVisualize.cpp
+++ b/Engine/Source/Runtime/Renderer/Private/Substrate/SubstrateVisualize.cpp
@@ -1129,7 +1129,7 @@ static void AddVisualizeSystemInfoPasses(FRDGBuilder& GraphBuilder, const FViewI
 	SystemInformation.ViewsMaxBytesPerPixel = View.SubstrateViewData.SceneData->ViewsMaxBytesPerPixel;
 	SystemInformation.CurrentClosuresPerPixel = View.SubstrateViewData.SceneData->EffectiveMaxClosurePerPixel;
 	SystemInformation.ViewsMaxClosuresPerPixel = View.SubstrateViewData.SceneData->ViewsMaxClosurePerPixel;
-	SystemInformation.ProjectMaxClosuresPerPixel = GetClosurePerPixel(View.GetShaderPlatform());
+	SystemInformation.ProjectMaxClosuresPerPixel = GetClosurePerPixel(View.GetShaderPlatform(), GetCurrentMaterialQualityLevelChecked());
 	SystemInformation.MaterialBufferAllocationMode = GetMaterialBufferAllocationMode();
 	SystemInformation.ShadingQuality = Substrate::GetShadingQuality(View.GetShaderPlatform());
 	SystemInformation.bRoughDiffuse = View.SubstrateViewData.SceneData->bRoughDiffuse;
