[Shaders] Support unmounting individual Pak chunks in the shader code library.

Unmounting shader libraries was previously only supported for entire logical libraries, i.e. for entire plugins.
When an individual container is unmounted, we must no longer attempt to load shaders from that container.

Example scenario:
Container A and B both contain shader X, if A is mounted before B and then A is unmounted, we still attempt to load shader X from A,
because it is the first in order of shader library instances. That library still exists, but the content (i.e. bytecode) for shader X is no longer available from container A (the IoStore container).
We then crash with the "I/O Error (Not Found)" message later on.

The new method of shader chunk override packaging with shadermap deduplication excacerbates this issue when individual pakchunks are unmounted.

#Synced-CL 51046133

diff --git a/Engine/Source/Runtime/RenderCore/Private/ShaderLibrary/ShaderCodeLibrary.cpp b/Engine/Source/Runtime/RenderCore/Private/ShaderLibrary/ShaderCodeLibrary.cpp
--- a/Engine/Source/Runtime/RenderCore/Private/ShaderLibrary/ShaderCodeLibrary.cpp
+++ b/Engine/Source/Runtime/RenderCore/Private/ShaderLibrary/ShaderCodeLibrary.cpp
@@ -80,6 +80,7 @@ namespace UE::ShaderLibrary::Private
 	using ESaveToDiskTarget = ShaderCodeArchive::ESaveToDiskTarget;
 
 	static FDelegateHandle OnPakFileMountedDelegateHandle;
+	static FDelegateHandle OnPakFileUnmountedDelegateHandle;
 	static FDelegateHandle OnPluginMountedDelegateHandle;
 	static FDelegateHandle OnPluginUnmountedDelegateHandle;
 
@@ -234,13 +235,22 @@ namespace UE::ShaderLibrary::Private
 			FShaderLibraryChunkInfo ChunkInfo = GetLibraryChunkInfo(MountInfo);
 			if (!GetPresentChunks().Contains(ChunkInfo.ChunkId))
 			{
-				// Parts of shader library might be in the UFS (e.g. .metallibs), hence we need to look for them in the
-				// appropriate directory.
+				// Parts of shader library might be in the UFS (e.g. .metallibs), hence we need to look for them in the appropriate directory.
 				OpenShaderCode(Directory, ChunkInfo.ChunkLibraryName, ChunkInfo.ChunkId, ChunkInfo.Priority);
 			}
 		}
 
+		void OnPakFileUnmounted(const FMountedPakFileInfo& MountInfo, const FString& Directory)
+		{
+			FShaderLibraryChunkInfo ChunkInfo = GetLibraryChunkInfo(MountInfo);
+			if (GetPresentChunks().Contains(ChunkInfo.ChunkId))
+			{
+				CloseShaderCode(Directory, ChunkInfo.ChunkLibraryName, ChunkInfo.ChunkId);
+			}
+		}
+
 		bool OpenShaderCode(const FString& ShaderCodeDir, FString const& Library, int32 ChunkId, int32 Priority = 0);
+		void CloseShaderCode(const FString& ShaderCodeDir, FString const& Library, int32 ChunkId);
 		void ForEachShaderLibraryWithShaderMap(const FSHAHash& Hash, TFunctionRef<void(FShaderLibraryInstance* InLibraryInstance, int32 InShaderMapIndex)> InPerShaderMapCallback);
 		FShaderLibraryInstance* FindShaderLibraryForShaderMap(const FSHAHash& Hash, int32& OutShaderMapIndex);
 		FShaderLibraryInstance* FindShaderLibraryForShader(const FSHAHash& Hash, int32& OutShaderIndex);
@@ -1233,7 +1243,7 @@ public:
 #endif
 	}
 
-	void OnPakFileMounted(const UE::ShaderLibrary::Private::FMountedPakFileInfo& MountInfo)
+	void PakFileStateChangeInternal(FShaderPipelineCache::ELibraryState InStateChange, TFunctionRef<void(UE::ShaderLibrary::Private::FNamedShaderLibrary* InNamedShaderLibrary)> InStateChangeCallback)
 	{
 		if (IsLibraryInitializedForRuntime())
 		{
@@ -1243,13 +1253,17 @@ public:
 				FRWScopeLock ReadLock(NamedLibrariesMutex, SLT_ReadOnly);
 				for (TTuple<FString, TUniquePtr<UE::ShaderLibrary::Private::FNamedShaderLibrary>>& NamedLibraryPair : NamedLibrariesStack)
 				{
-					TSet<int32> PrevComponentSet = NamedLibraryPair.Value->GetPresentChunks();
-					NamedLibraryPair.Value->OnPakFileMounted(MountInfo, NamedLibraryPair.Value->BaseDirectory);
-					for (int32 ComponentID : NamedLibraryPair.Value->GetPresentChunks().Difference(PrevComponentSet))
+					const TSet<int32> PrevComponentSet = NamedLibraryPair.Value->GetPresentChunks();
+					InStateChangeCallback(NamedLibraryPair.Value.Get());
+					const TSet<int32> NewComponentSet = NamedLibraryPair.Value->GetPresentChunks();
+					const TSet<int32> ChangedComponentSet = NewComponentSet.Num() > PrevComponentSet.Num()
+						? NewComponentSet.Difference(PrevComponentSet)
+						: PrevComponentSet.Difference(NewComponentSet);
+					for (int32 ComponentID : ChangedComponentSet)
 					{
 						// Defer these for outside of the lock, ShaderPipelineCache may want to inspect shader library.
-						ShaderLibraryStateChanges.Add([ShaderPlatform = NamedLibraryPair.Value->ShaderPlatform, LogicalName = NamedLibraryPair.Value->LogicalName, ComponentID]() {
-							FShaderPipelineCache::ShaderLibraryStateChanged(FShaderPipelineCache::OpenedComponent, ShaderPlatform, LogicalName, ComponentID);
+						ShaderLibraryStateChanges.Add([StateChange = InStateChange, ShaderPlatform = NamedLibraryPair.Value->ShaderPlatform, LogicalName = NamedLibraryPair.Value->LogicalName, ComponentID]() {
+							FShaderPipelineCache::ShaderLibraryStateChanged(StateChange, ShaderPlatform, LogicalName, ComponentID);
 						});
 					}
 				}
@@ -1261,6 +1275,20 @@ public:
 		}
 	}
 
+	void OnPakFileMounted(const UE::ShaderLibrary::Private::FMountedPakFileInfo& MountInfo)
+	{
+		PakFileStateChangeInternal(FShaderPipelineCache::OpenedComponent, [&MountInfo](UE::ShaderLibrary::Private::FNamedShaderLibrary* InNamedShaderLibrary) {
+			InNamedShaderLibrary->OnPakFileMounted(MountInfo, InNamedShaderLibrary->BaseDirectory);
+		});
+	}
+
+	void OnPakFileUnmounted(const UE::ShaderLibrary::Private::FMountedPakFileInfo& MountInfo)
+	{
+		PakFileStateChangeInternal(FShaderPipelineCache::Closed, [&MountInfo](UE::ShaderLibrary::Private::FNamedShaderLibrary* InNamedShaderLibrary) {
+			InNamedShaderLibrary->OnPakFileUnmounted(MountInfo, InNamedShaderLibrary->BaseDirectory);
+		});
+	}
+
 	uint32 GetShaderCount(void)
 	{
 		FRWScopeLock ReadLock(NamedLibrariesMutex, SLT_ReadOnly);
@@ -2141,7 +2169,6 @@ static void FShaderLibraryPakFileMountedCallback(const IPakFile& PakFile)
 {
 	using namespace UE::ShaderLibrary::Private;
 
-	int32 NewChunk = PakFile.PakGetPakchunkIndex();
 	UE_LOG(LogShaderLibrary, Verbose, TEXT("ShaderCodeLibraryPakFileMountedCallback: PakFile '%s' (chunk index %d, root '%s') mounted"), *PakFile.PakGetPakFilename(), PakFile.PakGetPakchunkIndex(), *PakFile.PakGetMountPoint());
 
 	FMountedPakFileInfo PakFileInfo(PakFile);
@@ -2150,7 +2177,7 @@ static void FShaderLibraryPakFileMountedCallback(const IPakFile& PakFile)
 		FMountedPakFileInfo::KnownPakFiles.Add(PakFileInfo);
 	}
 
-	// if shaderlibrary has not yet been initialized, add the chunk as pending
+	// If shaderlibrary has not yet been initialized, report the chunk as pending
 	if (FShaderLibrariesCollection::Impl)
 	{
 		FShaderLibrariesCollection::Impl->OnPakFileMounted(PakFileInfo);
@@ -2161,10 +2188,34 @@ static void FShaderLibraryPakFileMountedCallback(const IPakFile& PakFile)
 	}
 }
 
+static void FShaderLibraryPakFileUnmountedCallback(const IPakFile& PakFile)
+{
+	using namespace UE::ShaderLibrary::Private;
+
+	UE_LOG(LogShaderLibrary, Verbose, TEXT("ShaderCodeLibraryPakFileUnmountedCallback: PakFile '%s' (chunk index %d, root '%s') unmounted"), *PakFile.PakGetPakFilename(), PakFile.PakGetPakchunkIndex(), *PakFile.PakGetMountPoint());
+
+	FMountedPakFileInfo PakFileInfo(PakFile);
+	{
+		FScopeLock PakFilesLocker(&FMountedPakFileInfo::KnownPakFilesAccessLock);
+		FMountedPakFileInfo::KnownPakFiles.Remove(PakFileInfo);
+	}
+
+	// If shaderlibrary has not yet been initialized, report the chunk as pending
+	if (FShaderLibrariesCollection::Impl)
+	{
+		FShaderLibrariesCollection::Impl->OnPakFileUnmounted(PakFileInfo);
+	}
+	else
+	{
+		UE_LOG(LogShaderLibrary, Verbose, TEXT("ShaderCodeLibraryPakFileUnmountedCallback: pending pak file info (%s)"), *PakFileInfo.ToString());
+	}
+}
+
 void FShaderCodeLibrary::PreInit()
 {
 	// add a callback for opening later chunks
 	UE::ShaderLibrary::Private::OnPakFileMountedDelegateHandle = FCoreDelegates::GetOnPakFileMounted2().AddStatic(&FShaderLibraryPakFileMountedCallback);
+	UE::ShaderLibrary::Private::OnPakFileUnmountedDelegateHandle = FCoreDelegates::GetOnPakFileUnmounted().AddStatic(&FShaderLibraryPakFileUnmountedCallback);
 }
 
 void FShaderCodeLibrary::InitForRuntime(EShaderPlatform ShaderPlatform)
@@ -2238,6 +2289,11 @@ void FShaderCodeLibrary::Shutdown()
 		FCoreDelegates::GetOnPakFileMounted2().Remove(UE::ShaderLibrary::Private::OnPakFileMountedDelegateHandle);
 		UE::ShaderLibrary::Private::OnPakFileMountedDelegateHandle.Reset();
 	}
+	if (UE::ShaderLibrary::Private::OnPakFileUnmountedDelegateHandle.IsValid())
+	{
+		FCoreDelegates::GetOnPakFileUnmounted().Remove(UE::ShaderLibrary::Private::OnPakFileUnmountedDelegateHandle);
+		UE::ShaderLibrary::Private::OnPakFileUnmountedDelegateHandle.Reset();
+	}
 	if (UE::ShaderLibrary::Private::OnPluginMountedDelegateHandle.IsValid())
 	{
 		IPluginManager::Get().OnNewPluginMounted().Remove(UE::ShaderLibrary::Private::OnPluginMountedDelegateHandle);
@@ -2857,13 +2913,33 @@ void FShaderCodeLibrary::OpenPluginShaderLibrary(IPlugin& Plugin, bool bMonolith
 
 // FNamedShaderLibrary methods
 
+static void LogShaderLibraryInstanceState(const FShaderLibraryInstance* InLibraryInstance, const TCHAR* InLibraryName, const TCHAR* InStateInfo)
+{
+	UE_LOG(LogShaderLibrary, Display, TEXT("Cooked Context: %s %s Shader Library %s"),
+		InStateInfo, InLibraryInstance->Library->IsNativeLibrary() ? TEXT("Native Shared") : TEXT("Shared"), InLibraryName);
+}
+
+static FString PrintComponentChunkIds(const TConstArrayView<FShaderLibraryInstancePriorityPair>& InComponents)
+{
+	TStringBuilder<1024> OutString;
+	for (const FShaderLibraryInstancePriorityPair& Component : InComponents)
+	{
+		if (OutString.Len() > 0)
+		{
+			OutString.Append(TEXT(", "));
+		}
+		OutString.Appendf(TEXT("%d"), Component.Value->GetChunkId());
+	}
+	return FString(OutString);
+}
+
 // At runtime, open shader code collection for specified shader platform. Returns true if new code was opened
 bool UE::ShaderLibrary::Private::FNamedShaderLibrary::OpenShaderCode(const FString& ShaderCodeDir, FString const& Library, int32 ChunkId, int32 Priority)
 {
 	LLM_SCOPE(ELLMTag::Shaders);
 	// check if any of the components has this content
 	{
-		FRWScopeLock ReadLock(ComponentsMutex, SLT_Write);
+		FRWScopeLock WriteLock(ComponentsMutex, SLT_Write);
 		for (const FShaderLibraryInstancePriorityPair& Component : Components)
 		{
 			if (Component.Value->HasContentFrom(ShaderCodeDir, Library))
@@ -2901,14 +2977,7 @@ bool UE::ShaderLibrary::Private::FNamedShaderLibrary::OpenShaderCode(const FStri
 		}
 	}
 
-	if (LibraryInstance->Library->IsNativeLibrary())
-	{
-		UE_LOG(LogShaderLibrary, Display, TEXT("Cooked Context: Loaded Native Shared Shader Library %s"), *Library);
-	}
-	else
-	{
-		UE_LOG(LogShaderLibrary, Display, TEXT("Cooked Context: Using Shared Shader Library %s"), *Library);
-	}
+	LogShaderLibraryInstanceState(LibraryInstance, *Library, TEXT("Loaded"));
 
 	// Insertion sort of new component by their priority.
 	// This is mandatory because some shaders can be grouped across different IoChunks, some of which can be packaged into an IAD and a non-IAD container respectively.
@@ -2925,23 +2994,34 @@ bool UE::ShaderLibrary::Private::FNamedShaderLibrary::OpenShaderCode(const FStri
 	}
 
 	// Log new set of component IDs if verbose logging is enabled
-	auto PrintComponentChunkIds = [](const TConstArrayView<FShaderLibraryInstancePriorityPair>& InComponents) -> FString
+	UE_LOG(LogShaderLibrary, Verbose, TEXT("OpenShaderCode('%s'): Components=(%s)"), *Library, *PrintComponentChunkIds(Components));
+
+	return true;
+}
+
+void UE::ShaderLibrary::Private::FNamedShaderLibrary::CloseShaderCode(const FString& ShaderCodeDir, FString const& Library, int32 ChunkId)
+{
+	LLM_SCOPE(ELLMTag::Shaders);
+
+	FRWScopeLock WriteLock(ComponentsMutex, SLT_Write);
+	for (int32 Index = 0; Index < Components.Num(); ++Index)
+	{
+		if (const FShaderLibraryInstance* LibraryInstance = Components[Index].Value.Get())
 		{
-			TStringBuilder<1024> OutString;
-			for (const FShaderLibraryInstancePriorityPair& Component : InComponents)
+			if (LibraryInstance->Library->GetName() == Library && LibraryInstance->GetChunkId() == ChunkId)
 			{
-				if (OutString.Len() > 0)
-				{
-					OutString.Append(TEXT(", "));
-				}
-				OutString.Appendf(TEXT("%d"), Component.Value->GetChunkId());
-			}
-			return FString(OutString);
-		};
+				LogShaderLibraryInstanceState(LibraryInstance, *Library, TEXT("Released"));
 
-	UE_LOG(LogShaderLibrary, Verbose, TEXT("OpenShaderCode('%s'): Components=(%s)"), *Library, *PrintComponentChunkIds(Components));
+				// Remove component. Ensure we preserve the order since non-IAD components need to be iterated before IAD components.
+				Components.RemoveAt(Index);
 
-	return true;
+				// Log new set of component IDs if verbose logging is enabled
+				UE_LOG(LogShaderLibrary, Verbose, TEXT("CloseShaderCode('%s'): Components=(%s)"), *Library, *PrintComponentChunkIds(Components));
+
+				return;
+			}
+		}
+	}
 }
 
 void UE::ShaderLibrary::Private::FNamedShaderLibrary::ForEachShaderLibraryWithShaderMap(const FSHAHash& Hash, TFunctionRef<void(FShaderLibraryInstance* InLibraryInstance, int32 InShaderMapIndex)> InPerShaderMapCallback)
