RHI: Fix comment with a simpler example (PF_G8 actually uses R)

#rnx

#Synced-CL 51094475

#Synced-CL 51094475

diff --git a/Engine/Source/Runtime/RenderCore/Public/RenderUtils.h b/Engine/Source/Runtime/RenderCore/Public/RenderUtils.h
--- a/Engine/Source/Runtime/RenderCore/Public/RenderUtils.h
+++ b/Engine/Source/Runtime/RenderCore/Public/RenderUtils.h
@@ -236,7 +236,7 @@ RENDERCORE_API void CopyTextureData2D(const void* Source,void* Dest,uint32 SizeY
 /**
  *  Returns the valid channels for this pixel format
  * 
- * @return e.g. EPixelFormatChannelFlags::G for PF_G8
+ * @return e.g. EPixelFormatChannelFlags::R for PF_R8
  */
 RENDERCORE_API EPixelFormatChannelFlags GetPixelFormatValidChannels(EPixelFormat InPixelFormat);
 
