[shaders] remove warnings from DDC records; the only mechanism that uses these is to print warnings on job cache hits or shadermap load (optionally, disabled by default), both I think are of dubious value, and this is a significant amount of DDC data (which also tends to continually cause non-determinism issues).

#Synced-CL 51095924

diff --git a/Engine/Shaders/Public/ShaderVersion.ush b/Engine/Shaders/Public/ShaderVersion.ush
--- a/Engine/Shaders/Public/ShaderVersion.ush
+++ b/Engine/Shaders/Public/ShaderVersion.ush
@@ -3,4 +3,4 @@
 // Each time the console command is executed it generates a new GUID. As this file is included
 // in Platform.ush (which should be included in any shader) it allows to invalidate the shader DDC.
 // 
-#pragma message("UESHADERMETADATA_VERSION 626F9AAB-3832-46E1-8F46-64B302F153D6")
\ No newline at end of file
+#pragma message("UESHADERMETADATA_VERSION 76323D69-1266-4061-B1C3-291C928683AE")
\ No newline at end of file
diff --git a/Engine/Source/Runtime/RenderCore/Private/ShaderResource.cpp b/Engine/Source/Runtime/RenderCore/Private/ShaderResource.cpp
--- a/Engine/Source/Runtime/RenderCore/Private/ShaderResource.cpp
+++ b/Engine/Source/Runtime/RenderCore/Private/ShaderResource.cpp
@@ -30,20 +30,10 @@
 #include "Interfaces/IShaderFormat.h"
 #endif
 
-DECLARE_LOG_CATEGORY_CLASS(LogShaderWarnings, Log, Log);
-
 #if (CSV_PROFILER_STATS && !UE_BUILD_SHIPPING) 
 TCsvPersistentCustomStat<int>* CsvStatNumShaderMapsUsedForRendering = nullptr;
 #endif
 
-static int32 GShaderCompilerEmitWarningsOnLoad = 0;
-static FAutoConsoleVariableRef CVarShaderCompilerEmitWarningsOnLoad(
-	TEXT("r.ShaderCompiler.EmitWarningsOnLoad"),
-	GShaderCompilerEmitWarningsOnLoad,
-	TEXT("When 1, shader compiler warnings are emitted to the log for all shaders as they are loaded."),
-	ECVF_Default
-);
-
 FName GetShaderCompressionFormat()
 {
 	// We always use oodle now. This was instituted because UnrealPak recompresses the shaders and doens't have
@@ -236,10 +226,6 @@ void FShaderMapResourceCode::Finalize()
 	Hasher.Final();
 	Hasher.GetHash(ResourceHash.Hash);
 	ApplyResourceStats(*this);
-
-#if WITH_EDITORONLY_DATA
-	LogShaderCompilerWarnings();
-#endif
 }
 
 uint32 FShaderMapResourceCode::GetSizeBytes() const
@@ -292,7 +278,7 @@ void FShaderMapResourceCode::AddShaderCompilerOutput(const FShaderCompilerOutput
 
 #if WITH_EDITORONLY_DATA
 		// Output.Errors contains warnings in the case any exist (no errors since if there were the job would have failed)
-		AddEditorOnlyData(Index, DebugName, ShaderTypeHash, CodeResource.GetSymbolsBuffer(), Output.Errors, MoveTemp(ShaderStatsSorted), DebugInfo);
+		AddEditorOnlyData(Index, DebugName, ShaderTypeHash, CodeResource.GetSymbolsBuffer(), MoveTemp(ShaderStatsSorted), DebugInfo);
 #endif
 
 		ShaderCodeResources.Insert(CodeResource, Index);
@@ -303,14 +289,14 @@ void FShaderMapResourceCode::AddShaderCompilerOutput(const FShaderCompilerOutput
 		// Output.Errors contains warnings in the case any exist (no errors since if there were the job would have failed)
 		// We append the warnings and deduplicate other data like DebugInfo for any additional jobs which resulted in the
 		// same bytecode for the sake of determinism in the results saved to DDC.
-		UpdateEditorOnlyData(Index, DebugName, ShaderTypeHash, CodeResource.GetSymbolsBuffer(), Output.Errors, DebugInfo);
+		UpdateEditorOnlyData(Index, DebugName, ShaderTypeHash, CodeResource.GetSymbolsBuffer(), DebugInfo);
 		ValidateShaderStatisticsEditorOnlyData(Index, ShaderStatsSorted);
 	}
 #endif
 }
 
 #if WITH_EDITORONLY_DATA
-void FShaderMapResourceCode::AddEditorOnlyData(int32 Index, const FString& DebugName, uint64 ShaderTypeHash, FCompressedBuffer InSymbols, TConstArrayView<FShaderCompilerError> InCompilerWarnings, TArray<FGenericShaderStat>&& ShaderStatistics, const FString& DebugInfo)
+void FShaderMapResourceCode::AddEditorOnlyData(int32 Index, const FString& DebugName, uint64 ShaderTypeHash, FCompressedBuffer InSymbols, TArray<FGenericShaderStat>&& ShaderStatistics, const FString& DebugInfo)
 {
 	FShaderEditorOnlyDataEntry& Entry = ShaderEditorOnlyDataEntries.InsertDefaulted_GetRef(Index);
 
@@ -318,10 +304,10 @@ void FShaderMapResourceCode::AddEditorOnlyData(int32 Index, const FString& Debug
 	check(Entry.ShaderStatistics.Num() == 0);
 	Entry.ShaderStatistics = MoveTemp(ShaderStatistics);
 
-	UpdateEditorOnlyData(Index, DebugName, ShaderTypeHash, InSymbols, InCompilerWarnings, DebugInfo);
+	UpdateEditorOnlyData(Index, DebugName, ShaderTypeHash, InSymbols, DebugInfo);
 }
 
-void FShaderMapResourceCode::UpdateEditorOnlyData(int32 Index, const FString& DebugName, uint64 ShaderTypeHash, FCompressedBuffer InSymbols, TConstArrayView<FShaderCompilerError> InCompilerWarnings, const FString& DebugInfo)
+void FShaderMapResourceCode::UpdateEditorOnlyData(int32 Index, const FString& DebugName, uint64 ShaderTypeHash, FCompressedBuffer InSymbols, const FString& DebugInfo)
 {
 	FShaderEditorOnlyDataEntry& Entry = ShaderEditorOnlyDataEntries[Index];
 
@@ -333,17 +319,6 @@ void FShaderMapResourceCode::UpdateEditorOnlyData(int32 Index, const FString& De
 
 	Entry.ConditionalSetSymbolBuffer(InSymbols);
 
-	for (const FShaderCompilerError& Warning : InCompilerWarnings)
-	{
-		FString ModifiedWarning = !DebugName.IsEmpty() ? FString::Printf(TEXT("%s [%s]"), *Warning.GetErrorString(), *DebugName) : Warning.GetErrorString();
-		// Maintain sorted order in Entry.CompilerWarnings & deduplicate
-		const int32 WarningIndex = Algo::LowerBound(Entry.CompilerWarnings, ModifiedWarning);
-		if (WarningIndex >= Entry.CompilerWarnings.Num() || Entry.CompilerWarnings[WarningIndex] != ModifiedWarning)
-		{
-			Entry.CompilerWarnings.Insert(ModifiedWarning, WarningIndex);
-		}
-	}
-	
 	const int32 HashIndex = Algo::LowerBound(Entry.ShaderTypeHashes, ShaderTypeHash);
 	if (HashIndex >= Entry.ShaderTypeHashes.Num() || Entry.ShaderTypeHashes[HashIndex] != ShaderTypeHash)
 	{
@@ -385,22 +360,6 @@ void FShaderMapResourceCode::ValidateShaderStatisticsEditorOnlyData(int32 Index,
 	}
 }
 
-void FShaderMapResourceCode::LogShaderCompilerWarnings()
-{
-	if (ShaderEditorOnlyDataEntries.Num() > 0 && GShaderCompilerEmitWarningsOnLoad != 0)
-	{
-		// Emit all the compiler warnings seen whilst serializing/loading this shader to the log.
-		// Since successfully compiled shaders are stored in the DDC, we'll get the compiler warnings
-		// even if we didn't compile the shader this run.
-		for (const FShaderEditorOnlyDataEntry& Entry : ShaderEditorOnlyDataEntries)
-		{
-			for (const FString& CompilerWarning : Entry.CompilerWarnings)
-			{
-				UE_LOG(LogShaderWarnings, Warning, TEXT("%s"), *CompilerWarning);
-			}
-		}
-	}
-}
 #endif // WITH_EDITORONLY_DATA
 
 void FShaderMapResourceCode::ToString(FStringBuilderBase& OutString) const
@@ -456,13 +415,6 @@ void FShaderMapResourceCode::Serialize(FShaderSerializeContext& Ctx)
 	}
 #endif // WITH_EDITORONLY_DATA
 	ApplyResourceStats(*this);
-
-#if WITH_EDITORONLY_DATA
-	if (Ar.IsLoading())
-	{
-		LogShaderCompilerWarnings();
-	}
-#endif
 }
 
 #if WITH_EDITORONLY_DATA
diff --git a/Engine/Source/Runtime/RenderCore/Public/Shader.h b/Engine/Source/Runtime/RenderCore/Public/Shader.h
--- a/Engine/Source/Runtime/RenderCore/Public/Shader.h
+++ b/Engine/Source/Runtime/RenderCore/Public/Shader.h
@@ -518,10 +518,7 @@ public:
 		*   information is available */
 		FString DebugInfo;
 
-		/** A (deduplicated/sorted) array of all the compiler warnings that were emitted when all shaders resulting 
-		 *  in the associated bytecode were compiled (i.e. if multiple shader sources have warnings but compile to
-		 *  the same code, all warnings for each unique source will be reported).
-		 *  Does not contain errors since if there were any errors, this object wouldn't exist. */
+		UE_DEPRECATED(5.8, "Warnings are no longer serialized with shadermaps to DDC due to significant overhead. If you need to see shader compile warnings, please enable shader dev mode and force a recompile.")
 		TArray<FString> CompilerWarnings;
 		
 		/** An array of the hash of all FShaderType names which generate this particular shader (i.e. the return value
@@ -533,7 +530,7 @@ public:
 
 		friend FArchive& operator<<(FArchive& Ar, FShaderEditorOnlyDataEntry& Entry)
 		{
-			return Ar << Entry.DebugInfo << Entry.CompilerWarnings << Entry.ShaderTypeHashes << Entry.ShaderStatistics;
+			return Ar << Entry.DebugInfo << Entry.ShaderTypeHashes << Entry.ShaderStatistics;
 		}
 
 		PRAGMA_DISABLE_DEPRECATION_WARNINGS
@@ -566,10 +563,11 @@ public:
 	int32 FindShaderIndex(const FSHAHash& InHash) const;
 
 #if WITH_EDITORONLY_DATA
-	void AddEditorOnlyData(int32 Index, const FString& DebugName, uint64 ShaderTypeHash, FCompressedBuffer InSymbols, TConstArrayView<FShaderCompilerError> InCompilerWarnings, TArray<FGenericShaderStat>&& ShaderStatistics, const FString& DebugInfo = FString());
-	void UpdateEditorOnlyData(int32 Index, const FString& DebugName, uint64 ShaderTypeHash, FCompressedBuffer InSymbols, TConstArrayView<FShaderCompilerError> InCompilerWarnings, const FString& DebugInfo);
+	void AddEditorOnlyData(int32 Index, const FString& DebugName, uint64 ShaderTypeHash, FCompressedBuffer InSymbols, TArray<FGenericShaderStat>&& ShaderStatistics, const FString& DebugInfo = FString());
+	void UpdateEditorOnlyData(int32 Index, const FString& DebugName, uint64 ShaderTypeHash, FCompressedBuffer InSymbols, const FString& DebugInfo);
 
-	RENDERCORE_API void LogShaderCompilerWarnings();
+	UE_DEPRECATED(5.8, "Shader compiler warnings are no longer stored with shadermap data, and so cannot be logged via this API; please force a recompile of any shaders you wish to see warnings for.")
+	inline void LogShaderCompilerWarnings() {}
 #endif
 
 	RENDERCORE_API void ToString(FStringBuilderBase& OutString) const;
diff --git a/Engine/Source/Runtime/RenderCore/Public/ShaderCompilerCore.h b/Engine/Source/Runtime/RenderCore/Public/ShaderCompilerCore.h
--- a/Engine/Source/Runtime/RenderCore/Public/ShaderCompilerCore.h
+++ b/Engine/Source/Runtime/RenderCore/Public/ShaderCompilerCore.h
@@ -33,7 +33,7 @@ typedef TSharedPtr<TArray<ANSICHAR>, ESPMode::ThreadSafe> FShaderSharedAnsiStrin
 // this is for the protocol, not the data, bump if FShaderCompilerInput/FShaderPreprocessOutput serialization, SerializeWorkerInput or ProcessInputFromArchive changes.
 inline const int32 ShaderCompileWorkerInputVersion = 33;
 // this is for the protocol, not the data, bump if FShaderCompilerOutput or WriteToOutputArchive changes.
-inline const int32 ShaderCompileWorkerOutputVersion = 30;
+inline const int32 ShaderCompileWorkerOutputVersion = 31;
 // this is for the protocol, not the data.
 inline const int32 ShaderCompileWorkerSingleJobHeader = 'S';
 // this is for the protocol, not the data.
@@ -624,12 +624,12 @@ struct FShaderCompilerOutput
 	{
 		// Note: this serialize is used to pass between UE and the shader compile worker, recompile both when modifying
 		Ar << Output.ParameterMap;
-		Ar << Output.Errors;
 		Ar << Output.Target;
 		Ar << Output.bSerializingForCache;
 		if (!Output.bSerializingForCache)
 		{
 			// skip serializing these fields when saving to cache/DDC; only needed when reading back results from workers
+			Ar << Output.Errors;
 			Ar << Output.ShaderCode;
 			Ar << Output.ValidateInputHash;
 			Ar << Output.CompileTime;
