FIoStoreShaderCodeArchive - Implement Memory Mapping with fallback to IoRequest
- Add a new packaging option to disable shader group compression ([/Script/UnrealEd.ProjectPackagingSettings] bDisallowShaderGroupCompression)
- Modify RHICreate{ShaderType}Shader family of functions to use FShaderCodeBuffer that allows to optionally share memory ownership between FRHIShaderLibrary / FRHIShader and update all existing RHIs to match the new signatures and remove unused FSHAHash parameter.
- MetalBaseShader: Add optional shared memory ownership and clean up dead code related to symbols & source code

All new behaviour is disabled by default on all platforms.
To enable on any platforms that support memory mapping: bDisallowShaderGroupCompression + USE_MMAPPED_SHADERARCHIVE
In addition, to enable on Apple platforms: bSharedMaterialNativeLibraries=False

Note: Currently, only Apple platforms use shared memory ownership, which is designed with memory mapping in mind, but also prevents memory allocation & copy without it.

#tests Preflight & local tests

#Synced-CL 51263606

diff --git a/Engine/Source/Developer/IoStoreUtilities/Private/IoStoreShaderLibraryProcessor.cpp b/Engine/Source/Developer/IoStoreUtilities/Private/IoStoreShaderLibraryProcessor.cpp
--- a/Engine/Source/Developer/IoStoreUtilities/Private/IoStoreShaderLibraryProcessor.cpp
+++ b/Engine/Source/Developer/IoStoreUtilities/Private/IoStoreShaderLibraryProcessor.cpp
@@ -17,7 +17,6 @@
 #include "ShaderCodeArchive.h"
 #include "ShaderLibrary/EditorShaderCodeArchive.h"
 
-
 #define IOSTORE_CPU_SCOPE(NAME) TRACE_CPUPROFILER_EVENT_SCOPE(IoStore##NAME);
 #define IOSTORE_CPU_SCOPE_DATA(NAME, DATA) TRACE_CPUPROFILER_EVENT_SCOPE(IoStore##NAME);
 
@@ -85,6 +84,22 @@ static bool IsChunkOverrideDebugEnabled()
 	return CVarShaderLibraryPakChunkOverrideDumpDebugInfo.GetValueOnAnyThread();
 }
 
+static bool ShouldCompressShaderGroups()
+{
+	FString PlatformIniName;
+	if (FParse::Value(FCommandLine::Get(), TEXT("platform="), PlatformIniName, false))
+	{
+		FConfigFile PlatformIniFile;
+		FConfigCacheIni::LoadLocalIniFile(PlatformIniFile, TEXT("Game"), true, *PlatformIniName);
+		bool bDisallowShaderGroupCompression = false;
+		if (PlatformIniFile.GetBool(TEXT("/Script/UnrealEd.ProjectPackagingSettings"), TEXT("bDisallowShaderGroupCompression"), bDisallowShaderGroupCompression))
+		{
+			return !bDisallowShaderGroupCompression;
+		}
+	}
+	return true;
+}
+
 // Loads an *.assetinfo.json text file either from local file system or Zen store.
 static bool LoadShaderAssetInfo(const FString& Filename, FCookedPackageStore* PackageStore, FShaderMapAssetAssociations& OutShaderCodeToAssets)
 {
@@ -529,19 +544,26 @@ bool FIoStoreShaderLibraryProcessor::ConvertToIoStoreShaderLibrary(FContainerTar
 	// This also updates group's compressed size, the only value left to calculate.
 	const int32 GroupCount = IoStoreLibraryHeader.ShaderGroupEntries.Num();
 	CodeIoChunks.SetNum(GroupCount);
+	
+	const bool bCompressShaderGroup = ShouldCompressShaderGroups();
 
 	ParallelFor(CodeIoChunks.Num(),
-		[this, &IoStoreLibraryHeader, &SerializedShaders](int32 GroupIndex)
+		[this, &IoStoreLibraryHeader, &SerializedShaders, bCompressShaderGroup](int32 GroupIndex)
 		{
 			FIoStoreShaderGroupEntry& Group = IoStoreLibraryHeader.ShaderGroupEntries[GroupIndex];
 
 			// Decompress all individual shaders of the current group
 			int32 SmallestShaderInGroupByOrdinal = MAX_int32;
 			const TArray<uint8> UncompressedGroupMemory = DecompressGroupedShaders(IoStoreLibraryHeader, Group, SerializedShaders, SmallestShaderInGroupByOrdinal);
-
-			// Now compress the whole group
-			const TArray<uint8> CompressedGroupMemory = CompressShaderGroup(UncompressedGroupMemory.GetData(), Group, this->ShaderOodleCompressor, this->ShaderOodleLevel);
-			const uint32 CompressedGroupMemorySize = static_cast<uint32>(CompressedGroupMemory.Num());
+		
+			TArray<uint8> CompressedGroupMemory;
+			uint32 CompressedGroupMemorySize = 0;
+			if (bCompressShaderGroup)
+			{
+				// Now compress the whole group
+				CompressedGroupMemory = CompressShaderGroup(UncompressedGroupMemory.GetData(), Group, this->ShaderOodleCompressor, this->ShaderOodleLevel);
+				CompressedGroupMemorySize = static_cast<uint32>(CompressedGroupMemory.Num());
+			}
 
 			// Write out IoChunk for current shader group
 			{
@@ -554,8 +576,9 @@ bool FIoStoreShaderLibraryProcessor::ConvertToIoStoreShaderLibrary(FContainerTar
 					SmallestShaderInGroupByOrdinal, IoStoreLibraryHeader.ShaderEntries.Num());
 				OutCodeIoChunk.LoadOrderFactor = static_cast<uint32>(SmallestShaderInGroupByOrdinal);
 
+
 				// Copy shader group payload into IoBuffer
-				if (CompressedGroupMemorySize < Group.UncompressedSize)
+				if (bCompressShaderGroup && CompressedGroupMemorySize < Group.UncompressedSize)
 				{
 					OutCodeIoChunk.CodeIoBuffer = FIoBuffer(CompressedGroupMemorySize);
 					FMemory::Memcpy(OutCodeIoChunk.CodeIoBuffer.GetData(), CompressedGroupMemory.GetData(), CompressedGroupMemorySize);
diff --git a/Engine/Source/Runtime/Apple/MetalRHI/Private/MetalDynamicRHI.h b/Engine/Source/Runtime/Apple/MetalRHI/Private/MetalDynamicRHI.h
--- a/Engine/Source/Runtime/Apple/MetalRHI/Private/MetalDynamicRHI.h
+++ b/Engine/Source/Runtime/Apple/MetalRHI/Private/MetalDynamicRHI.h
@@ -113,13 +113,13 @@ public:
 	virtual FDepthStencilStateRHIRef RHICreateDepthStencilState(const FDepthStencilStateInitializerRHI& Initializer) final override;
 	virtual FBlendStateRHIRef RHICreateBlendState(const FBlendStateInitializerRHI& Initializer) final override;
 	virtual FVertexDeclarationRHIRef RHICreateVertexDeclaration(const FVertexDeclarationElementList& Elements) final override;
-	virtual FPixelShaderRHIRef RHICreatePixelShader(TArrayView<const uint8> Code, const FSHAHash& Hash) final override;
-	virtual FVertexShaderRHIRef RHICreateVertexShader(TArrayView<const uint8> Code, const FSHAHash& Hash) final override;
-	virtual FGeometryShaderRHIRef RHICreateGeometryShader(TArrayView<const uint8> Code, const FSHAHash& Hash) final override;
-	virtual FComputeShaderRHIRef RHICreateComputeShader(TArrayView<const uint8> Code, const FSHAHash& Hash) final override;
+	virtual FPixelShaderRHIRef RHICreatePixelShader(const FRHICreateShaderDesc& CreateShaderDesc) final override;
+	virtual FVertexShaderRHIRef RHICreateVertexShader(const FRHICreateShaderDesc& CreateShaderDesc) final override;
+	virtual FGeometryShaderRHIRef RHICreateGeometryShader(const FRHICreateShaderDesc& CreateShaderDesc) final override;
+	virtual FComputeShaderRHIRef RHICreateComputeShader(const FRHICreateShaderDesc& CreateShaderDesc) final override;
 #if PLATFORM_SUPPORTS_MESH_SHADERS
-    virtual FMeshShaderRHIRef RHICreateMeshShader(TArrayView<const uint8> Code, const FSHAHash& Hash) final override;
-    virtual FAmplificationShaderRHIRef RHICreateAmplificationShader(TArrayView<const uint8> Code, const FSHAHash& Hash) final override;
+    virtual FMeshShaderRHIRef RHICreateMeshShader(const FRHICreateShaderDesc& CreateShaderDesc) final override;
+    virtual FAmplificationShaderRHIRef RHICreateAmplificationShader(const FRHICreateShaderDesc& CreateShaderDesc) final override;
 #endif
 	virtual FBoundShaderStateRHIRef RHICreateBoundShaderState(FRHIVertexDeclaration* VertexDeclaration, FRHIVertexShader* VertexShader, FRHIPixelShader* PixelShader, FRHIGeometryShader* GeometryShader) final override;
 	virtual FGraphicsPipelineStateRHIRef RHICreateGraphicsPipelineState(const FGraphicsPipelineStateInitializer& Initializer) final override;
@@ -224,7 +224,7 @@ public:
 
 	virtual FRayTracingGeometryRHIRef RHICreateRayTracingGeometry(FRHICommandListBase& RHICmdList, const FRayTracingGeometryInitializer& Initializer) final override;
 	virtual FRayTracingSceneRHIRef RHICreateRayTracingScene(FRayTracingSceneInitializer Initializer) final override;
-	virtual FRayTracingShaderRHIRef RHICreateRayTracingShader(TArrayView<const uint8> Code, const FSHAHash& Hash, EShaderFrequency ShaderFrequency) final override;
+	virtual FRayTracingShaderRHIRef RHICreateRayTracingShader(const FRHICreateShaderDesc& CreateShaderDesc, EShaderFrequency ShaderFrequency) final override;
 	virtual FRayTracingPipelineStateRHIRef RHICreateRayTracingPipelineState(const FRayTracingPipelineStateInitializer& Initializer) final override;
 
 	virtual FShaderBindingTableRHIRef RHICreateShaderBindingTable(FRHICommandListBase& RHICmdList, const FRayTracingShaderBindingTableInitializer& Initializer) final override;
diff --git a/Engine/Source/Runtime/Apple/MetalRHI/Private/MetalDynamicRHI_Shaders.cpp b/Engine/Source/Runtime/Apple/MetalRHI/Private/MetalDynamicRHI_Shaders.cpp
--- a/Engine/Source/Runtime/Apple/MetalRHI/Private/MetalDynamicRHI_Shaders.cpp
+++ b/Engine/Source/Runtime/Apple/MetalRHI/Private/MetalDynamicRHI_Shaders.cpp
@@ -18,61 +18,61 @@
 #pragma mark - Metal Dynamic RHI Shader Methods
 
 
-FVertexShaderRHIRef FMetalDynamicRHI::RHICreateVertexShader(TArrayView<const uint8> Code, const FSHAHash& Hash)
+FVertexShaderRHIRef FMetalDynamicRHI::RHICreateVertexShader(const FRHICreateShaderDesc& CreateShaderDesc)
 {
     MTL_SCOPED_AUTORELEASE_POOL;
     
-	FMetalVertexShader* Shader = new FMetalVertexShader(*Device, Code);
+	FMetalVertexShader* Shader = new FMetalVertexShader(*Device, CreateShaderDesc);
 	return Shader;
 }
 
-FPixelShaderRHIRef FMetalDynamicRHI::RHICreatePixelShader(TArrayView<const uint8> Code, const FSHAHash& Hash)
+FPixelShaderRHIRef FMetalDynamicRHI::RHICreatePixelShader(const FRHICreateShaderDesc& CreateShaderDesc)
 {
     MTL_SCOPED_AUTORELEASE_POOL;
     
-	FMetalPixelShader* Shader = new FMetalPixelShader(*Device, Code);
+	FMetalPixelShader* Shader = new FMetalPixelShader(*Device, CreateShaderDesc);
 	return Shader;
 }
 
-FGeometryShaderRHIRef FMetalDynamicRHI::RHICreateGeometryShader(TArrayView<const uint8> Code, const FSHAHash& Hash)
+FGeometryShaderRHIRef FMetalDynamicRHI::RHICreateGeometryShader(const FRHICreateShaderDesc& CreateShaderDesc)
 {
     MTL_SCOPED_AUTORELEASE_POOL;
 #if PLATFORM_SUPPORTS_GEOMETRY_SHADERS
-    FMetalGeometryShader* Shader = new FMetalGeometryShader(*Device, Code);
+    FMetalGeometryShader* Shader = new FMetalGeometryShader(*Device, CreateShaderDesc);
     return Shader;
 #else
 	return nullptr;
 #endif
 }
 
-FComputeShaderRHIRef FMetalDynamicRHI::RHICreateComputeShader(TArrayView<const uint8> Code, const FSHAHash& Hash)
+FComputeShaderRHIRef FMetalDynamicRHI::RHICreateComputeShader(const FRHICreateShaderDesc& CreateShaderDesc)
 {
     MTL_SCOPED_AUTORELEASE_POOL;
-    return new FMetalComputeShader(*Device, Code, MTLLibraryPtr());
+    return new FMetalComputeShader(*Device, CreateShaderDesc, MTLLibraryPtr());
 }
 
 #if PLATFORM_SUPPORTS_MESH_SHADERS
-FMeshShaderRHIRef FMetalDynamicRHI::RHICreateMeshShader(TArrayView<const uint8> Code, const FSHAHash& Hash)
+FMeshShaderRHIRef FMetalDynamicRHI::RHICreateMeshShader(const FRHICreateShaderDesc& CreateShaderDesc)
 {
     MTL_SCOPED_AUTORELEASE_POOL;
 
-    return new FMetalMeshShader(*Device, Code);
+    return new FMetalMeshShader(*Device, CreateShaderDesc);
 }
 
-FAmplificationShaderRHIRef FMetalDynamicRHI::RHICreateAmplificationShader(TArrayView<const uint8> Code, const FSHAHash& Hash)
+FAmplificationShaderRHIRef FMetalDynamicRHI::RHICreateAmplificationShader(const FRHICreateShaderDesc& CreateShaderDesc)
 {
     MTL_SCOPED_AUTORELEASE_POOL;
 
-    return new FMetalAmplificationShader(*Device, Code);
+    return new FMetalAmplificationShader(*Device, CreateShaderDesc);
 }
 #endif
 
 #if METAL_RHI_RAYTRACING
-FRayTracingShaderRHIRef FMetalDynamicRHI::RHICreateRayTracingShader(TArrayView<const uint8> Code, const FSHAHash& Hash, EShaderFrequency ShaderFrequency)
+FRayTracingShaderRHIRef FMetalDynamicRHI::RHICreateRayTracingShader(const FRHICreateShaderDesc& CreateShaderDesc, EShaderFrequency ShaderFrequency)
 {
 	MTL_SCOPED_AUTORELEASE_POOL;
 	
-	return new FMetalRayShader(*Device, Code, ShaderFrequency);
+	return new FMetalRayShader(*Device, CreateShaderDesc, ShaderFrequency);
 }
 #endif // METAL_RHI_RAYTRACING
 
diff --git a/Engine/Source/Runtime/Apple/MetalRHI/Private/MetalShaders.cpp b/Engine/Source/Runtime/Apple/MetalRHI/Private/MetalShaders.cpp
--- a/Engine/Source/Runtime/Apple/MetalRHI/Private/MetalShaders.cpp
+++ b/Engine/Source/Runtime/Apple/MetalRHI/Private/MetalShaders.cpp
@@ -28,24 +28,6 @@
 #	include "Developer/ShaderCompilerCommon/Public/ShaderCompilerCommon.h"
 #undef SHADERCOMPILERCOMMON_API
 
-
-NS::String* DecodeMetalSourceCode(uint32 CodeSize, TArray<uint8> const& CompressedSource)
-{
-	NS::String* GlslCodeNSString = nullptr;
-	if (CodeSize && CompressedSource.Num())
-	{
-		TArray<ANSICHAR> UncompressedCode;
-		UncompressedCode.AddZeroed(CodeSize+1);
-		bool bSucceed = FCompression::UncompressMemory(NAME_Zlib, UncompressedCode.GetData(), CodeSize, CompressedSource.GetData(), CompressedSource.Num());
-		if (bSucceed)
-		{
-			GlslCodeNSString = NS::String::string(UncompressedCode.GetData(), NS::UTF8StringEncoding);
-            GlslCodeNSString->retain();
-		}
-	}
-	return GlslCodeNSString;
-}
-
 MTL::LanguageVersion ValidateVersion(uint32 Version)
 {
     MTL::LanguageVersion Result = MTL::LanguageVersion2_4;
diff --git a/Engine/Source/Runtime/Apple/MetalRHI/Private/Shaders/MetalShaderLibrary.cpp b/Engine/Source/Runtime/Apple/MetalRHI/Private/Shaders/MetalShaderLibrary.cpp
--- a/Engine/Source/Runtime/Apple/MetalRHI/Private/Shaders/MetalShaderLibrary.cpp
+++ b/Engine/Source/Runtime/Apple/MetalRHI/Private/Shaders/MetalShaderLibrary.cpp
@@ -34,9 +34,9 @@ static void ReleaseLibraryMemory(FMetalShaderLibrary::FShaderLibDataOwner* Libra
 
 
 template<typename ShaderType>
-static TRefCountPtr<FRHIShader> CreateMetalShader(FMetalDevice& Device, TArrayView<const uint8> InCode, MTLLibraryPtr InLibrary)
+static TRefCountPtr<FRHIShader> CreateMetalShader(FMetalDevice& Device, TConstArrayView<uint8> InView,  MTLLibraryPtr InLibrary)
 {
-	ShaderType* Shader = new ShaderType(Device, InCode, InLibrary);
+	ShaderType* Shader = new ShaderType(Device, FRHICreateShaderDesc(InView), InLibrary);
 	if (!Shader->GetFunction())
 	{
 		delete Shader;
@@ -243,7 +243,7 @@ TRefCountPtr<FRHIShader> FMetalShaderLibrary::CreateShader(int32 Index, bool bRe
 		case SF_RayHitGroup:
 		case SF_RayCallable:
 		{
-			FMetalRayShader* RayShader = new FMetalRayShader(Device, Code, (EShaderFrequency)ShaderEntry.GetFrequency(), LazyLibraries[LibraryIndex]->Library);
+			FMetalRayShader* RayShader = new FMetalRayShader(Device, FRHICreateShaderDesc(Code), (EShaderFrequency)ShaderEntry.GetFrequency(), LazyLibraries[LibraryIndex]->Library);
 			Shader = TRefCountPtr<FRHIShader>(RayShader);
 			break;
 		}
@@ -255,10 +255,5 @@ TRefCountPtr<FRHIShader> FMetalShaderLibrary::CreateShader(int32 Index, bool bRe
 			break;
 	}
 
-	if (Shader)
-	{
-		Shader->SetHash(SerializedShaders.GetShaderHashes()[Index]);
-	}
-
 	return Shader;
 }
diff --git a/Engine/Source/Runtime/Apple/MetalRHI/Private/Shaders/MetalShaderLibrary.h b/Engine/Source/Runtime/Apple/MetalRHI/Private/Shaders/MetalShaderLibrary.h
--- a/Engine/Source/Runtime/Apple/MetalRHI/Private/Shaders/MetalShaderLibrary.h
+++ b/Engine/Source/Runtime/Apple/MetalRHI/Private/Shaders/MetalShaderLibrary.h
@@ -8,6 +8,7 @@
 
 #include "MetalRHIPrivate.h"
 #include "MetalShaderResources.h"
+#include "Misc/CoreDelegates.h"
 #include "RHIShaderLibrary.h"
 #include "ShaderCodeArchive.h"
 #include "Async/MappedFileHandle.h"
diff --git a/Engine/Source/Runtime/Apple/MetalRHI/Private/Shaders/Types/MetalAmplificationShader.cpp b/Engine/Source/Runtime/Apple/MetalRHI/Private/Shaders/Types/MetalAmplificationShader.cpp
--- a/Engine/Source/Runtime/Apple/MetalRHI/Private/Shaders/Types/MetalAmplificationShader.cpp
+++ b/Engine/Source/Runtime/Apple/MetalRHI/Private/Shaders/Types/MetalAmplificationShader.cpp
@@ -12,16 +12,16 @@
 #pragma mark - Metal RHI Amplification Shader Class
 
 #if PLATFORM_SUPPORTS_MESH_SHADERS
-FMetalAmplificationShader::FMetalAmplificationShader(FMetalDevice& MetalDevice, TArrayView<const uint8> InCode) : TMetalBaseShader<FRHIAmplificationShader, SF_Amplification>(MetalDevice)
+FMetalAmplificationShader::FMetalAmplificationShader(FMetalDevice& MetalDevice, const FRHICreateShaderDesc& CreateShaderDesc) : TMetalBaseShader<FRHIAmplificationShader, SF_Amplification>(MetalDevice)
 {
 	FMetalCodeHeader Header;
-	Init(InCode, Header, MTLLibraryPtr());
+	Init(CreateShaderDesc, Header, MTLLibraryPtr());
 }
 
-FMetalAmplificationShader::FMetalAmplificationShader(FMetalDevice& MetalDevice, TArrayView<const uint8> InCode, MTLLibraryPtr InLibrary) : TMetalBaseShader<FRHIAmplificationShader, SF_Amplification>(MetalDevice)
+FMetalAmplificationShader::FMetalAmplificationShader(FMetalDevice& MetalDevice, const FRHICreateShaderDesc& CreateShaderDesc, MTLLibraryPtr InLibrary) : TMetalBaseShader<FRHIAmplificationShader, SF_Amplification>(MetalDevice)
 {
 	FMetalCodeHeader Header;
-	Init(InCode, Header, InLibrary);
+	Init(CreateShaderDesc, Header, InLibrary);
 }
 
 MTLFunctionPtr FMetalAmplificationShader::GetFunction()
diff --git a/Engine/Source/Runtime/Apple/MetalRHI/Private/Shaders/Types/MetalAmplificationShader.h b/Engine/Source/Runtime/Apple/MetalRHI/Private/Shaders/Types/MetalAmplificationShader.h
--- a/Engine/Source/Runtime/Apple/MetalRHI/Private/Shaders/Types/MetalAmplificationShader.h
+++ b/Engine/Source/Runtime/Apple/MetalRHI/Private/Shaders/Types/MetalAmplificationShader.h
@@ -18,8 +18,8 @@
 class FMetalAmplificationShader : public TMetalBaseShader<FRHIAmplificationShader, SF_Amplification>
 {
 public:
-    FMetalAmplificationShader(FMetalDevice& Device, TArrayView<const uint8> InCode);
-    FMetalAmplificationShader(FMetalDevice& Device, TArrayView<const uint8> InCode, MTLLibraryPtr InLibrary);
+    FMetalAmplificationShader(FMetalDevice& Device, const FRHICreateShaderDesc& CreateShaderDesc);
+    FMetalAmplificationShader(FMetalDevice& Device, const FRHICreateShaderDesc& CreateShaderDesc, MTLLibraryPtr InLibrary);
 
     MTLFunctionPtr GetFunction();
 };
diff --git a/Engine/Source/Runtime/Apple/MetalRHI/Private/Shaders/Types/MetalComputeShader.cpp b/Engine/Source/Runtime/Apple/MetalRHI/Private/Shaders/Types/MetalComputeShader.cpp
--- a/Engine/Source/Runtime/Apple/MetalRHI/Private/Shaders/Types/MetalComputeShader.cpp
+++ b/Engine/Source/Runtime/Apple/MetalRHI/Private/Shaders/Types/MetalComputeShader.cpp
@@ -17,7 +17,7 @@
 extern int32 GMetalShaderValidationType;
 extern FString GMetalShaderValidationShaderName;
 
-FMetalComputeShader::FMetalComputeShader(FMetalDevice& Device, TArrayView<const uint8> InCode, MTLLibraryPtr InLibrary)
+FMetalComputeShader::FMetalComputeShader(FMetalDevice& Device, const FRHICreateShaderDesc& CreateShaderDesc, MTLLibraryPtr InLibrary)
 	: TMetalBaseShader<FRHIComputeShader, SF_Compute>(Device)
 	, NumThreadsX(0)
 	, NumThreadsY(0)
@@ -25,7 +25,7 @@ FMetalComputeShader::FMetalComputeShader(FMetalDevice& Device, TArrayView<const
 {
 	Pipeline = nullptr;
 	FMetalCodeHeader Header;
-	Init(InCode, Header, InLibrary);
+	Init(CreateShaderDesc, Header, InLibrary);
 
 #if METAL_RHI_RAYTRACING
 	RayTracingBindings = Header.RayTracing;
diff --git a/Engine/Source/Runtime/Apple/MetalRHI/Private/Shaders/Types/MetalComputeShader.h b/Engine/Source/Runtime/Apple/MetalRHI/Private/Shaders/Types/MetalComputeShader.h
--- a/Engine/Source/Runtime/Apple/MetalRHI/Private/Shaders/Types/MetalComputeShader.h
+++ b/Engine/Source/Runtime/Apple/MetalRHI/Private/Shaders/Types/MetalComputeShader.h
@@ -18,7 +18,7 @@ class FMetalDevice;
 class FMetalComputeShader : public TMetalBaseShader<FRHIComputeShader, SF_Compute>
 {
 public:
-	FMetalComputeShader(FMetalDevice& Device, TArrayView<const uint8> InCode, MTLLibraryPtr InLibrary);
+	FMetalComputeShader(FMetalDevice& Device, const FRHICreateShaderDesc& CreateShaderDesc, MTLLibraryPtr InLibrary);
 	virtual ~FMetalComputeShader();
 
 	FMetalShaderPipelinePtr GetPipeline();
diff --git a/Engine/Source/Runtime/Apple/MetalRHI/Private/Shaders/Types/MetalGeometryShader.cpp b/Engine/Source/Runtime/Apple/MetalRHI/Private/Shaders/Types/MetalGeometryShader.cpp
--- a/Engine/Source/Runtime/Apple/MetalRHI/Private/Shaders/Types/MetalGeometryShader.cpp
+++ b/Engine/Source/Runtime/Apple/MetalRHI/Private/Shaders/Types/MetalGeometryShader.cpp
@@ -11,18 +11,18 @@
 #pragma mark - Metal RHI Geometry Shader Class
 
 #if METAL_USE_METAL_SHADER_CONVERTER
-FMetalGeometryShader::FMetalGeometryShader(FMetalDevice& MetalDevice, TArrayView<const uint8> InCode) :
+FMetalGeometryShader::FMetalGeometryShader(FMetalDevice& MetalDevice, const FRHICreateShaderDesc& CreateShaderDesc) :
 	TMetalBaseShader<FRHIGeometryShader, SF_Geometry>(MetalDevice)
 {
 	FMetalCodeHeader Header;
-	Init(InCode, Header, MTLLibraryPtr());
+	Init(CreateShaderDesc, Header, MTLLibraryPtr());
 }
 
-FMetalGeometryShader::FMetalGeometryShader(FMetalDevice& MetalDevice, TArrayView<const uint8> InCode, MTLLibraryPtr InLibrary) :
+FMetalGeometryShader::FMetalGeometryShader(FMetalDevice& MetalDevice, const FRHICreateShaderDesc& CreateShaderDesc, MTLLibraryPtr InLibrary) :
 	TMetalBaseShader<FRHIGeometryShader, SF_Geometry>(MetalDevice)
 {
 	FMetalCodeHeader Header;
-	Init(InCode, Header, InLibrary);
+	Init(CreateShaderDesc, Header, InLibrary);
 }
 
 MTLFunctionPtr FMetalGeometryShader::GetFunction()
diff --git a/Engine/Source/Runtime/Apple/MetalRHI/Private/Shaders/Types/MetalGeometryShader.h b/Engine/Source/Runtime/Apple/MetalRHI/Private/Shaders/Types/MetalGeometryShader.h
--- a/Engine/Source/Runtime/Apple/MetalRHI/Private/Shaders/Types/MetalGeometryShader.h
+++ b/Engine/Source/Runtime/Apple/MetalRHI/Private/Shaders/Types/MetalGeometryShader.h
@@ -19,8 +19,8 @@ class FMetalGeometryShader : public TMetalBaseShader<FRHIGeometryShader, SF_Geom
 {
 #if METAL_USE_METAL_SHADER_CONVERTER
 public:
-    FMetalGeometryShader(FMetalDevice& Device, TArrayView<const uint8> InCode);
-    FMetalGeometryShader(FMetalDevice& Device, TArrayView<const uint8> InCode, MTLLibraryPtr InLibrary);
+    FMetalGeometryShader(FMetalDevice& Device, const FRHICreateShaderDesc& CreateShaderDesc);
+    FMetalGeometryShader(FMetalDevice& Device, const FRHICreateShaderDesc& CreateShaderDesc, MTLLibraryPtr InLibrary);
 
     MTLFunctionPtr GetFunction();
 #endif
diff --git a/Engine/Source/Runtime/Apple/MetalRHI/Private/Shaders/Types/MetalMeshShader.cpp b/Engine/Source/Runtime/Apple/MetalRHI/Private/Shaders/Types/MetalMeshShader.cpp
--- a/Engine/Source/Runtime/Apple/MetalRHI/Private/Shaders/Types/MetalMeshShader.cpp
+++ b/Engine/Source/Runtime/Apple/MetalRHI/Private/Shaders/Types/MetalMeshShader.cpp
@@ -11,18 +11,18 @@
 #pragma mark - Metal RHI Mesh Shader Class
 
 #if PLATFORM_SUPPORTS_MESH_SHADERS
-FMetalMeshShader::FMetalMeshShader(FMetalDevice& MetalDevice, TArrayView<const uint8> InCode)
+FMetalMeshShader::FMetalMeshShader(FMetalDevice& MetalDevice, const FRHICreateShaderDesc& CreateShaderDesc)
 	: TMetalBaseShader<FRHIMeshShader, SF_Mesh>(MetalDevice)
 {
 	FMetalCodeHeader Header;
-	Init(InCode, Header, MTLLibraryPtr());
+	Init(CreateShaderDesc, Header, MTLLibraryPtr());
 }
 
-FMetalMeshShader::FMetalMeshShader(FMetalDevice& MetalDevice, TArrayView<const uint8> InCode, MTLLibraryPtr InLibrary)
+FMetalMeshShader::FMetalMeshShader(FMetalDevice& MetalDevice, const FRHICreateShaderDesc& CreateShaderDesc, MTLLibraryPtr InLibrary)
 	: TMetalBaseShader<FRHIMeshShader, SF_Mesh>(MetalDevice)
 {
 	FMetalCodeHeader Header;
-	Init(InCode, Header, InLibrary);
+	Init(CreateShaderDesc, Header, InLibrary);
 }
 
 MTLFunctionPtr FMetalMeshShader::GetFunction()
diff --git a/Engine/Source/Runtime/Apple/MetalRHI/Private/Shaders/Types/MetalMeshShader.h b/Engine/Source/Runtime/Apple/MetalRHI/Private/Shaders/Types/MetalMeshShader.h
--- a/Engine/Source/Runtime/Apple/MetalRHI/Private/Shaders/Types/MetalMeshShader.h
+++ b/Engine/Source/Runtime/Apple/MetalRHI/Private/Shaders/Types/MetalMeshShader.h
@@ -19,8 +19,8 @@
 class FMetalMeshShader : public TMetalBaseShader<FRHIMeshShader, SF_Mesh>
 {
 public:
-    FMetalMeshShader(FMetalDevice& Device, TArrayView<const uint8> InCode);
-    FMetalMeshShader(FMetalDevice& Device, TArrayView<const uint8> InCode, MTLLibraryPtr InLibrary);
+    FMetalMeshShader(FMetalDevice& Device, const FRHICreateShaderDesc& CreateShaderDesc);
+    FMetalMeshShader(FMetalDevice& Device, const FRHICreateShaderDesc& CreateShaderDesc, MTLLibraryPtr InLibrary);
 
     MTLFunctionPtr GetFunction();
 };
diff --git a/Engine/Source/Runtime/Apple/MetalRHI/Private/Shaders/Types/MetalPixelShader.cpp b/Engine/Source/Runtime/Apple/MetalRHI/Private/Shaders/Types/MetalPixelShader.cpp
--- a/Engine/Source/Runtime/Apple/MetalRHI/Private/Shaders/Types/MetalPixelShader.cpp
+++ b/Engine/Source/Runtime/Apple/MetalRHI/Private/Shaders/Types/MetalPixelShader.cpp
@@ -11,18 +11,18 @@
 #pragma mark - Metal RHI Pixel Shader Class
 
 
-FMetalPixelShader::FMetalPixelShader(FMetalDevice& MetalDevice, TArrayView<const uint8> InCode)
+FMetalPixelShader::FMetalPixelShader(FMetalDevice& MetalDevice, const FRHICreateShaderDesc& CreateShaderDesc)
 	: TMetalBaseShader<FRHIPixelShader, SF_Pixel>(MetalDevice)
 {
 	FMetalCodeHeader Header;
-	Init(InCode, Header, MTLLibraryPtr());
+	Init(CreateShaderDesc, Header, MTLLibraryPtr());
 }
 
-FMetalPixelShader::FMetalPixelShader(FMetalDevice& MetalDevice, TArrayView<const uint8> InCode, MTLLibraryPtr InLibrary)
+FMetalPixelShader::FMetalPixelShader(FMetalDevice& MetalDevice, const FRHICreateShaderDesc& CreateShaderDesc, MTLLibraryPtr InLibrary)
 	: TMetalBaseShader<FRHIPixelShader, SF_Pixel>(MetalDevice)
 {
 	FMetalCodeHeader Header;
-	Init(InCode, Header, InLibrary);
+	Init(CreateShaderDesc, Header, InLibrary);
 }
 
 MTLFunctionPtr FMetalPixelShader::GetFunction()
diff --git a/Engine/Source/Runtime/Apple/MetalRHI/Private/Shaders/Types/MetalPixelShader.h b/Engine/Source/Runtime/Apple/MetalRHI/Private/Shaders/Types/MetalPixelShader.h
--- a/Engine/Source/Runtime/Apple/MetalRHI/Private/Shaders/Types/MetalPixelShader.h
+++ b/Engine/Source/Runtime/Apple/MetalRHI/Private/Shaders/Types/MetalPixelShader.h
@@ -18,8 +18,8 @@
 class FMetalPixelShader : public TMetalBaseShader<FRHIPixelShader, SF_Pixel>
 {
 public:
-	FMetalPixelShader(FMetalDevice& Device, TArrayView<const uint8> InCode);
-	FMetalPixelShader(FMetalDevice& Device, TArrayView<const uint8> InCode, MTLLibraryPtr InLibrary);
+	FMetalPixelShader(FMetalDevice& Device, const FRHICreateShaderDesc& CreateShaderDesc);
+	FMetalPixelShader(FMetalDevice& Device, const FRHICreateShaderDesc& CreateShaderDesc, MTLLibraryPtr InLibrary);
 
 	MTLFunctionPtr GetFunction();
 };
diff --git a/Engine/Source/Runtime/Apple/MetalRHI/Private/Shaders/Types/MetalRayShader.cpp b/Engine/Source/Runtime/Apple/MetalRHI/Private/Shaders/Types/MetalRayShader.cpp
--- a/Engine/Source/Runtime/Apple/MetalRHI/Private/Shaders/Types/MetalRayShader.cpp
+++ b/Engine/Source/Runtime/Apple/MetalRHI/Private/Shaders/Types/MetalRayShader.cpp
@@ -10,7 +10,7 @@
 
 #if METAL_RHI_RAYTRACING
 
-FMetalRayShader::FMetalRayShader(FMetalDevice& Device, TArrayView<const uint8> InCode, EShaderFrequency InFrequency, MTLLibraryPtr InLibrary)
+FMetalRayShader::FMetalRayShader(FMetalDevice& Device, const FRHICreateShaderDesc& CreateShaderDesc, EShaderFrequency InFrequency, MTLLibraryPtr InLibrary)
 	: FRHIRayTracingShader(InFrequency)
 	, Frequency(InFrequency)
 {
@@ -18,12 +18,12 @@ FMetalRayShader::FMetalRayShader(FMetalDevice& Device, TArrayView<const uint8> I
 	FMetalCodeHeader Header;
 	FShaderResourceTable SRT;
 	
-	FShaderCodeReader ShaderCodeReader(InCode);
+	FShaderCodeReader ShaderCodeReader(CreateShaderDesc.Code);
 	
 	uint8 OfflineCompiledFlag = 0;
 	IRShaderStage DefaultBytecodeType, ExtendedBytecodeType;
 	{ 
-		FMemoryReaderView Ar(InCode, true);
+		FMemoryReaderView Ar(CreateShaderDesc.Code, true);
 		Ar.SetLimitSize(ShaderCodeReader.GetActualShaderCodeSize());
 		Ar << OfflineCompiledFlag;
 		check(OfflineCompiledFlag);
@@ -54,21 +54,21 @@ FMetalRayShader::FMetalRayShader(FMetalDevice& Device, TArrayView<const uint8> I
 		{
 			check(DefaultBytecodeType == IRShaderStageRayGeneration);
 			// This shouldn't need anything special?
-			RayGen = new FMetalRayGenShader(Device, InCode, InLibrary, DefaultBytecodeType, EntryPoint, bCompileRayTracingShadersAsBinary);
+			RayGen = new FMetalRayGenShader(Device, CreateShaderDesc, InLibrary, DefaultBytecodeType, EntryPoint, bCompileRayTracingShadersAsBinary);
 			break;
 		}
 		case SF_RayMiss:
 		{
 			check(DefaultBytecodeType == IRShaderStageMiss);
 			// This also doesn't need anything special
-			Miss = new FMetalRayMissShader(Device, InCode, InLibrary, DefaultBytecodeType, EntryPoint, bCompileRayTracingShadersAsBinary);
+			Miss = new FMetalRayMissShader(Device, CreateShaderDesc, InLibrary, DefaultBytecodeType, EntryPoint, bCompileRayTracingShadersAsBinary);
 			break;
 		}
 		case SF_RayCallable:
 		{
 			check(DefaultBytecodeType == IRShaderStageCallable);
 			// This also doesn't need anything special
-			Callable = new FMetalRayCallableShader(Device, InCode, InLibrary, DefaultBytecodeType, EntryPoint, bCompileRayTracingShadersAsBinary);
+			Callable = new FMetalRayCallableShader(Device, CreateShaderDesc, InLibrary, DefaultBytecodeType, EntryPoint, bCompileRayTracingShadersAsBinary);
 			break;
 		}
 		case SF_RayHitGroup:
@@ -77,17 +77,17 @@ FMetalRayShader::FMetalRayShader(FMetalDevice& Device, TArrayView<const uint8> I
 			{
 				case IRShaderStageClosestHit:
 				{
-					ClosestHit = new FMetalHitGroupShader(Device, InCode, InLibrary, DefaultBytecodeType, EntryPoint, bCompileRayTracingShadersAsBinary, false);
+					ClosestHit = new FMetalHitGroupShader(Device, CreateShaderDesc, InLibrary, DefaultBytecodeType, EntryPoint, bCompileRayTracingShadersAsBinary, false);
 					break;
 				}
 				case IRShaderStageAnyHit:
 				{
-					AnyHitAndIntersection = new FMetalHitGroupShader(Device, InCode, InLibrary, DefaultBytecodeType, AnyHitEntryPoint, bCompileRayTracingShadersAsBinary, false);
+					AnyHitAndIntersection = new FMetalHitGroupShader(Device, CreateShaderDesc, InLibrary, DefaultBytecodeType, AnyHitEntryPoint, bCompileRayTracingShadersAsBinary, false);
 					break;
 				}
 				case IRShaderStageIntersection:
 				{
-					AnyHitAndIntersection = new FMetalHitGroupShader(Device, InCode, InLibrary, DefaultBytecodeType, IntersectionEntryPoint, bCompileRayTracingShadersAsBinary, false);
+					AnyHitAndIntersection = new FMetalHitGroupShader(Device, CreateShaderDesc, InLibrary, DefaultBytecodeType, IntersectionEntryPoint, bCompileRayTracingShadersAsBinary, false);
 					break;
 				}
 				default:
@@ -100,13 +100,13 @@ FMetalRayShader::FMetalRayShader(FMetalDevice& Device, TArrayView<const uint8> I
 				case IRShaderStageAnyHit:
 				{
 					check(!AnyHitAndIntersection);
-					AnyHitAndIntersection = new FMetalHitGroupShader(Device, InCode, InLibrary, ExtendedBytecodeType, AnyHitEntryPoint, bCompileRayTracingShadersAsBinary, true);
+					AnyHitAndIntersection = new FMetalHitGroupShader(Device, CreateShaderDesc, InLibrary, ExtendedBytecodeType, AnyHitEntryPoint, bCompileRayTracingShadersAsBinary, true);
 					break;
 				}
 				case IRShaderStageIntersection:
 				{
 					check(!AnyHitAndIntersection);
-					AnyHitAndIntersection = new FMetalHitGroupShader(Device, InCode, InLibrary, ExtendedBytecodeType, IntersectionEntryPoint, bCompileRayTracingShadersAsBinary, true);
+					AnyHitAndIntersection = new FMetalHitGroupShader(Device, CreateShaderDesc, InLibrary, ExtendedBytecodeType, IntersectionEntryPoint, bCompileRayTracingShadersAsBinary, true);
 					break;
 				}
 				case IRShaderStageInvalid:
diff --git a/Engine/Source/Runtime/Apple/MetalRHI/Private/Shaders/Types/MetalRayShader.h b/Engine/Source/Runtime/Apple/MetalRHI/Private/Shaders/Types/MetalRayShader.h
--- a/Engine/Source/Runtime/Apple/MetalRHI/Private/Shaders/Types/MetalRayShader.h
+++ b/Engine/Source/Runtime/Apple/MetalRHI/Private/Shaders/Types/MetalRayShader.h
@@ -18,14 +18,14 @@ template<typename BaseResourceType, int32 ShaderType>
 class FMetalBaseRayShader : public TMetalBaseShader<BaseResourceType, ShaderType>
 {
 public:
-	FMetalBaseRayShader(FMetalDevice& MetalDevice, TArrayView<const uint8> InCode, MTLLibraryPtr InLibrary, IRShaderStage InShaderStage, FString& InEntryPoint, bool bInCompileAsBinary, bool bUseExtendedCodeSection = false)
+	FMetalBaseRayShader(FMetalDevice& MetalDevice, const FRHICreateShaderDesc& CreateShaderDesc, MTLLibraryPtr InLibrary, IRShaderStage InShaderStage, FString& InEntryPoint, bool bInCompileAsBinary, bool bUseExtendedCodeSection = false)
 				: TMetalBaseShader<BaseResourceType, ShaderType>(MetalDevice),
 				bCompileAsBinary(bInCompileAsBinary),
 				ShaderStage(InShaderStage),
 				EntryPoint(InEntryPoint)
 	{
 		FMetalCodeHeader Header;
-		TMetalBaseShader<BaseResourceType, ShaderType>::Init(InCode, Header, InLibrary, bCompileAsBinary, bUseExtendedCodeSection);
+		TMetalBaseShader<BaseResourceType, ShaderType>::Init(CreateShaderDesc, Header, InLibrary, bCompileAsBinary, bUseExtendedCodeSection);
 	}
 	
 	IRShaderStage GetShaderStage()
@@ -53,7 +53,7 @@ class FMetalRayShader : public FRHIRayTracingShader
 public:
 	UE_NONCOPYABLE(FMetalRayShader)
 
-	FMetalRayShader(FMetalDevice& Device, TArrayView<const uint8> InCode, EShaderFrequency InFrequency, MTLLibraryPtr InLibrary = nullptr);
+	FMetalRayShader(FMetalDevice& Device, const FRHICreateShaderDesc& CreateShaderDesc, EShaderFrequency InFrequency, MTLLibraryPtr InLibrary = nullptr);
 	virtual ~FMetalRayShader();
 
 	inline IRRootSignature *GetGlobalRootSignature() const { return GlobalRootSignature; }
diff --git a/Engine/Source/Runtime/Apple/MetalRHI/Private/Shaders/Types/MetalVertexShader.cpp b/Engine/Source/Runtime/Apple/MetalRHI/Private/Shaders/Types/MetalVertexShader.cpp
--- a/Engine/Source/Runtime/Apple/MetalRHI/Private/Shaders/Types/MetalVertexShader.cpp
+++ b/Engine/Source/Runtime/Apple/MetalRHI/Private/Shaders/Types/MetalVertexShader.cpp
@@ -11,18 +11,18 @@
 #pragma mark - Metal RHI Vertex Shader Class
 
 
-FMetalVertexShader::FMetalVertexShader(FMetalDevice& MetalDevice, TArrayView<const uint8> InCode)
+FMetalVertexShader::FMetalVertexShader(FMetalDevice& MetalDevice, const FRHICreateShaderDesc& CreateShaderDesc)
 	: TMetalBaseShader<FRHIVertexShader, SF_Vertex>(MetalDevice)
 {
 	FMetalCodeHeader Header;
-	Init(InCode, Header, MTLLibraryPtr());
+	Init(CreateShaderDesc, Header, MTLLibraryPtr());
 }
 
-FMetalVertexShader::FMetalVertexShader(FMetalDevice& MetalDevice, TArrayView<const uint8> InCode, MTLLibraryPtr InLibrary)
+FMetalVertexShader::FMetalVertexShader(FMetalDevice& MetalDevice, const FRHICreateShaderDesc& CreateShaderDesc, MTLLibraryPtr InLibrary)
 	: TMetalBaseShader<FRHIVertexShader, SF_Vertex>(MetalDevice)
 {
 	FMetalCodeHeader Header;
-	Init(InCode, Header, InLibrary);
+	Init(CreateShaderDesc, Header, InLibrary);
 }
 
 MTLFunctionPtr FMetalVertexShader::GetFunction()
diff --git a/Engine/Source/Runtime/Apple/MetalRHI/Private/Shaders/Types/MetalVertexShader.h b/Engine/Source/Runtime/Apple/MetalRHI/Private/Shaders/Types/MetalVertexShader.h
--- a/Engine/Source/Runtime/Apple/MetalRHI/Private/Shaders/Types/MetalVertexShader.h
+++ b/Engine/Source/Runtime/Apple/MetalRHI/Private/Shaders/Types/MetalVertexShader.h
@@ -18,8 +18,8 @@
 class FMetalVertexShader : public TMetalBaseShader<FRHIVertexShader, SF_Vertex>
 {
 public:
-	FMetalVertexShader(FMetalDevice& Device, TArrayView<const uint8> InCode);
-	FMetalVertexShader(FMetalDevice& Device, TArrayView<const uint8> InCode, MTLLibraryPtr InLibrary);
+	FMetalVertexShader(FMetalDevice& Device, const FRHICreateShaderDesc& CreateShaderDesc);
+	FMetalVertexShader(FMetalDevice& Device, const FRHICreateShaderDesc& CreateShaderDesc, MTLLibraryPtr InLibrary);
 
 	MTLFunctionPtr GetFunction();
 #if PLATFORM_SUPPORTS_GEOMETRY_SHADERS
diff --git a/Engine/Source/Runtime/Apple/MetalRHI/Private/Shaders/Types/Templates/MetalBaseShader.h b/Engine/Source/Runtime/Apple/MetalRHI/Private/Shaders/Types/Templates/MetalBaseShader.h
--- a/Engine/Source/Runtime/Apple/MetalRHI/Private/Shaders/Types/Templates/MetalBaseShader.h
+++ b/Engine/Source/Runtime/Apple/MetalRHI/Private/Shaders/Types/Templates/MetalBaseShader.h
@@ -70,13 +70,7 @@ struct FMetalShaderData
     MTLLibraryPtr Library;
 
     /** The debuggable text source */
-    NS::String* GlslCodeNSString = nullptr;
-
-    /** The compressed text source */
-    TArray<uint8> CompressedSource;
-
-    /** The uncompressed text source size */
-    uint32 CodeSize = 0;
+    NS::String* MetalCodeNSString = nullptr;
 
     /** Index of the function (in the library) pointing to the function requested by the user (when GetCompiledFunction() is called with an explicit index). */
     uint32 LibraryFunctionIndex = -1;
@@ -101,7 +95,7 @@ public:
 		this->Destroy();
 	}
 
-	void Init(TArrayView<const uint8> InCode, FMetalCodeHeader& Header, MTLLibraryPtr InLibrary, bool bBinaryCompiled = false, bool bUseExtendedCodeSection = false);
+	void Init(const FRHICreateShaderDesc& CreateShaderDesc, FMetalCodeHeader& Header, MTLLibraryPtr InLibrary, bool bBinaryCompiled = false, bool bUseExtendedCodeSection = false);
 	void Destroy();
 
 	virtual FString GetFunctionName()
@@ -110,9 +104,7 @@ public:
 	}
 	
 	/**
-	 * Gets the Metal source code as an NSString if available or nullptr if not.  Note
-	 * that this will dynamically decompress from compressed data on first
-	 * invocation.
+	 * Gets the Metal source code as an NSString if available or build a string with the hash & function name inside instead.
 	 */
 	NS::String* GetSourceCode();
 
@@ -130,11 +122,11 @@ protected:
 
 
 template<typename BaseResourceType, int32 ShaderType>
-void TMetalBaseShader<BaseResourceType, ShaderType>::Init(TArrayView<const uint8> InShaderCode, FMetalCodeHeader& Header, MTLLibraryPtr InLibrary, bool bBinaryCompiled, bool bUseExtendedCodeSection)
+void TMetalBaseShader<BaseResourceType, ShaderType>::Init(const FRHICreateShaderDesc& CreateShaderDesc, FMetalCodeHeader& Header, MTLLibraryPtr InLibrary, bool bBinaryCompiled, bool bUseExtendedCodeSection)
 {
-	FShaderCodeReader ShaderCode(InShaderCode);
+	FShaderCodeReader ShaderCode(CreateShaderDesc.Code);
 
-	FMemoryReaderView Ar(InShaderCode, true);
+	FMemoryReaderView Ar(CreateShaderDesc.Code, true);
 
 	Ar.SetLimitSize(ShaderCode.GetActualShaderCodeSize());
 
@@ -157,7 +149,7 @@ void TMetalBaseShader<BaseResourceType, ShaderType>::Init(TArrayView<const uint8
 	// remember where the header ended and code (precompiled or source) begins
 	int32 CodeOffset = Ar.Tell();
 	uint32 BufferSize = ShaderCode.GetActualShaderCodeSize() - CodeOffset;
-	const ANSICHAR* SourceCode = (ANSICHAR*)InShaderCode.GetData() + CodeOffset;
+	const ANSICHAR* SourceCode = (ANSICHAR*)CreateShaderDesc.Code.GetData() + CodeOffset;
 
 	if (bUseExtendedCodeSection)
 	{
@@ -168,7 +160,7 @@ void TMetalBaseShader<BaseResourceType, ShaderType>::Init(TArrayView<const uint8
 		if (ExtendedCodeDataOffset != -1 && ExtendedCodeDataSize != -1)
 		{
 			int32 ExtendedCodeStartOffset = CodeOffset + ExtendedCodeDataOffset;
-			SourceCodeExtended = (ANSICHAR*)InShaderCode.GetData() + ExtendedCodeStartOffset;	
+			SourceCodeExtended = (ANSICHAR*)CreateShaderDesc.Code.GetData() + ExtendedCodeStartOffset;	
 		}
 		else
 		{
@@ -192,55 +184,41 @@ void TMetalBaseShader<BaseResourceType, ShaderType>::Init(TArrayView<const uint8
 		UE_LOG(LogMetal, Display, TEXT("Loaded a text shader (will be slower to load)"));
 	}
 
-	bool bOfflineCompile = (OfflineCompiledFlag > 0);
-
-	bool bUseMetalShaderConverter = EnumHasAnyFlags(Header.Bindings.Flags, EMetalBindingsFlags::UseMetalShaderConverter); 
+	const bool bOfflineCompile = (OfflineCompiledFlag > 0);
+	const bool bUseMetalShaderConverter = EnumHasAnyFlags(Header.Bindings.Flags, EMetalBindingsFlags::UseMetalShaderConverter); 
 	
 	const ANSICHAR* ShaderSource = ShaderCode.FindOptionalData(EShaderOptionalDataKey::SourceCode);
 	bool bHasShaderSource = false;
 	static bool bForceTextShaders = false;
 	
-	static const auto CVar = IConsoleManager::Get().FindConsoleVariable(TEXT("r.Shaders.Symbols"));
-	
 	if(!bUseMetalShaderConverter)
 	{
-		bForceTextShaders = CVar->GetInt() == 1 || FParse::Param(FCommandLine::Get(),TEXT("metalshaderdebug"));
+		bForceTextShaders = FParse::Param(FCommandLine::Get(),TEXT("metalshaderdebug"));
 		bHasShaderSource = (ShaderSource && FCStringAnsi::Strlen(ShaderSource) > 0);
 	}
 	
 	if (!bHasShaderSource)
 	{
-		int32 LZMASourceSize = 0;
-		int32 SourceSize = 0;
-		const uint8* LZMASource = ShaderCode.FindOptionalDataAndSize(EShaderOptionalDataKey::CompressedDebugCode, LZMASourceSize);
-		const uint8* UnSourceLen = ShaderCode.FindOptionalDataAndSize(EShaderOptionalDataKey::UncompressedSize, SourceSize);
-		if (LZMASource && LZMASourceSize > 0 && UnSourceLen && SourceSize == sizeof(uint32))
-		{
-			CompressedSource.Append(LZMASource, LZMASourceSize);
-			memcpy(&CodeSize, UnSourceLen, sizeof(uint32));
-			bHasShaderSource = false;
-		}
 #if !UE_BUILD_SHIPPING
-		else if(bForceTextShaders)
+		if(bForceTextShaders)
 		{
-            GlslCodeNSString = FMetalShaderDebugCache::Get().GetShaderCode(SourceLen, SourceCRC);
-			if(GlslCodeNSString)
+			MetalCodeNSString = FMetalShaderDebugCache::Get().GetShaderCode(SourceLen, SourceCRC);
+			if(MetalCodeNSString)
 			{
-				GlslCodeNSString->retain();
+				MetalCodeNSString->retain();
+			}
+			else
+			{
+				UE_LOG(LogMetal, Warning, TEXT("-metalshaderdebug was used but shader source are unavailable, please use r.Shaders.ExtraData=1 or enable AllowsRuntimeShaderCompiling for the platform and add CFLAG_Debug to this shader before cooking."));
 			}
 		}
 #endif
-		if (bForceTextShaders && CodeSize && CompressedSource.Num())
-		{
-			bHasShaderSource = (GetSourceCode() != nullptr);
-		}
 	}
 	else if (bOfflineCompile && bHasShaderSource)
 	{
-		GlslCodeNSString = NS::String::string(ShaderSource, NS::UTF8StringEncoding);
-		check(GlslCodeNSString);
-        
-		GlslCodeNSString->retain();
+		MetalCodeNSString = NS::String::string(ShaderSource, NS::UTF8StringEncoding);
+		check(MetalCodeNSString);
+		MetalCodeNSString->retain();
 	}
 
 	Library = InLibrary;
@@ -272,15 +250,23 @@ void TMetalBaseShader<BaseResourceType, ShaderType>::Init(TArrayView<const uint8
 			{
 				// Archived shaders should never get in here.
 				check(!(Header.CompileFlags & (1 << CFLAG_Archive)) || BufferSize > 0);
-
-				// allow GCD to copy the data into its own buffer
-				//		dispatch_data_t GCDBuffer = dispatch_data_create(InShaderCode.GetTypedData() + CodeOffset, ShaderCode.GetActualShaderCodeSize() - CodeOffset, nullptr, DISPATCH_DATA_DESTRUCTOR_DEFAULT);
-				NS::Error* AError;
-				void* Buffer = FMemory::Malloc( BufferSize );
-				FMemory::Memcpy( Buffer, SourceCode, BufferSize );
-				dispatch_data_t GCDBuffer = dispatch_data_create(Buffer, BufferSize, dispatch_get_main_queue(), ^(void) { FMemory::Free(Buffer); } );
+				dispatch_data_t GCDBuffer;
+				if (CreateShaderDesc.Owner.IsValid())
+				{
+					// Hold a reference until the MTLLibrary is done using it
+					__block auto OwnerCopy = CreateShaderDesc.Owner;
+					GCDBuffer = dispatch_data_create(SourceCode, BufferSize, nil, ^{
+						OwnerCopy.Reset();
+					});
+				}
+				else
+				{
+					// Create a copy of the memory
+					GCDBuffer = dispatch_data_create(SourceCode, BufferSize, nil, DISPATCH_DATA_DESTRUCTOR_DEFAULT);
+				}
 
 				// load up the already compiled shader
+				NS::Error* AError;
 				Library = NS::TransferPtr(Device.GetDevice()->newLibrary(GCDBuffer, &AError));
 				dispatch_release(GCDBuffer);
 
@@ -294,7 +280,7 @@ void TMetalBaseShader<BaseResourceType, ShaderType>::Init(TArrayView<const uint8
 		{
 			MTL_SCOPED_AUTORELEASE_POOL;
 
-			NS::String* ShaderString = ((OfflineCompiledFlag == 0) ? NS::String::string(SourceCode, NS::UTF8StringEncoding) : GlslCodeNSString);
+			NS::String* ShaderString = ((OfflineCompiledFlag == 0) ? NS::String::string(SourceCode, NS::UTF8StringEncoding) : MetalCodeNSString);
 
             FString FinalShaderString;
 			const FString ShaderName = ShaderCode.FindOptionalData(FShaderCodeName::Key);
@@ -398,8 +384,8 @@ void TMetalBaseShader<BaseResourceType, ShaderType>::Init(TArrayView<const uint8
 				UE_LOG(LogRHI, Warning, TEXT("Created shader with warnings: %s"), *NSStringToFString(Error->description()));
 			}
 
-			GlslCodeNSString = NewShaderString;
-			GlslCodeNSString->retain();
+			MetalCodeNSString = NewShaderString;
+			MetalCodeNSString->retain();
 		}
 
 		GetCompiledFunction(true, -1, bBinaryCompiled);
@@ -416,10 +402,10 @@ void TMetalBaseShader<BaseResourceType, ShaderType>::Init(TArrayView<const uint8
 template<typename BaseResourceType, int32 ShaderType>
 void TMetalBaseShader<BaseResourceType, ShaderType>::Destroy()
 {
-    if(GlslCodeNSString)
+    if(MetalCodeNSString)
     {
-        GlslCodeNSString->release();
-        GlslCodeNSString = nullptr;
+        MetalCodeNSString->release();
+        MetalCodeNSString = nullptr;
     }
 }
 
@@ -428,17 +414,13 @@ inline NS::String* TMetalBaseShader<BaseResourceType, ShaderType>::GetSourceCode
 {
 	MTL_SCOPED_AUTORELEASE_POOL;
 	
-	if (!GlslCodeNSString && CodeSize && CompressedSource.Num())
-	{
-		GlslCodeNSString = DecodeMetalSourceCode(CodeSize, CompressedSource);
-	}
-	if (!GlslCodeNSString)
+	if (!MetalCodeNSString)
 	{
         FString ShaderString = FString::Printf(TEXT("Hash: %s, Name: %s"), *BaseResourceType::GetHash().ToString(), *GetFunctionName());
-        GlslCodeNSString = FStringToNSString(ShaderString);
-        GlslCodeNSString->retain();
+        MetalCodeNSString = FStringToNSString(ShaderString);
+        MetalCodeNSString->retain();
 	}
-	return GlslCodeNSString;
+	return MetalCodeNSString;
 }
 
 template<typename BaseResourceType, int32 ShaderType>
diff --git a/Engine/Source/Runtime/Apple/MetalRHI/Public/MetalResources.h b/Engine/Source/Runtime/Apple/MetalRHI/Public/MetalResources.h
--- a/Engine/Source/Runtime/Apple/MetalRHI/Public/MetalResources.h
+++ b/Engine/Source/Runtime/Apple/MetalRHI/Public/MetalResources.h
@@ -21,8 +21,6 @@ class FMetalContext;
 class FMetalShaderPipeline;
 class FMetalCommandBuffer;
 
-extern NS::String* DecodeMetalSourceCode(uint32 CodeSize, TArray<uint8> const& CompressedSource);
-
 struct FMetalRenderPipelineHash
 {
 	friend uint32 GetTypeHash(FMetalRenderPipelineHash const& Hash)
diff --git a/Engine/Source/Runtime/Core/Public/Misc/CoreDelegates.h b/Engine/Source/Runtime/Core/Public/Misc/CoreDelegates.h
--- a/Engine/Source/Runtime/Core/Public/Misc/CoreDelegates.h
+++ b/Engine/Source/Runtime/Core/Public/Misc/CoreDelegates.h
@@ -24,6 +24,12 @@
 #include "Templates/Function.h"
 #include "Templates/SharedPointer.h"
 
+#ifndef USE_MMAPPED_SHADERARCHIVE
+	// If enabled FSerializedShaderArchive, FIoStoreShaderCodeArchive and FMetalShaderLibrary will use a mmapping of the on-disk shader archive instead of normal serialization.
+	// The on-disk format must match the expected in memory layout.
+	#define USE_MMAPPED_SHADERARCHIVE 0
+#endif
+
 class AActor;
 class Error;
 class FConfigFile;
@@ -603,7 +609,12 @@ public:
 
 	static CORE_API TTSMulticastDelegate<void(FExtraBinaryConfigData&)>& TSAccessExtraBinaryConfigData();
 
-	using FAttachShaderReadRequestFunc = TFunctionRef<class FIoRequest(const class FIoChunkId&, FGraphEventRef)>;
+#if USE_MMAPPED_SHADERARCHIVE
+	using FShaderReadRequestResult = TPair<class FIoRequest, TIoStatusOr<struct FIoMappedRegion>>;
+#else
+	using FShaderReadRequestResult = class FIoRequest;
+#endif
+	using FAttachShaderReadRequestFunc = TFunctionRef<FShaderReadRequestResult(const class FIoChunkId&, FGraphEventRef)>;
 	static CORE_API TDelegate<void(TArrayView<const FSHAHash>, FAttachShaderReadRequestFunc)> PreloadPackageShaderMaps;
 	static CORE_API TDelegate<void(TArrayView<const FSHAHash>)> ReleasePreloadedPackageShaderMaps;
 	/** Called when the verbosity of a log category is changed */
@@ -654,4 +665,4 @@ private:
 
 #if UE_ENABLE_INCLUDE_ORDER_DEPRECATED_IN_5_8
 #include "Experimental/UnifiedError/UnifiedError.h"
-#endif
\ No newline at end of file
+#endif
diff --git a/Engine/Source/Runtime/CoreUObject/Private/Serialization/AsyncLoading2.cpp b/Engine/Source/Runtime/CoreUObject/Private/Serialization/AsyncLoading2.cpp
--- a/Engine/Source/Runtime/CoreUObject/Private/Serialization/AsyncLoading2.cpp
+++ b/Engine/Source/Runtime/CoreUObject/Private/Serialization/AsyncLoading2.cpp
@@ -87,6 +87,7 @@
 #include "Modules/ModuleManager.h"
 #include "Containers/MpscQueue.h"
 #include "Misc/AssetRegistryInterface.h"
+#include "Misc/CoreDelegates.h"
 #include "Misc/PathViews.h"
 #include "UObject/LinkerLoad.h"
 #include "IO/IoPriorityQueue.h"
@@ -7588,10 +7589,25 @@ void FAsyncPackage2::StartLoading(FAsyncLoadingThreadState2& ThreadState, FIoBat
 		TRACE_CPUPROFILER_EVENT_SCOPE(StartShaderMapRequests);
 		auto ReadShaderMapFunc = [this, &IoBatch](const FIoChunkId& ChunkId, FGraphEventRef GraphEvent)
 			{
+#if USE_MMAPPED_SHADERARCHIVE
+				FIoDispatcher& IoDispatcher = FIoDispatcher::Get();
+				TIoStatusOr<FIoMappedRegion> MappedRegionStatus = IoDispatcher.OpenMapped(ChunkId, FIoReadOptions());
+				if (MappedRegionStatus.IsOk())
+				{
+					GraphEvent->DispatchSubsequents();
+					GraphEvent.SafeRelease();
+					return FCoreDelegates::FShaderReadRequestResult(FIoRequest(), MoveTemp(MappedRegionStatus));
+				}
+#endif
 				GetPackageNode(Package_ExportsSerialized).AddBarrier();
 				int32 LocalPendingIoRequestsCounter = AsyncLoadingThread.PendingIoRequestsCounter.fetch_add(1) + 1;
 				TRACE_COUNTER_SET(AsyncLoadingPendingIoRequests, LocalPendingIoRequestsCounter);
-				return IoBatch.ReadWithCallback(ChunkId, FIoReadOptions(), Desc.Priority,
+#if USE_MMAPPED_SHADERARCHIVE
+				return FCoreDelegates::FShaderReadRequestResult(
+#else
+				return (// Avoid calling FIoRequest(FIoRequest...
+#endif
+				IoBatch.ReadWithCallback(ChunkId, FIoReadOptions(), Desc.Priority,
 					[this, GraphEvent](TIoStatusOr<FIoBuffer> Result)
 					{
 						GraphEvent->DispatchSubsequents();
@@ -7603,7 +7619,11 @@ void FAsyncPackage2::StartLoading(FAsyncLoadingThreadState2& ThreadState, FIoBat
 						{
 							LocalAsyncLoadingThread.AltZenaphore.NotifyOne();
 						}
-					});
+					})
+#if USE_MMAPPED_SHADERARCHIVE
+					, MoveTemp(MappedRegionStatus)
+#endif
+				);
 			};
 		FCoreDelegates::PreloadPackageShaderMaps.ExecuteIfBound(Data.ShaderMapHashes, ReadShaderMapFunc);
 	}
diff --git a/Engine/Source/Runtime/D3D12RHI/Private/D3D12RHIPrivate.h b/Engine/Source/Runtime/D3D12RHI/Private/D3D12RHIPrivate.h
--- a/Engine/Source/Runtime/D3D12RHI/Private/D3D12RHIPrivate.h
+++ b/Engine/Source/Runtime/D3D12RHI/Private/D3D12RHIPrivate.h
@@ -331,13 +331,13 @@ public:
 	virtual FDepthStencilStateRHIRef RHICreateDepthStencilState(const FDepthStencilStateInitializerRHI& Initializer) final override;
 	virtual FBlendStateRHIRef RHICreateBlendState(const FBlendStateInitializerRHI& Initializer) final override;
 	virtual FVertexDeclarationRHIRef RHICreateVertexDeclaration(const FVertexDeclarationElementList& Elements) final override;
-	virtual FPixelShaderRHIRef RHICreatePixelShader(TArrayView<const uint8> Code, const FSHAHash& Hash) final override;
-	virtual FVertexShaderRHIRef RHICreateVertexShader(TArrayView<const uint8> Code, const FSHAHash& Hash) final override;
-	virtual FMeshShaderRHIRef RHICreateMeshShader(TArrayView<const uint8> Code, const FSHAHash& Hash) final override;
-	virtual FAmplificationShaderRHIRef RHICreateAmplificationShader(TArrayView<const uint8> Code, const FSHAHash& Hash) final override;
-	virtual FGeometryShaderRHIRef RHICreateGeometryShader(TArrayView<const uint8> Code, const FSHAHash& Hash) final override;
-	virtual FComputeShaderRHIRef RHICreateComputeShader(TArrayView<const uint8> Code, const FSHAHash& Hash) override;
-	virtual FWorkGraphShaderRHIRef RHICreateWorkGraphShader(TArrayView<const uint8> Code, const FSHAHash& Hash, EShaderFrequency ShaderFrequency) override;
+	virtual FPixelShaderRHIRef RHICreatePixelShader(const FRHICreateShaderDesc& CreateShaderDesc) final override;
+	virtual FVertexShaderRHIRef RHICreateVertexShader(const FRHICreateShaderDesc& CreateShaderDesc) final override;
+	virtual FMeshShaderRHIRef RHICreateMeshShader(const FRHICreateShaderDesc& CreateShaderDesc) final override;
+	virtual FAmplificationShaderRHIRef RHICreateAmplificationShader(const FRHICreateShaderDesc& CreateShaderDesc) final override;
+	virtual FGeometryShaderRHIRef RHICreateGeometryShader(const FRHICreateShaderDesc& CreateShaderDesc) final override;
+	virtual FComputeShaderRHIRef RHICreateComputeShader(const FRHICreateShaderDesc& CreateShaderDesc) override;
+	virtual FWorkGraphShaderRHIRef RHICreateWorkGraphShader(const FRHICreateShaderDesc& CreateShaderDesc, EShaderFrequency ShaderFrequency) override;
 	virtual FGPUFenceRHIRef RHICreateGPUFence(const FName& Name) final override;
 	virtual FStagingBufferRHIRef RHICreateStagingBuffer() final override;
 	virtual void* RHILockStagingBuffer(FRHIStagingBuffer* StagingBuffer, FRHIGPUFence* Fence, uint32 Offset, uint32 SizeRHI) final override;
@@ -493,7 +493,7 @@ public:
 
 	virtual FRayTracingGeometryRHIRef RHICreateRayTracingGeometry(FRHICommandListBase& RHICmdList, const FRayTracingGeometryInitializer& Initializer) final override;
 	virtual FRayTracingSceneRHIRef RHICreateRayTracingScene(FRayTracingSceneInitializer Initializer) final override;
-	virtual FRayTracingShaderRHIRef RHICreateRayTracingShader(TArrayView<const uint8> Code, const FSHAHash& Hash, EShaderFrequency ShaderFrequency) final override;
+	virtual FRayTracingShaderRHIRef RHICreateRayTracingShader(const FRHICreateShaderDesc& CreateShaderDesc, EShaderFrequency ShaderFrequency) final override;
 	virtual FRayTracingPipelineStateRHIRef RHICreateRayTracingPipelineState(const FRayTracingPipelineStateInitializer& Initializer) final override;
 
 	virtual FShaderBindingTableRHIRef RHICreateShaderBindingTable(FRHICommandListBase& RHICmdList, const FRayTracingShaderBindingTableInitializer& Initializer) final override;
diff --git a/Engine/Source/Runtime/D3D12RHI/Private/D3D12Shaders.cpp b/Engine/Source/Runtime/D3D12RHI/Private/D3D12Shaders.cpp
--- a/Engine/Source/Runtime/D3D12RHI/Private/D3D12Shaders.cpp
+++ b/Engine/Source/Runtime/D3D12RHI/Private/D3D12Shaders.cpp
@@ -201,34 +201,34 @@ TShaderType* InitStandardShader(TShaderType* InShader, TArrayView<const uint8> I
 	return InitStandardShaderWithCustomSerialization<TShaderType>(InShader, InCode, CustomSerialization);
 }
 
-FVertexShaderRHIRef FD3D12DynamicRHI::RHICreateVertexShader(TArrayView<const uint8> Code, const FSHAHash& Hash)
+FVertexShaderRHIRef FD3D12DynamicRHI::RHICreateVertexShader(const FRHICreateShaderDesc& CreateShaderDesc)
 {
-	return InitStandardShader(new FD3D12VertexShader(), Code);
+	return InitStandardShader(new FD3D12VertexShader(), CreateShaderDesc.Code);
 }
 
-FMeshShaderRHIRef FD3D12DynamicRHI::RHICreateMeshShader(TArrayView<const uint8> Code, const FSHAHash& Hash)
+FMeshShaderRHIRef FD3D12DynamicRHI::RHICreateMeshShader(const FRHICreateShaderDesc& CreateShaderDesc)
 {
-	return InitStandardShader(new FD3D12MeshShader(), Code);
+	return InitStandardShader(new FD3D12MeshShader(), CreateShaderDesc.Code);
 }
 
-FAmplificationShaderRHIRef FD3D12DynamicRHI::RHICreateAmplificationShader(TArrayView<const uint8> Code, const FSHAHash& Hash)
+FAmplificationShaderRHIRef FD3D12DynamicRHI::RHICreateAmplificationShader(const FRHICreateShaderDesc& CreateShaderDesc)
 {
-	return InitStandardShader(new FD3D12AmplificationShader(), Code);
+	return InitStandardShader(new FD3D12AmplificationShader(), CreateShaderDesc.Code);
 }
 
-FPixelShaderRHIRef FD3D12DynamicRHI::RHICreatePixelShader(TArrayView<const uint8> Code, const FSHAHash& Hash)
+FPixelShaderRHIRef FD3D12DynamicRHI::RHICreatePixelShader(const FRHICreateShaderDesc& CreateShaderDesc)
 {
-	return InitStandardShader(new FD3D12PixelShader(), Code);
+	return InitStandardShader(new FD3D12PixelShader(), CreateShaderDesc.Code);
 }
 
-FGeometryShaderRHIRef FD3D12DynamicRHI::RHICreateGeometryShader(TArrayView<const uint8> Code, const FSHAHash& Hash)
+FGeometryShaderRHIRef FD3D12DynamicRHI::RHICreateGeometryShader(const FRHICreateShaderDesc& CreateShaderDesc)
 {
-	return InitStandardShader(new FD3D12GeometryShader(), Code);
+	return InitStandardShader(new FD3D12GeometryShader(), CreateShaderDesc.Code);
 }
 
-FComputeShaderRHIRef FD3D12DynamicRHI::RHICreateComputeShader(TArrayView<const uint8> Code, const FSHAHash& Hash)
+FComputeShaderRHIRef FD3D12DynamicRHI::RHICreateComputeShader(const FRHICreateShaderDesc& CreateShaderDesc)
 {
-	FD3D12ComputeShader* Shader = InitStandardShader(new FD3D12ComputeShader(), Code);
+	FD3D12ComputeShader* Shader = InitStandardShader(new FD3D12ComputeShader(), CreateShaderDesc.Code);
 	if (Shader)
 	{
 		Shader->RootSignature = GetAdapter().GetRootSignature(Shader);
@@ -238,9 +238,9 @@ FComputeShaderRHIRef FD3D12DynamicRHI::RHICreateComputeShader(TArrayView<const u
 	return Shader;
 }
 
-FWorkGraphShaderRHIRef FD3D12DynamicRHI::RHICreateWorkGraphShader(TArrayView<const uint8> Code, const FSHAHash& Hash, EShaderFrequency ShaderFrequency)
+FWorkGraphShaderRHIRef FD3D12DynamicRHI::RHICreateWorkGraphShader(const FRHICreateShaderDesc& CreateShaderDesc, EShaderFrequency ShaderFrequency)
 {
-	FD3D12WorkGraphShader* Shader = InitStandardShader(new FD3D12WorkGraphShader(ShaderFrequency), Code);
+	FD3D12WorkGraphShader* Shader = InitStandardShader(new FD3D12WorkGraphShader(ShaderFrequency), CreateShaderDesc.Code);
 	if (Shader)
 	{
 		Shader->RootSignature = GetAdapter().GetRootSignature(Shader);
@@ -252,7 +252,7 @@ FWorkGraphShaderRHIRef FD3D12DynamicRHI::RHICreateWorkGraphShader(TArrayView<con
 
 #if D3D12_RHI_RAYTRACING
 
-FRayTracingShaderRHIRef FD3D12DynamicRHI::RHICreateRayTracingShader(TArrayView<const uint8> Code, const FSHAHash& Hash, EShaderFrequency ShaderFrequency)
+FRayTracingShaderRHIRef FD3D12DynamicRHI::RHICreateRayTracingShader(const FRHICreateShaderDesc& CreateShaderDesc, EShaderFrequency ShaderFrequency)
 {
 	checkf(GRHISupportsRayTracing && GRHISupportsRayTracingShaders, TEXT("Tried to create RayTracing shader but RHI doesn't support it!"));
 
@@ -281,7 +281,7 @@ FRayTracingShaderRHIRef FD3D12DynamicRHI::RHICreateRayTracingShader(TArrayView<c
 		}
 	};
 
-	FD3D12RayTracingShader* Shader = InitStandardShaderWithCustomSerialization(new FD3D12RayTracingShader(ShaderFrequency), Code, CustomSerialization);
+	FD3D12RayTracingShader* Shader = InitStandardShaderWithCustomSerialization(new FD3D12RayTracingShader(ShaderFrequency), CreateShaderDesc.Code, CustomSerialization);
 
 	// Setup local root signature (RayGen only has global root signature)
 	if (Shader && ShaderFrequency != SF_RayGen)
diff --git a/Engine/Source/Runtime/NullDrv/Public/NullRHI.h b/Engine/Source/Runtime/NullDrv/Public/NullRHI.h
--- a/Engine/Source/Runtime/NullDrv/Public/NullRHI.h
+++ b/Engine/Source/Runtime/NullDrv/Public/NullRHI.h
@@ -45,22 +45,22 @@ public:
 		return new FRHIVertexDeclaration(); 
 	}
 
-	virtual FPixelShaderRHIRef RHICreatePixelShader(TArrayView<const uint8> Code, const FSHAHash& Hash) final override
+	virtual FPixelShaderRHIRef RHICreatePixelShader(const FRHICreateShaderDesc& CreateShaderDesc) final override
 	{ 
 		return new FRHIPixelShader(); 
 	}
 
-	virtual FVertexShaderRHIRef RHICreateVertexShader(TArrayView<const uint8> Code, const FSHAHash& Hash) final override
+	virtual FVertexShaderRHIRef RHICreateVertexShader(const FRHICreateShaderDesc& CreateShaderDesc) final override
 	{ 
 		return new FRHIVertexShader(); 
 	}
 
-	virtual FGeometryShaderRHIRef RHICreateGeometryShader(TArrayView<const uint8> Code, const FSHAHash& Hash) final override
+	virtual FGeometryShaderRHIRef RHICreateGeometryShader(const FRHICreateShaderDesc& CreateShaderDesc) final override
 	{ 
 		return new FRHIGeometryShader(); 
 	}
 
-	virtual FComputeShaderRHIRef RHICreateComputeShader(TArrayView<const uint8> Code, const FSHAHash& Hash) final override
+	virtual FComputeShaderRHIRef RHICreateComputeShader(const FRHICreateShaderDesc& CreateShaderDesc) final override
 	{ 
 		return new FRHIComputeShader(); 
 	}
diff --git a/Engine/Source/Runtime/OpenGLDrv/Private/OpenGLDrv.h b/Engine/Source/Runtime/OpenGLDrv/Private/OpenGLDrv.h
--- a/Engine/Source/Runtime/OpenGLDrv/Private/OpenGLDrv.h
+++ b/Engine/Source/Runtime/OpenGLDrv/Private/OpenGLDrv.h
@@ -114,10 +114,10 @@ public:
 	UE_API virtual FBlendStateRHIRef RHICreateBlendState(const FBlendStateInitializerRHI& Initializer) final override;
 	UE_API virtual FVertexDeclarationRHIRef RHICreateVertexDeclaration(const FVertexDeclarationElementList& Elements) final override;
 
-	UE_API virtual FPixelShaderRHIRef RHICreatePixelShader(TArrayView<const uint8> Code, const FSHAHash& Hash) final override;
-	UE_API virtual FVertexShaderRHIRef RHICreateVertexShader(TArrayView<const uint8> Code, const FSHAHash& Hash) final override;
-	UE_API virtual FGeometryShaderRHIRef RHICreateGeometryShader(TArrayView<const uint8> Code, const FSHAHash& Hash) final override;
-	UE_API virtual FComputeShaderRHIRef RHICreateComputeShader(TArrayView<const uint8> Code, const FSHAHash& Hash) final override;
+	UE_API virtual FPixelShaderRHIRef RHICreatePixelShader(const FRHICreateShaderDesc& CreateShaderDesc) final override;
+	UE_API virtual FVertexShaderRHIRef RHICreateVertexShader(const FRHICreateShaderDesc& CreateShaderDesc) final override;
+	UE_API virtual FGeometryShaderRHIRef RHICreateGeometryShader(const FRHICreateShaderDesc& CreateShaderDesc) final override;
+	UE_API virtual FComputeShaderRHIRef RHICreateComputeShader(const FRHICreateShaderDesc& CreateShaderDesc) final override;
 
 	// SRV / UAV creation functions
 	UE_API virtual FShaderResourceViewRHIRef  RHICreateShaderResourceView (class FRHICommandListBase& RHICmdList, FRHIViewableResource* Resource, FRHIViewDesc const& ViewDesc) final override;
diff --git a/Engine/Source/Runtime/OpenGLDrv/Private/OpenGLShaders.cpp b/Engine/Source/Runtime/OpenGLDrv/Private/OpenGLShaders.cpp
--- a/Engine/Source/Runtime/OpenGLDrv/Private/OpenGLShaders.cpp
+++ b/Engine/Source/Runtime/OpenGLDrv/Private/OpenGLShaders.cpp
@@ -981,7 +981,7 @@ void OPENGLDRV_API GLSLToDeviceCompatibleGLSL(FAnsiCharArray& GlslCodeOriginal,
 	}
 }
 
-FOpenGLShader::FOpenGLShader(TArrayView<const uint8> Code, const FSHAHash& Hash, GLenum TypeEnum, FShaderResourceTable& SRT, FRHIShader* RHIShader)
+FOpenGLShader::FOpenGLShader(TArrayView<const uint8> Code, GLenum TypeEnum, FShaderResourceTable& SRT, FRHIShader* RHIShader)
 {
 	FMemory::Memzero(&Bindings, sizeof(Bindings));
 
@@ -1079,8 +1079,8 @@ FOpenGLShader::FOpenGLShader(TArrayView<const uint8> Code, const FSHAHash& Hash,
 	// The shader is compiled when we link program
 }
 
-FOpenGLVertexShader::FOpenGLVertexShader(TArrayView<const uint8> Code, const FSHAHash& Hash)
-	: FOpenGLShader(Code, Hash, GL_VERTEX_SHADER, ShaderResourceTable, this)
+FOpenGLVertexShader::FOpenGLVertexShader(TArrayView<const uint8> Code)
+	: FOpenGLShader(Code, GL_VERTEX_SHADER, ShaderResourceTable, this)
 {}
 
 void FOpenGLVertexShader::ConditionalyCompile()
@@ -1091,8 +1091,8 @@ void FOpenGLVertexShader::ConditionalyCompile()
 	}
 }
 
-FOpenGLPixelShader::FOpenGLPixelShader(TArrayView<const uint8> Code, const FSHAHash& Hash)
-	: FOpenGLShader(Code, Hash, GL_FRAGMENT_SHADER, ShaderResourceTable, this)
+FOpenGLPixelShader::FOpenGLPixelShader(TArrayView<const uint8> Code)
+	: FOpenGLShader(Code, GL_FRAGMENT_SHADER, ShaderResourceTable, this)
 {}
 
 void FOpenGLPixelShader::ConditionalyCompile()
@@ -1103,8 +1103,8 @@ void FOpenGLPixelShader::ConditionalyCompile()
 	}
 }
 
-FOpenGLGeometryShader::FOpenGLGeometryShader(TArrayView<const uint8> Code, const FSHAHash& Hash)
-	: FOpenGLShader(Code, Hash, GL_GEOMETRY_SHADER, ShaderResourceTable, this)
+FOpenGLGeometryShader::FOpenGLGeometryShader(TArrayView<const uint8> Code)
+	: FOpenGLShader(Code, GL_GEOMETRY_SHADER, ShaderResourceTable, this)
 {}
 
 void FOpenGLGeometryShader::ConditionalyCompile()
@@ -1115,8 +1115,8 @@ void FOpenGLGeometryShader::ConditionalyCompile()
 	}
 }
 
-FOpenGLComputeShader::FOpenGLComputeShader(TArrayView<const uint8> Code, const FSHAHash& Hash)
-	: FOpenGLShader(Code, Hash, GL_COMPUTE_SHADER, ShaderResourceTable, this)
+FOpenGLComputeShader::FOpenGLComputeShader(TArrayView<const uint8> Code)
+	: FOpenGLShader(Code, GL_COMPUTE_SHADER, ShaderResourceTable, this)
 {}
 
 void FOpenGLComputeShader::ConditionalyCompile()
@@ -1127,19 +1127,19 @@ void FOpenGLComputeShader::ConditionalyCompile()
 	}
 }
 
-FVertexShaderRHIRef FOpenGLDynamicRHI::RHICreateVertexShader(TArrayView<const uint8> Code, const FSHAHash& Hash)
+FVertexShaderRHIRef FOpenGLDynamicRHI::RHICreateVertexShader(const FRHICreateShaderDesc& CreateShaderDesc)
 {
-	return new FOpenGLVertexShader(Code, Hash);
+	return new FOpenGLVertexShader(CreateShaderDesc.Code);
 }
 
-FPixelShaderRHIRef FOpenGLDynamicRHI::RHICreatePixelShader(TArrayView<const uint8> Code, const FSHAHash& Hash)
+FPixelShaderRHIRef FOpenGLDynamicRHI::RHICreatePixelShader(const FRHICreateShaderDesc& CreateShaderDesc)
 {
-	return new FOpenGLPixelShader(Code, Hash);
+	return new FOpenGLPixelShader(CreateShaderDesc.Code);
 }
 
-FGeometryShaderRHIRef FOpenGLDynamicRHI::RHICreateGeometryShader(TArrayView<const uint8> Code, const FSHAHash& Hash)
+FGeometryShaderRHIRef FOpenGLDynamicRHI::RHICreateGeometryShader(const FRHICreateShaderDesc& CreateShaderDesc)
 {
-	return new FOpenGLGeometryShader(Code, Hash);
+	return new FOpenGLGeometryShader(CreateShaderDesc.Code);
 }
 
 static void MarkShaderParameterCachesDirty(FOpenGLShaderParameterCache* ShaderParameters, bool UpdateCompute)
@@ -2615,9 +2615,9 @@ void FOpenGLDynamicRHI::LinkComputeProgram(FRHIComputeShader* ComputeShaderRHI)
 	}
 }
 
-FComputeShaderRHIRef FOpenGLDynamicRHI::RHICreateComputeShader(TArrayView<const uint8> Code, const FSHAHash& Hash)
+FComputeShaderRHIRef FOpenGLDynamicRHI::RHICreateComputeShader(const FRHICreateShaderDesc& CreateShaderDesc)
 {
-	return new FOpenGLComputeShader(Code, Hash);
+	return new FOpenGLComputeShader(CreateShaderDesc.Code);
 }
 
 template<class TOpenGLStage>
diff --git a/Engine/Source/Runtime/OpenGLDrv/Public/OpenGLShaderResources.h b/Engine/Source/Runtime/OpenGLDrv/Public/OpenGLShaderResources.h
--- a/Engine/Source/Runtime/OpenGLDrv/Public/OpenGLShaderResources.h
+++ b/Engine/Source/Runtime/OpenGLDrv/Public/OpenGLShaderResources.h
@@ -214,7 +214,7 @@ public:
 	const ANSICHAR*  GlslCodeString; // make it easier in VS to see shader code in debug mode; points to begin of GlslCode
 #endif
 
-	FOpenGLShader(TArrayView<const uint8> Code, const FSHAHash& Hash, GLenum TypeEnum, FShaderResourceTable& SRT, FRHIShader* RHIShader);
+	FOpenGLShader(TArrayView<const uint8> Code, GLenum TypeEnum, FShaderResourceTable& SRT, FRHIShader* RHIShader);
 
 	~FOpenGLShader() = default;
 
@@ -227,7 +227,7 @@ class FOpenGLVertexShader : public FRHIVertexShader, public FOpenGLShader
 public:
 	static constexpr EShaderFrequency Frequency = SF_Vertex;
 
-	FOpenGLVertexShader(TArrayView<const uint8> Code, const FSHAHash& Hash);
+	FOpenGLVertexShader(TArrayView<const uint8> Code);
 
 	void ConditionalyCompile();
 };
@@ -237,7 +237,7 @@ class FOpenGLPixelShader : public FRHIPixelShader, public FOpenGLShader
 public:
 	static constexpr EShaderFrequency Frequency = SF_Pixel;
 
-	FOpenGLPixelShader(TArrayView<const uint8> Code, const FSHAHash& Hash);
+	FOpenGLPixelShader(TArrayView<const uint8> Code);
 
 	void ConditionalyCompile();
 };
@@ -247,7 +247,7 @@ class FOpenGLGeometryShader : public FRHIGeometryShader, public FOpenGLShader
 public:
 	static constexpr EShaderFrequency Frequency = SF_Geometry;
 
-	FOpenGLGeometryShader(TArrayView<const uint8> Code, const FSHAHash& Hash);
+	FOpenGLGeometryShader(TArrayView<const uint8> Code);
 
 	void ConditionalyCompile();
 };
@@ -257,7 +257,7 @@ class FOpenGLComputeShader : public FRHIComputeShader, public FOpenGLShader
 public:
 	static constexpr EShaderFrequency Frequency = SF_Compute;
 
-	FOpenGLComputeShader(TArrayView<const uint8> Code, const FSHAHash& Hash);
+	FOpenGLComputeShader(TArrayView<const uint8> Code);
 
 	void ConditionalyCompile();
 
diff --git a/Engine/Source/Runtime/RHI/Public/DynamicRHI.h b/Engine/Source/Runtime/RHI/Public/DynamicRHI.h
--- a/Engine/Source/Runtime/RHI/Public/DynamicRHI.h
+++ b/Engine/Source/Runtime/RHI/Public/DynamicRHI.h
@@ -244,22 +244,22 @@ public:
 	virtual FVertexDeclarationRHIRef RHICreateVertexDeclaration(const FVertexDeclarationElementList& Elements) = 0;
 
 	// FlushType: Wait RHI Thread
-	virtual FPixelShaderRHIRef RHICreatePixelShader(TArrayView<const uint8> Code, const FSHAHash& Hash) = 0;
+	virtual FPixelShaderRHIRef RHICreatePixelShader(const FRHICreateShaderDesc& CreateShaderDesc) = 0;
 
 	// FlushType: Wait RHI Thread
-	virtual FVertexShaderRHIRef RHICreateVertexShader(TArrayView<const uint8> Code, const FSHAHash& Hash) = 0;
+	virtual FVertexShaderRHIRef RHICreateVertexShader(const FRHICreateShaderDesc& CreateShaderDesc) = 0;
 
 	// FlushType: Wait RHI Thread
-	virtual FGeometryShaderRHIRef RHICreateGeometryShader(TArrayView<const uint8> Code, const FSHAHash& Hash) = 0;
+	virtual FGeometryShaderRHIRef RHICreateGeometryShader(const FRHICreateShaderDesc& CreateShaderDesc) = 0;
 
 	// FlushType: Wait RHI Thread
-	virtual FMeshShaderRHIRef RHICreateMeshShader(TArrayView<const uint8> Code, const FSHAHash& Hash)
+	virtual FMeshShaderRHIRef RHICreateMeshShader(const FRHICreateShaderDesc& CreateShaderDesc)
 	{
 		return FMeshShaderRHIRef();
 	}
 
 	// FlushType: Wait RHI Thread
-	virtual FAmplificationShaderRHIRef RHICreateAmplificationShader(TArrayView<const uint8> Code, const FSHAHash& Hash)
+	virtual FAmplificationShaderRHIRef RHICreateAmplificationShader(const FRHICreateShaderDesc& CreateShaderDesc)
 	{
 		return FAmplificationShaderRHIRef();
 	}
@@ -268,10 +268,10 @@ public:
 	virtual void FlushPendingLogs() {}
 
 	// FlushType: Wait RHI Thread
-	virtual FComputeShaderRHIRef RHICreateComputeShader(TArrayView<const uint8> Code, const FSHAHash& Hash) = 0;
+	virtual FComputeShaderRHIRef RHICreateComputeShader(const FRHICreateShaderDesc& CreateShaderDesc) = 0;
 
 	// FlushType: Wait RHI Thread
-	virtual FWorkGraphShaderRHIRef RHICreateWorkGraphShader(TArrayView<const uint8> Code, const FSHAHash& Hash, EShaderFrequency ShaderFrequency)
+	virtual FWorkGraphShaderRHIRef RHICreateWorkGraphShader(const FRHICreateShaderDesc& CreateShaderDesc, EShaderFrequency ShaderFrequency)
 	{
 		return FWorkGraphShaderRHIRef();
 	}
@@ -962,7 +962,7 @@ public:
 		return nullptr;
 	}
 
-	virtual FRayTracingShaderRHIRef RHICreateRayTracingShader(TArrayView<const uint8> Code, const FSHAHash& Hash, EShaderFrequency ShaderFrequency)
+	virtual FRayTracingShaderRHIRef RHICreateRayTracingShader(const FRHICreateShaderDesc& CreateShaderDesc, EShaderFrequency ShaderFrequency)
 	{
 		checkNoEntry();
 		return nullptr;
@@ -1038,46 +1038,46 @@ inline TRHI* GetDynamicRHI()
 	return CastDynamicRHI<TRHI>(GDynamicRHI);
 }
 
-inline FPixelShaderRHIRef RHICreatePixelShader(TArrayView<const uint8> Code, const FSHAHash& Hash)
+inline FPixelShaderRHIRef RHICreatePixelShader(const FRHICreateShaderDesc& CreateShaderDesc)
 {
 	LLM_SCOPE(ELLMTag::Shaders);
-	return GDynamicRHI->RHICreatePixelShader(Code, Hash);
+	return GDynamicRHI->RHICreatePixelShader(CreateShaderDesc);
 }
 
-inline FVertexShaderRHIRef RHICreateVertexShader(TArrayView<const uint8> Code, const FSHAHash& Hash)
+inline FVertexShaderRHIRef RHICreateVertexShader(const FRHICreateShaderDesc& CreateShaderDesc)
 {
 	LLM_SCOPE(ELLMTag::Shaders);
-	return GDynamicRHI->RHICreateVertexShader(Code, Hash);
+	return GDynamicRHI->RHICreateVertexShader(CreateShaderDesc);
 }
 
-inline FMeshShaderRHIRef RHICreateMeshShader(TArrayView<const uint8> Code, const FSHAHash& Hash)
+inline FMeshShaderRHIRef RHICreateMeshShader(const FRHICreateShaderDesc& CreateShaderDesc)
 {
 	LLM_SCOPE(ELLMTag::Shaders);
-	return GDynamicRHI->RHICreateMeshShader(Code, Hash);
+	return GDynamicRHI->RHICreateMeshShader(CreateShaderDesc);
 }
 
-inline FAmplificationShaderRHIRef RHICreateAmplificationShader(TArrayView<const uint8> Code, const FSHAHash& Hash)
+inline FAmplificationShaderRHIRef RHICreateAmplificationShader(const FRHICreateShaderDesc& CreateShaderDesc)
 {
 	LLM_SCOPE(ELLMTag::Shaders);
-	return GDynamicRHI->RHICreateAmplificationShader(Code, Hash);
+	return GDynamicRHI->RHICreateAmplificationShader(CreateShaderDesc);
 }
 
-inline FGeometryShaderRHIRef RHICreateGeometryShader(TArrayView<const uint8> Code, const FSHAHash& Hash)
+inline FGeometryShaderRHIRef RHICreateGeometryShader(const FRHICreateShaderDesc& CreateShaderDesc)
 {
 	LLM_SCOPE(ELLMTag::Shaders);
-	return GDynamicRHI->RHICreateGeometryShader(Code, Hash);
+	return GDynamicRHI->RHICreateGeometryShader(CreateShaderDesc);
 }
 
-inline FComputeShaderRHIRef RHICreateComputeShader(TArrayView<const uint8> Code, const FSHAHash& Hash)
+inline FComputeShaderRHIRef RHICreateComputeShader(const FRHICreateShaderDesc& CreateShaderDesc)
 {
 	LLM_SCOPE(ELLMTag::Shaders);
-	return GDynamicRHI->RHICreateComputeShader(Code, Hash);
+	return GDynamicRHI->RHICreateComputeShader(CreateShaderDesc);
 }
 
-inline FWorkGraphShaderRHIRef RHICreateWorkGraphShader(TArrayView<const uint8> Code, const FSHAHash& Hash, EShaderFrequency ShaderFrequency)
+inline FWorkGraphShaderRHIRef RHICreateWorkGraphShader(const FRHICreateShaderDesc& CreateShaderDesc, EShaderFrequency ShaderFrequency)
 {
 	LLM_SCOPE(ELLMTag::Shaders);
-	return GDynamicRHI->RHICreateWorkGraphShader(Code, Hash, ShaderFrequency);
+	return GDynamicRHI->RHICreateWorkGraphShader(CreateShaderDesc, ShaderFrequency);
 }
 
 inline FGPUFenceRHIRef RHICreateGPUFence(const FName& Name)
@@ -1370,9 +1370,49 @@ inline FRayTracingSceneRHIRef RHICreateRayTracingScene(FRayTracingSceneInitializ
 	return GDynamicRHI->RHICreateRayTracingScene(MoveTemp(Initializer));
 }
 
+inline FRayTracingShaderRHIRef RHICreateRayTracingShader(const FRHICreateShaderDesc& CreateShaderDesc, EShaderFrequency ShaderFrequency)
+{
+	return GDynamicRHI->RHICreateRayTracingShader(CreateShaderDesc, ShaderFrequency);
+}
+
+inline FPixelShaderRHIRef RHICreatePixelShader(TArrayView<const uint8> Code, const FSHAHash& Hash)
+{
+	return RHICreatePixelShader(FRHICreateShaderDesc(Code));
+}
+
+inline FVertexShaderRHIRef RHICreateVertexShader(TArrayView<const uint8> Code, const FSHAHash& Hash)
+{
+	return RHICreateVertexShader(FRHICreateShaderDesc(Code));
+}
+
+inline FMeshShaderRHIRef RHICreateMeshShader(TArrayView<const uint8> Code, const FSHAHash& Hash)
+{
+	return RHICreateMeshShader(FRHICreateShaderDesc(Code));
+}
+
+inline FAmplificationShaderRHIRef RHICreateAmplificationShader(TArrayView<const uint8> Code, const FSHAHash& Hash)
+{
+	return RHICreateAmplificationShader(FRHICreateShaderDesc(Code));
+}
+
+inline FGeometryShaderRHIRef RHICreateGeometryShader(TArrayView<const uint8> Code, const FSHAHash& Hash)
+{
+	return RHICreateGeometryShader(FRHICreateShaderDesc(Code));
+}
+
+inline FComputeShaderRHIRef RHICreateComputeShader(TArrayView<const uint8> Code, const FSHAHash& Hash)
+{
+	return RHICreateComputeShader(FRHICreateShaderDesc(Code));
+}
+
+inline FWorkGraphShaderRHIRef RHICreateWorkGraphShader(TArrayView<const uint8> Code, const FSHAHash& Hash, EShaderFrequency ShaderFrequency)
+{
+	return RHICreateWorkGraphShader(FRHICreateShaderDesc(Code), ShaderFrequency);
+}
+
 inline FRayTracingShaderRHIRef RHICreateRayTracingShader(TArrayView<const uint8> Code, const FSHAHash& Hash, EShaderFrequency ShaderFrequency)
 {
-	return GDynamicRHI->RHICreateRayTracingShader(Code, Hash, ShaderFrequency);
+	return RHICreateRayTracingShader(FRHICreateShaderDesc(Code), ShaderFrequency);
 }
 
 inline FShaderBundleRHIRef RHICreateShaderBundle(const FShaderBundleCreateInfo& CreateInfo)
diff --git a/Engine/Source/Runtime/RHI/Public/RHIResources.h b/Engine/Source/Runtime/RHI/Public/RHIResources.h
--- a/Engine/Source/Runtime/RHI/Public/RHIResources.h
+++ b/Engine/Source/Runtime/RHI/Public/RHIResources.h
@@ -785,6 +785,28 @@ namespace UE
 	}
 }
 
+struct FRHICreateShaderDesc
+{
+	TConstArrayView<uint8> Code;
+	// Optional Owner of the memory backing the View
+	TSharedPtr<void> Owner;
+	
+	FRHICreateShaderDesc(TConstArrayView<uint8> InCode)
+	: Code(InCode)
+	{
+	}
+	
+	template<typename DeleterType>
+	FRHICreateShaderDesc(TConstArrayView<uint8> InCode, DeleterType&& InDeleter)
+	 : Code(InCode)
+	// Giving the view data here serve two purposes:
+	// - Hints on the intent that this ownership correspond to the data behind the View
+	// - Make the Owner valid to differentiate between data that stay valid as long as we keep a reference to it versus temporary data
+	 , Owner(MakeShareable((void*)InCode.GetData(), Forward<DeleterType>(InDeleter)))
+	{
+	}
+};
+
 class FRHIShaderData
 {
 public:
diff --git a/Engine/Source/Runtime/RHI/Public/RHIValidation.h b/Engine/Source/Runtime/RHI/Public/RHIValidation.h
--- a/Engine/Source/Runtime/RHI/Public/RHIValidation.h
+++ b/Engine/Source/Runtime/RHI/Public/RHIValidation.h
@@ -194,36 +194,36 @@ public:
 	}
 
 	// FlushType: Wait RHI Thread
-	virtual FPixelShaderRHIRef RHICreatePixelShader(TArrayView<const uint8> Code, const FSHAHash& Hash) override final
+	virtual FPixelShaderRHIRef RHICreatePixelShader(const FRHICreateShaderDesc& CreateShaderDesc) override final
 	{
-		return RHI->RHICreatePixelShader(Code, Hash);
+		return RHI->RHICreatePixelShader(CreateShaderDesc);
 	}
 
 	// FlushType: Wait RHI Thread
-	virtual FVertexShaderRHIRef RHICreateVertexShader(TArrayView<const uint8> Code, const FSHAHash& Hash) override final
+	virtual FVertexShaderRHIRef RHICreateVertexShader(const FRHICreateShaderDesc& CreateShaderDesc) override final
 	{
-		return RHI->RHICreateVertexShader(Code, Hash);
+		return RHI->RHICreateVertexShader(CreateShaderDesc);
 	}
 
 	// FlushType: Wait RHI Thread
-	virtual FGeometryShaderRHIRef RHICreateGeometryShader(TArrayView<const uint8> Code, const FSHAHash& Hash) override final
+	virtual FGeometryShaderRHIRef RHICreateGeometryShader(const FRHICreateShaderDesc& CreateShaderDesc) override final
 	{
 		check(RHISupportsGeometryShaders(GMaxRHIShaderPlatform));
-		return RHI->RHICreateGeometryShader(Code, Hash);
+		return RHI->RHICreateGeometryShader(CreateShaderDesc);
 	}
 
 	// FlushType: Wait RHI Thread
-	virtual FMeshShaderRHIRef RHICreateMeshShader(TArrayView<const uint8> Code, const FSHAHash& Hash)
+	virtual FMeshShaderRHIRef RHICreateMeshShader(const FRHICreateShaderDesc& CreateShaderDesc)
 	{
 		check(RHISupportsMeshShadersTier0(GMaxRHIShaderPlatform));
-		return RHI->RHICreateMeshShader(Code, Hash);
+		return RHI->RHICreateMeshShader(CreateShaderDesc);
 	}
 
 	// FlushType: Wait RHI Thread
-	virtual FAmplificationShaderRHIRef RHICreateAmplificationShader(TArrayView<const uint8> Code, const FSHAHash& Hash)
+	virtual FAmplificationShaderRHIRef RHICreateAmplificationShader(const FRHICreateShaderDesc& CreateShaderDesc)
 	{
 		check(RHISupportsMeshShadersTier0(GMaxRHIShaderPlatform));
-		return RHI->RHICreateAmplificationShader(Code, Hash);
+		return RHI->RHICreateAmplificationShader(CreateShaderDesc);
 	}
 
 	// Some RHIs can have pending messages/logs for error tracking, or debug modes
@@ -233,9 +233,9 @@ public:
 	}
 
 	// FlushType: Wait RHI Thread
-	virtual FComputeShaderRHIRef RHICreateComputeShader(TArrayView<const uint8> Code, const FSHAHash& Hash) override final
+	virtual FComputeShaderRHIRef RHICreateComputeShader(const FRHICreateShaderDesc& CreateShaderDesc) override final
 	{
-		return RHI->RHICreateComputeShader(Code, Hash);
+		return RHI->RHICreateComputeShader(CreateShaderDesc);
 	}
 
 	/**
@@ -975,9 +975,9 @@ public:
 		return Result;
 	}
 
-	virtual FRayTracingShaderRHIRef RHICreateRayTracingShader(TArrayView<const uint8> Code, const FSHAHash& Hash, EShaderFrequency ShaderFrequency) override final
+	virtual FRayTracingShaderRHIRef RHICreateRayTracingShader(const FRHICreateShaderDesc& CreateShaderDesc, EShaderFrequency ShaderFrequency) override final
 	{
-		return RHI->RHICreateRayTracingShader(Code, Hash, ShaderFrequency);
+		return RHI->RHICreateRayTracingShader(CreateShaderDesc, ShaderFrequency);
 	}
 
 	virtual FRayTracingPipelineStateRHIRef RHICreateRayTracingPipelineState(const FRayTracingPipelineStateInitializer& Initializer) override final
diff --git a/Engine/Source/Runtime/RenderCore/Private/ShaderLibrary/ShaderCodeArchive.cpp b/Engine/Source/Runtime/RenderCore/Private/ShaderLibrary/ShaderCodeArchive.cpp
--- a/Engine/Source/Runtime/RenderCore/Private/ShaderLibrary/ShaderCodeArchive.cpp
+++ b/Engine/Source/Runtime/RenderCore/Private/ShaderLibrary/ShaderCodeArchive.cpp
@@ -3,6 +3,9 @@
 #include "ShaderCodeArchive.h"
 
 #include "Async/ParallelFor.h"
+#if USE_MMAPPED_SHADERARCHIVE
+#include "Async/MappedFileHandle.h"
+#endif
 #include "Compression/OodleDataCompression.h"
 #include "DataDrivenShaderPlatformInfo.h"
 #include "Misc/Compression.h"
@@ -1616,23 +1619,23 @@ TRefCountPtr<FRHIShader> FShaderCodeArchive::CreateShader(int32 Index, bool bReq
 	// detect the breach of contract early
 	ensureAlwaysMsgf(IsInRenderingThread() || GRHISupportsMultithreadedShaderCreation, TEXT("More than one thread is creating shaders, but GRHISupportsMultithreadedShaderCreation is false."));
 
-	const auto ShaderCodeView = MakeArrayView(ShaderCode, ShaderEntry.GetUncompressedSize());
-	const FSHAHash& ShaderHash = SerializedShaders.GetShaderHashes()[Index];
+	FRHICreateShaderDesc CreateShaderDesc = FRHICreateShaderDesc(MakeArrayView(ShaderCode, ShaderEntry.GetUncompressedSize()));
+
 	switch (ShaderEntry.GetFrequency())
 	{
-	case SF_Vertex: Shader = RHICreateVertexShader(ShaderCodeView, ShaderHash); CheckShaderCreation(Shader, Index); break;
-	case SF_Mesh: Shader = RHICreateMeshShader(ShaderCodeView, ShaderHash); CheckShaderCreation(Shader, Index); break;
-	case SF_Amplification: Shader = RHICreateAmplificationShader(ShaderCodeView, ShaderHash); CheckShaderCreation(Shader, Index); break;
-	case SF_Pixel: Shader = RHICreatePixelShader(ShaderCodeView, ShaderHash); CheckShaderCreation(Shader, Index); break;
-	case SF_Geometry: Shader = RHICreateGeometryShader(ShaderCodeView, ShaderHash); CheckShaderCreation(Shader, Index); break;
-	case SF_Compute: Shader = RHICreateComputeShader(ShaderCodeView, ShaderHash); CheckShaderCreation(Shader, Index); break;
-	case SF_WorkGraphRoot: Shader = RHICreateWorkGraphShader(ShaderCodeView, ShaderHash, SF_WorkGraphRoot); CheckShaderCreation(Shader, Index); break;
-	case SF_WorkGraphComputeNode: Shader = RHICreateWorkGraphShader(ShaderCodeView, ShaderHash, SF_WorkGraphComputeNode); CheckShaderCreation(Shader, Index); break;
+	case SF_Vertex: Shader = RHICreateVertexShader(CreateShaderDesc); CheckShaderCreation(Shader, Index); break;
+	case SF_Mesh: Shader = RHICreateMeshShader(CreateShaderDesc); CheckShaderCreation(Shader, Index); break;
+	case SF_Amplification: Shader = RHICreateAmplificationShader(CreateShaderDesc); CheckShaderCreation(Shader, Index); break;
+	case SF_Pixel: Shader = RHICreatePixelShader(CreateShaderDesc); CheckShaderCreation(Shader, Index); break;
+	case SF_Geometry: Shader = RHICreateGeometryShader(CreateShaderDesc); CheckShaderCreation(Shader, Index); break;
+	case SF_Compute: Shader = RHICreateComputeShader(CreateShaderDesc); CheckShaderCreation(Shader, Index); break;
+	case SF_WorkGraphRoot: Shader = RHICreateWorkGraphShader(CreateShaderDesc, SF_WorkGraphRoot); CheckShaderCreation(Shader, Index); break;
+	case SF_WorkGraphComputeNode: Shader = RHICreateWorkGraphShader(CreateShaderDesc, SF_WorkGraphComputeNode); CheckShaderCreation(Shader, Index); break;
 	case SF_RayGen: case SF_RayMiss: case SF_RayHitGroup: case SF_RayCallable:
 #if RHI_RAYTRACING
 		if (GRHISupportsRayTracing && GRHISupportsRayTracingShaders)
 		{
-			Shader = RHICreateRayTracingShader(ShaderCodeView, ShaderHash, ShaderEntry.GetFrequency());
+			Shader = RHICreateRayTracingShader(CreateShaderDesc, ShaderEntry.GetFrequency());
 			CheckShaderCreation(Shader, Index);
 		}
 #endif // RHI_RAYTRACING
@@ -1652,8 +1655,7 @@ TRefCountPtr<FRHIShader> FShaderCodeArchive::CreateShader(int32 Index, bool bReq
 		TCsvPersistentCustomStat<int>* CsvStatNumShadersCreated = FCsvProfiler::Get()->GetOrCreatePersistentCustomStatInt(TEXT("NumShadersCreated"), CSV_CATEGORY_INDEX(Shaders));
 		CsvStatNumShadersCreated->Add(1);
 #endif
-
-		Shader->SetHash(ShaderHash);
+		Shader->SetHash(SerializedShaders.GetShaderHashes()[Index]);
 	}
 
 	return Shader;
@@ -2521,17 +2523,38 @@ bool FIoStoreShaderCodeArchive::PreloadShaderGroup(int32 ShaderGroupIndex, FGrap
 			// only global shaders and retries of asset shaders are going to hit this path, regular asset shaders will be preloaded with the package
 			if (UNLIKELY(AttachShaderReadRequestFuncPtr == nullptr))
 			{
-				FIoBatch IoBatch = IoDispatcher.NewBatch();
-				PreloadEntry.IoRequest = IoBatch.Read(Header.ShaderGroupIoHashes[ShaderGroupIndex], FIoReadOptions(), GetShaderCodeArchivePriority());
-				IoBatch.IssueAndDispatchSubsequents(PreloadEntry.PreloadEvent);
+#if USE_MMAPPED_SHADERARCHIVE
+				PreloadEntry.MappedRegionStatus = IoDispatcher.OpenMapped(Header.ShaderGroupIoHashes[ShaderGroupIndex], FIoReadOptions());
+				if (PreloadEntry.MappedRegionStatus.IsOk())
+				{
+					PreloadEntry.PreloadEvent->DispatchSubsequents();
+					PreloadEntry.PreloadEvent.SafeRelease();
+				}
+				else
+#endif
+				{
+					FIoBatch IoBatch = IoDispatcher.NewBatch();
+					PreloadEntry.IoRequest = IoBatch.Read(Header.ShaderGroupIoHashes[ShaderGroupIndex], FIoReadOptions(), GetShaderCodeArchivePriority());
+					IoBatch.IssueAndDispatchSubsequents(PreloadEntry.PreloadEvent);
+				}
 			}
 			else
 			{
-				PreloadEntry.IoRequest = (*AttachShaderReadRequestFuncPtr)(Header.ShaderGroupIoHashes[ShaderGroupIndex], PreloadEntry.PreloadEvent);
+				FCoreDelegates::FShaderReadRequestResult Result = (*AttachShaderReadRequestFuncPtr)(Header.ShaderGroupIoHashes[ShaderGroupIndex], PreloadEntry.PreloadEvent);
+#if USE_MMAPPED_SHADERARCHIVE
+				PreloadEntry.IoRequest = Result.Key;
+				PreloadEntry.MappedRegionStatus = Result.Value;
+#else
+				PreloadEntry.IoRequest = Result;
+#endif
 			}
 		};
 
-	if (PreloadEntry.IoRequest.Status() == FIoStatus::Invalid)
+	if (PreloadEntry.IoRequest.Status() == FIoStatus::Invalid
+#if USE_MMAPPED_SHADERARCHIVE
+		&& !PreloadEntry.MappedRegionStatus.IsOk()
+#endif
+		)
 	{
 		KickOffIoRequest();
 
@@ -2543,7 +2566,11 @@ bool FIoStoreShaderCodeArchive::PreloadShaderGroup(int32 ShaderGroupIndex, FGrap
 		CsvStatPreloadedShaderMB->Add((float)ShaderGroupSize / (1024.0f * 1024.0f));
 #endif
 	}
-	else if (PreloadEntry.PreloadEvent.IsValid() && PreloadEntry.PreloadEvent->IsCompleted() && PreloadEntry.IoRequest.GetResult() == nullptr)
+	else if (PreloadEntry.PreloadEvent.IsValid() && PreloadEntry.PreloadEvent->IsCompleted() && PreloadEntry.IoRequest.GetResult() == nullptr
+#if USE_MMAPPED_SHADERARCHIVE
+		&& !PreloadEntry.MappedRegionStatus.IsOk()
+#endif
+			 )
 	{
 		// failed request, re-kick in-place, without changing stats
 		PreloadEntry.IoRequest.Cancel();
@@ -2603,7 +2630,11 @@ void FIoStoreShaderCodeArchive::AddRefPreloadedShaderGroup(int32 ShaderGroupInde
 
 void FIoStoreShaderCodeArchive::ReleasePreloadedShaderGroup(int32 ShaderGroupIndex)
 {
-	ReleasePreloadEntry(ShaderGroupIndex);
+	ReleasePreloadEntry(ShaderGroupIndex
+#if UE_SCA_DEBUG_PRELOADING
+		, FString::Printf(TEXT("ReleasePreloadedShaderGroup %d"), ShaderGroupIndex)
+#endif
+						);
 }
 
 bool FIoStoreShaderCodeArchive::IsPreloading(int32 ShaderIndex, FGraphEventArray& OutCompletionEvents)
@@ -2780,6 +2811,14 @@ void FIoStoreShaderCodeArchive::ReleasePreloadEntry(int32 ShaderGroupIndex
 			{
 				PreloadEntry->IoRequest.Cancel();
 				PreloadEntry->IoRequest = FIoRequest();
+#if USE_MMAPPED_SHADERARCHIVE
+				if (PreloadEntry->MappedRegionStatus.IsOk())
+				{
+					FIoMappedRegion MappedRegion = PreloadEntry->MappedRegionStatus.ConsumeValueOrDie();
+					delete MappedRegion.MappedFileRegion;
+					delete MappedRegion.MappedFileHandle;
+				}
+#endif
 				PreloadEntry->PreloadEvent.SafeRelease();
 				const FIoStoreShaderGroupEntry& GroupEntry = Header.ShaderGroupEntries[ShaderGroupIndex];
 				ShaderGroupLoadBytes = GroupEntry.CompressedSize + sizeof(FShaderGroupPreloadEntry);
@@ -2898,74 +2937,93 @@ TRefCountPtr<FRHIShader> FIoStoreShaderCodeArchive::CreateShader(int32 ShaderInd
 		FReadScopeLock Lock(PreloadedShaderGroupsLock);
 		PreloadEntryPtr = FindExistingPreloadEntry(GroupIndex);
 	}
-
-	// raise the prio if still ongoing
-	if (!PreloadEntryPtr->IoRequest.Status().IsCompleted())
-	{
-		PreloadEntryPtr->IoRequest.UpdatePriority(IoDispatcherPriority_Max);
-	}
-	FGraphEventRef Event = PreloadEntryPtr->PreloadEvent;
-
-	bool bMissedPreLoaded = false;
-	const bool bNeededToWait = Event.IsValid() && !Event->IsComplete();
-	if (bNeededToWait)
+#if USE_MMAPPED_SHADERARCHIVE
+	if (!PreloadEntryPtr->MappedRegionStatus.IsOk())
+#endif
 	{
-		if (!bRequired)
+		// raise the prio if still ongoing
+		if (!PreloadEntryPtr->IoRequest.Status().IsCompleted())
 		{
-			PreloadEntryPtr = nullptr;
-			ReleasePreloadEntry(GroupIndex
+			PreloadEntryPtr->IoRequest.UpdatePriority(IoDispatcherPriority_Max);
+		}
+		FGraphEventRef Event = PreloadEntryPtr->PreloadEvent;
+		
+		bool bMissedPreLoaded = false;
+		const bool bNeededToWait = Event.IsValid() && !Event->IsComplete();
+		if (bNeededToWait)
+		{
+			if (!bRequired)
+			{
+				PreloadEntryPtr = nullptr;
+				ReleasePreloadEntry(GroupIndex
 #if UE_SCA_DEBUG_PRELOADING
-				, Callsite
+					, Callsite
 #endif
-			);
-			return Shader;
+				);
+				return Shader;
+			}
+			
+			bMissedPreLoaded = true;
+			TRACE_CPUPROFILER_EVENT_SCOPE(BlockingShaderLoad);
+			
+			const double TimeStarted = FPlatformTime::Seconds();
+			FTaskGraphInterface::Get().WaitUntilTaskCompletes(Event);
+			const double WaitDuration = FPlatformTime::Seconds() - TimeStarted;
+			
+			// only complain if we spent more than 1ms waiting
+			if (WaitDuration > GShaderCodeLibraryMaxShaderPreloadWaitTime)
+			{
+				UE_LOG(LogShaderLibrary, Warning, TEXT("Spent %.2f ms in a blocking wait for shader preload, NumRefs: %d, FramePreloadStarted: %d, CurrentFrame: %d"), WaitDuration * 1000.0, PreloadEntryPtr->NumRefs, PreloadEntryPtr->FramePreloadStarted, GFrameNumber);
+			}
+			CSV_CUSTOM_STAT(Shaders, PreloadShaderMiss, 1, ECsvCustomStatOp::Accumulate);
+			CSV_CUSTOM_STAT(Shaders, PreloadShaderWaitTime, WaitDuration * 1000.0f, ECsvCustomStatOp::Accumulate);
 		}
-
-		bMissedPreLoaded = true;
-		TRACE_CPUPROFILER_EVENT_SCOPE(BlockingShaderLoad);
-
-		const double TimeStarted = FPlatformTime::Seconds();
-		FTaskGraphInterface::Get().WaitUntilTaskCompletes(Event);
-		const double WaitDuration = FPlatformTime::Seconds() - TimeStarted;
-
-		// only complain if we spent more than 1ms waiting
-		if (WaitDuration > GShaderCodeLibraryMaxShaderPreloadWaitTime)
+		
+		// we should not have this situation, but if we do, log a little more info about it before dying
+		if (UNLIKELY(PreloadEntryPtr->IoRequest.GetResult() == nullptr))
 		{
-			UE_LOG(LogShaderLibrary, Warning, TEXT("Spent %.2f ms in a blocking wait for shader preload, NumRefs: %d, FramePreloadStarted: %d, CurrentFrame: %d"), WaitDuration * 1000.0, PreloadEntryPtr->NumRefs, PreloadEntryPtr->FramePreloadStarted, GFrameNumber);
+			UE_LOG(LogShaderLibrary, Warning, TEXT("IoRequest backing buffer is not there when we need it, about to crash. Library '%s', shader hash '%s' (freq=%d), group %d IoChunkId '%s', started preloading at frame %d, current frame is %d, preloading %s"),
+				*GetName(),
+				*Header.ShaderHashes[ShaderIndex].ToString(),
+				ShaderEntry.Frequency,
+				GroupIndex,
+				*LexToString(Header.ShaderGroupIoHashes[GroupIndex]),
+				PreloadEntryPtr->FramePreloadStarted,
+				GFrameNumber,
+				PreloadEntryPtr->PreloadEvent.IsValid() ? (PreloadEntryPtr->PreloadEvent->IsComplete() ? TEXT("is complete") : TEXT("is NOT complete!")) : TEXT("event is NOT VALID!")
+				);
 		}
-		CSV_CUSTOM_STAT(Shaders, PreloadShaderMiss, 1, ECsvCustomStatOp::Accumulate);
-		CSV_CUSTOM_STAT(Shaders, PreloadShaderWaitTime, WaitDuration * 1000.0f, ECsvCustomStatOp::Accumulate);
 	}
-
-	// we should not have this situation, but if we do, log a little more info about it before dying
-	if (UNLIKELY(PreloadEntryPtr->IoRequest.GetResult() == nullptr))
+	
+	const uint8* ShaderCode = nullptr;
+	uint32 DataSize = 0;
+#if USE_MMAPPED_SHADERARCHIVE
+	if (PreloadEntryPtr->MappedRegionStatus.IsOk())
 	{
-		UE_LOG(LogShaderLibrary, Warning, TEXT("IoRequest backing buffer is not there when we need it, about to crash. Library '%s', shader hash '%s' (freq=%d), group %d IoChunkId '%s', started preloading at frame %d, current frame is %d, preloading %s"),
-			*GetName(),
-			*Header.ShaderHashes[ShaderIndex].ToString(),
-			ShaderEntry.Frequency,
-			GroupIndex,
-			*LexToString(Header.ShaderGroupIoHashes[GroupIndex]),
-			PreloadEntryPtr->FramePreloadStarted,
-			GFrameNumber,
-			PreloadEntryPtr->PreloadEvent.IsValid() ? (PreloadEntryPtr->PreloadEvent->IsComplete() ? TEXT("is complete") : TEXT("is NOT complete!")) : TEXT("event is NOT VALID!")
-		);
+		IMappedFileRegion* MappedFileRegion = PreloadEntryPtr->MappedRegionStatus.ValueOrDie().MappedFileRegion;
+		ShaderCode = MappedFileRegion->GetMappedPtr();
+		DataSize = MappedFileRegion->GetMappedSize();
 	}
-
-	const uint8* ShaderCode = PreloadEntryPtr->IoRequest.GetResultOrDie().Data();
-
-	// Verbose logging from what IoDispatcher backend this shader chunk was loaded. If this is null, GetResultOrDie() will have already crashed.
-	UE_LOG(LogShaderLibrary, VeryVerbose, TEXT("Received shader chunk '%s' in library '%s' from '%s'"),
-		*LexToString(Header.ShaderGroupIoHashes[GroupIndex]), *GetName(), GetIoRequestBackendName(PreloadEntryPtr->IoRequest));
-
+	else
+#endif
+	{
+		const FIoBuffer& IoBuffer = PreloadEntryPtr->IoRequest.GetResultOrDie();
+		ShaderCode = IoBuffer.Data();
+		DataSize = IoBuffer.DataSize();
+		
+		// Verbose logging from what IoDispatcher backend this shader chunk was loaded. If this is null, GetResultOrDie() will have already crashed.
+		UE_LOG(LogShaderLibrary, VeryVerbose, TEXT("Received shader chunk '%s' in library '%s' from '%s'"),
+			*LexToString(Header.ShaderGroupIoHashes[GroupIndex]), *GetName(), GetIoRequestBackendName(PreloadEntryPtr->IoRequest));
+	}
+	
 	FMemStackBase& MemStack = FMemStack::Get();
 	FMemMark Mark(MemStack);
 	FIoStoreShaderGroupEntry& GroupEntry = Header.ShaderGroupEntries[GroupIndex];
-	uint32 CompressedSize = PreloadEntryPtr->IoRequest.GetResultOrDie().DataSize();
-	if (GroupEntry.IsGroupCompressed())
+	const bool IsGroupCompressed = GroupEntry.IsGroupCompressed();
+	if (IsGroupCompressed)
 	{
 		uint8* UncompressedCode = reinterpret_cast<uint8*>(MemStack.Alloc(GroupEntry.UncompressedSize, 16));
-		DecompressShadergroupWithOodleAndExtraLogging(GroupIndex, Header.ShaderGroupIoHashes[GroupIndex], GroupEntry, ShaderIndex, ShaderEntry.ShaderGroupIndex, Header.ShaderHashes[ShaderIndex], UncompressedCode, GroupEntry.UncompressedSize, ShaderCode, CompressedSize);
+		DecompressShadergroupWithOodleAndExtraLogging(GroupIndex, Header.ShaderGroupIoHashes[GroupIndex], GroupEntry, ShaderIndex, ShaderEntry.ShaderGroupIndex, Header.ShaderHashes[ShaderIndex], UncompressedCode, GroupEntry.UncompressedSize, ShaderCode, DataSize);
 		ShaderCode = reinterpret_cast<uint8*>(UncompressedCode) + ShaderEntry.UncompressedOffsetInGroup;
 
 #if UE_SCA_VISUALIZE_SHADER_USAGE
@@ -2979,24 +3037,36 @@ TRefCountPtr<FRHIShader> FIoStoreShaderCodeArchive::CreateShader(int32 ShaderInd
 	{
 		ShaderCode += ShaderEntry.UncompressedOffsetInGroup;
 	}
-
-	const auto ShaderCodeView = MakeArrayView(ShaderCode, Header.GetShaderUncompressedSize(ShaderIndex));
-	const FSHAHash& ShaderHash = Header.ShaderHashes[ShaderIndex];
+	FRHICreateShaderDesc CreateShaderDesc = IsGroupCompressed ? FRHICreateShaderDesc(MakeArrayView(ShaderCode, Header.GetShaderUncompressedSize(ShaderIndex))) 
+															  : FRHICreateShaderDesc(MakeArrayView(ShaderCode, Header.GetShaderUncompressedSize(ShaderIndex)),
+															  [this, GroupIndex
+#if UE_SCA_DEBUG_PRELOADING
+															  , Callsite
+#endif
+															  ](void*)
+															  {
+															  	ReleasePreloadEntry(GroupIndex
+#if UE_SCA_DEBUG_PRELOADING
+																					, Callsite
+#endif
+																					);
+	});
+	
 	switch (ShaderEntry.Frequency)
 	{
-	case SF_Vertex: Shader = RHICreateVertexShader(ShaderCodeView, ShaderHash); break;
-	case SF_Mesh: Shader = RHICreateMeshShader(ShaderCodeView, ShaderHash); break;
-	case SF_Amplification: Shader = RHICreateAmplificationShader(ShaderCodeView, ShaderHash); break;
-	case SF_Pixel: Shader = RHICreatePixelShader(ShaderCodeView, ShaderHash); break;
-	case SF_Geometry: Shader = RHICreateGeometryShader(ShaderCodeView, ShaderHash); break;
-	case SF_Compute: Shader = RHICreateComputeShader(ShaderCodeView, ShaderHash); break;
-	case SF_WorkGraphRoot: Shader = RHICreateWorkGraphShader(ShaderCodeView, ShaderHash, SF_WorkGraphRoot); break;
-	case SF_WorkGraphComputeNode: Shader = RHICreateWorkGraphShader(ShaderCodeView, ShaderHash, SF_WorkGraphComputeNode); break;
+	case SF_Vertex: Shader = RHICreateVertexShader(CreateShaderDesc); break;
+	case SF_Mesh: Shader = RHICreateMeshShader(CreateShaderDesc); break;
+	case SF_Amplification: Shader = RHICreateAmplificationShader(CreateShaderDesc); break;
+	case SF_Pixel: Shader = RHICreatePixelShader(CreateShaderDesc); break;
+	case SF_Geometry: Shader = RHICreateGeometryShader(CreateShaderDesc); break;
+	case SF_Compute: Shader = RHICreateComputeShader(CreateShaderDesc); break;
+	case SF_WorkGraphRoot: Shader = RHICreateWorkGraphShader(CreateShaderDesc, SF_WorkGraphRoot); break;
+	case SF_WorkGraphComputeNode: Shader = RHICreateWorkGraphShader(CreateShaderDesc, SF_WorkGraphComputeNode); break;
 	case SF_RayGen: case SF_RayMiss: case SF_RayHitGroup: case SF_RayCallable:
 #if RHI_RAYTRACING
 		if (GRHISupportsRayTracing && GRHISupportsRayTracingShaders)
 		{
-			Shader = RHICreateRayTracingShader(ShaderCodeView, ShaderHash, ShaderEntry.GetFrequency());
+			Shader = RHICreateRayTracingShader(CreateShaderDesc, ShaderEntry.GetFrequency());
 		}
 #endif // RHI_RAYTRACING
 		break;
@@ -3004,13 +3074,16 @@ TRefCountPtr<FRHIShader> FIoStoreShaderCodeArchive::CreateShader(int32 ShaderInd
 	}
 	DebugVisualizer.MarkCreatedForVisualization(ShaderIndex);
 
-	PreloadEntryPtr = nullptr;
-	ReleasePreloadEntry(GroupIndex
+	PreloadEntryPtr = nullptr;	
+	if (IsGroupCompressed)
+	{
+		ReleasePreloadEntry(GroupIndex
 #if UE_SCA_DEBUG_PRELOADING
-		, Callsite
+							, Callsite
 #endif
-	);
-
+							);
+	}
+	
 	if (Shader)
 	{
 		INC_DWORD_STAT(STAT_Shaders_NumShadersCreated);
@@ -3020,7 +3093,7 @@ TRefCountPtr<FRHIShader> FIoStoreShaderCodeArchive::CreateShader(int32 ShaderInd
 		CsvStatNumShadersCreated->Add(1);
 #endif
 
-		Shader->SetHash(ShaderHash);
+		Shader->SetHash(Header.ShaderHashes[ShaderIndex]);
 	}
 
 	return Shader;
diff --git a/Engine/Source/Runtime/RenderCore/Private/ShaderResource.cpp b/Engine/Source/Runtime/RenderCore/Private/ShaderResource.cpp
--- a/Engine/Source/Runtime/RenderCore/Private/ShaderResource.cpp
+++ b/Engine/Source/Runtime/RenderCore/Private/ShaderResource.cpp
@@ -688,25 +688,25 @@ FRHIShader* FShaderMapResource_InlineCode::CreateRHIShaderOrCrash(int32 ShaderIn
 		ShaderCodeView = MakeArrayView(reinterpret_cast<const uint8*>(UncompressedCode), UncompressedSize);
 	}
 
-	const FSHAHash& ShaderHash = Code->ShaderHashes[ShaderIndex];
+	FRHICreateShaderDesc CreateShaderDesc = FRHICreateShaderDesc(ShaderCodeView);
 	const EShaderFrequency Frequency = ShaderCodeResource.GetFrequency();
 
 	TRefCountPtr<FRHIShader> RHIShader;
 	switch (Frequency)
 	{
-	case SF_Vertex: RHIShader = RHICreateVertexShader(ShaderCodeView, ShaderHash); break;
-	case SF_Mesh: RHIShader = RHICreateMeshShader(ShaderCodeView, ShaderHash); break;
-	case SF_Amplification: RHIShader = RHICreateAmplificationShader(ShaderCodeView, ShaderHash); break;
-	case SF_Pixel: RHIShader = RHICreatePixelShader(ShaderCodeView, ShaderHash); break;
-	case SF_Geometry: RHIShader = RHICreateGeometryShader(ShaderCodeView, ShaderHash); break;
-	case SF_Compute: RHIShader = RHICreateComputeShader(ShaderCodeView, ShaderHash); break;
-	case SF_WorkGraphRoot: RHIShader = RHICreateWorkGraphShader(ShaderCodeView, ShaderHash, SF_WorkGraphRoot); break;
-	case SF_WorkGraphComputeNode: RHIShader = RHICreateWorkGraphShader(ShaderCodeView, ShaderHash, SF_WorkGraphComputeNode); break;
+	case SF_Vertex: RHIShader = RHICreateVertexShader(CreateShaderDesc); break;
+	case SF_Mesh: RHIShader = RHICreateMeshShader(CreateShaderDesc); break;
+	case SF_Amplification: RHIShader = RHICreateAmplificationShader(CreateShaderDesc); break;
+	case SF_Pixel: RHIShader = RHICreatePixelShader(CreateShaderDesc); break;
+	case SF_Geometry: RHIShader = RHICreateGeometryShader(CreateShaderDesc); break;
+	case SF_Compute: RHIShader = RHICreateComputeShader(CreateShaderDesc); break;
+	case SF_WorkGraphRoot: RHIShader = RHICreateWorkGraphShader(CreateShaderDesc, SF_WorkGraphRoot); break;
+	case SF_WorkGraphComputeNode: RHIShader = RHICreateWorkGraphShader(CreateShaderDesc, SF_WorkGraphComputeNode); break;
 	case SF_RayGen: case SF_RayMiss: case SF_RayHitGroup: case SF_RayCallable:
 #if RHI_RAYTRACING
 		if (GRHISupportsRayTracing && GRHISupportsRayTracingShaders)
 		{
-			RHIShader = RHICreateRayTracingShader(ShaderCodeView, ShaderHash, Frequency);
+			RHIShader = RHICreateRayTracingShader(CreateShaderDesc, Frequency);
 		}
 #endif // RHI_RAYTRACING
 		break;
@@ -714,6 +714,7 @@ FRHIShader* FShaderMapResource_InlineCode::CreateRHIShaderOrCrash(int32 ShaderIn
 		checkNoEntry();
 		break;
 	}
+	const FSHAHash& ShaderHash = Code->ShaderHashes[ShaderIndex];
 	if (UNLIKELY(RHIShader == nullptr))
 	{
 		if (bRequired)
diff --git a/Engine/Source/Runtime/RenderCore/Public/ShaderCodeArchive.h b/Engine/Source/Runtime/RenderCore/Public/ShaderCodeArchive.h
--- a/Engine/Source/Runtime/RenderCore/Public/ShaderCodeArchive.h
+++ b/Engine/Source/Runtime/RenderCore/Public/ShaderCodeArchive.h
@@ -33,11 +33,6 @@ class FCbFieldView;
 class FCbWriter;
 #endif
 
-#ifndef USE_MMAPPED_SHADERARCHIVE
-    // If enabled FSerializedShaderArchive will use a mmapping of the on-disk shader archive instead of normal serialization. The on-disk format must match the expected in memory layout.
-   #define USE_MMAPPED_SHADERARCHIVE 0
-#endif
-
 // enable visualization in the desktop Development builds only as it has a memory hit and writes files
 #define UE_SCA_VISUALIZE_SHADER_USAGE			(!WITH_EDITOR && !IS_PROGRAM && UE_BUILD_DEVELOPMENT && PLATFORM_DESKTOP)
 
@@ -915,6 +910,9 @@ private:
 	{
 		FGraphEventRef PreloadEvent;
 		FIoRequest IoRequest;
+#if USE_MMAPPED_SHADERARCHIVE
+		TIoStatusOr<FIoMappedRegion> MappedRegionStatus;
+#endif
 		uint32 FramePreloadStarted = ~0u;
 		uint32 NumRefs : 31;
 		uint32 bNeverToBePreloaded : 1;
diff --git a/Engine/Source/Runtime/VulkanRHI/Private/VulkanDynamicRHI.h b/Engine/Source/Runtime/VulkanRHI/Private/VulkanDynamicRHI.h
--- a/Engine/Source/Runtime/VulkanRHI/Private/VulkanDynamicRHI.h
+++ b/Engine/Source/Runtime/VulkanRHI/Private/VulkanDynamicRHI.h
@@ -126,12 +126,12 @@ public:
 	virtual FDepthStencilStateRHIRef RHICreateDepthStencilState(const FDepthStencilStateInitializerRHI& Initializer) final override;
 	virtual FBlendStateRHIRef RHICreateBlendState(const FBlendStateInitializerRHI& Initializer) final override;
 	virtual FVertexDeclarationRHIRef RHICreateVertexDeclaration(const FVertexDeclarationElementList& Elements) final override;
-	virtual FPixelShaderRHIRef RHICreatePixelShader(TArrayView<const uint8> Code, const FSHAHash& Hash) final override;
-	virtual FVertexShaderRHIRef RHICreateVertexShader(TArrayView<const uint8> Code, const FSHAHash& Hash) final override;
-	virtual FMeshShaderRHIRef RHICreateMeshShader(TArrayView<const uint8> Code, const FSHAHash& Hash) final override;
-	virtual FAmplificationShaderRHIRef RHICreateAmplificationShader(TArrayView<const uint8> Code, const FSHAHash& Hash) final override;
-	virtual FGeometryShaderRHIRef RHICreateGeometryShader(TArrayView<const uint8> Code, const FSHAHash& Hash) final override;
-	virtual FComputeShaderRHIRef RHICreateComputeShader(TArrayView<const uint8> Code, const FSHAHash& Hash) final override;
+	virtual FPixelShaderRHIRef RHICreatePixelShader(const FRHICreateShaderDesc& CreateShaderDesc) final override;
+	virtual FVertexShaderRHIRef RHICreateVertexShader(const FRHICreateShaderDesc& CreateShaderDesc) final override;
+	virtual FMeshShaderRHIRef RHICreateMeshShader(const FRHICreateShaderDesc& CreateShaderDesc) final override;
+	virtual FAmplificationShaderRHIRef RHICreateAmplificationShader(const FRHICreateShaderDesc& CreateShaderDesc) final override;
+	virtual FGeometryShaderRHIRef RHICreateGeometryShader(const FRHICreateShaderDesc& CreateShaderDesc) final override;
+	virtual FComputeShaderRHIRef RHICreateComputeShader(const FRHICreateShaderDesc& CreateShaderDesc) final override;
 	virtual FGPUFenceRHIRef RHICreateGPUFence(const FName &Name) final override;
 	virtual void RHIWriteGPUFence_TopOfPipe(FRHICommandListBase& RHICmdList, FRHIGPUFence* FenceRHI) final override;
 	virtual void RHICreateTransition(FRHITransition* Transition, const FRHITransitionCreateInfo& CreateInfo) final override;
@@ -261,7 +261,7 @@ public:
 	virtual FRayTracingAccelerationStructureSize RHICalcRayTracingGeometrySize(const FRayTracingGeometryInitializer& Initializer) final override;
 	virtual FRayTracingGeometryRHIRef RHICreateRayTracingGeometry(FRHICommandListBase& RHICmdList, const FRayTracingGeometryInitializer& Initializer) final override;
 	virtual FRayTracingSceneRHIRef RHICreateRayTracingScene(FRayTracingSceneInitializer Initializer) final override;
-	virtual FRayTracingShaderRHIRef RHICreateRayTracingShader(TArrayView<const uint8> Code, const FSHAHash& Hash, EShaderFrequency ShaderFrequency) final override;
+	virtual FRayTracingShaderRHIRef RHICreateRayTracingShader(const FRHICreateShaderDesc& CreateShaderDesc, EShaderFrequency ShaderFrequency) final override;
 	virtual FRayTracingPipelineStateRHIRef RHICreateRayTracingPipelineState(const FRayTracingPipelineStateInitializer& Initializer) final override;
 
 	virtual FShaderBindingTableRHIRef RHICreateShaderBindingTable(FRHICommandListBase& RHICmdList, const FRayTracingShaderBindingTableInitializer& Initializer) final override;
diff --git a/Engine/Source/Runtime/VulkanRHI/Private/VulkanShaders.cpp b/Engine/Source/Runtime/VulkanRHI/Private/VulkanShaders.cpp
--- a/Engine/Source/Runtime/VulkanRHI/Private/VulkanShaders.cpp
+++ b/Engine/Source/Runtime/VulkanRHI/Private/VulkanShaders.cpp
@@ -783,51 +783,51 @@ void FVulkanShader::PurgeShaderModules()
 	ShaderModules.Empty(0);
 }
 
-FVertexShaderRHIRef FVulkanDynamicRHI::RHICreateVertexShader(TArrayView<const uint8> Code, const FSHAHash& Hash)
+FVertexShaderRHIRef FVulkanDynamicRHI::RHICreateVertexShader(const FRHICreateShaderDesc& CreateShaderDesc)
 {
-	return Device->GetShaderFactory().CreateShader<FVulkanVertexShader>(Code, Device);
+	return Device->GetShaderFactory().CreateShader<FVulkanVertexShader>(CreateShaderDesc.Code, Device);
 }
 
-FPixelShaderRHIRef FVulkanDynamicRHI::RHICreatePixelShader(TArrayView<const uint8> Code, const FSHAHash& Hash)
+FPixelShaderRHIRef FVulkanDynamicRHI::RHICreatePixelShader(const FRHICreateShaderDesc& CreateShaderDesc)
 {
-	return Device->GetShaderFactory().CreateShader<FVulkanPixelShader>(Code, Device);
+	return Device->GetShaderFactory().CreateShader<FVulkanPixelShader>(CreateShaderDesc.Code, Device);
 }
 
-FMeshShaderRHIRef FVulkanDynamicRHI::RHICreateMeshShader(TArrayView<const uint8> Code, const FSHAHash& Hash)
+FMeshShaderRHIRef FVulkanDynamicRHI::RHICreateMeshShader(const FRHICreateShaderDesc& CreateShaderDesc)
 {
-	return Device->GetShaderFactory().CreateShader<FVulkanMeshShader>(Code, Device);
+	return Device->GetShaderFactory().CreateShader<FVulkanMeshShader>(CreateShaderDesc.Code, Device);
 }
 
-FAmplificationShaderRHIRef FVulkanDynamicRHI::RHICreateAmplificationShader(TArrayView<const uint8> Code, const FSHAHash& Hash)
+FAmplificationShaderRHIRef FVulkanDynamicRHI::RHICreateAmplificationShader(const FRHICreateShaderDesc& CreateShaderDesc)
 {
-	return Device->GetShaderFactory().CreateShader<FVulkanTaskShader>(Code, Device);
+	return Device->GetShaderFactory().CreateShader<FVulkanTaskShader>(CreateShaderDesc.Code, Device);
 }
 
-FGeometryShaderRHIRef FVulkanDynamicRHI::RHICreateGeometryShader(TArrayView<const uint8> Code, const FSHAHash& Hash)
+FGeometryShaderRHIRef FVulkanDynamicRHI::RHICreateGeometryShader(const FRHICreateShaderDesc& CreateShaderDesc)
 { 
-	return Device->GetShaderFactory().CreateShader<FVulkanGeometryShader>(Code, Device);
+	return Device->GetShaderFactory().CreateShader<FVulkanGeometryShader>(CreateShaderDesc.Code, Device);
 }
 
-FComputeShaderRHIRef FVulkanDynamicRHI::RHICreateComputeShader(TArrayView<const uint8> Code, const FSHAHash& Hash)
+FComputeShaderRHIRef FVulkanDynamicRHI::RHICreateComputeShader(const FRHICreateShaderDesc& CreateShaderDesc)
 { 
-	return Device->GetShaderFactory().CreateShader<FVulkanComputeShader>(Code, Device);
+	return Device->GetShaderFactory().CreateShader<FVulkanComputeShader>(CreateShaderDesc.Code, Device);
 }
 
-FRayTracingShaderRHIRef FVulkanDynamicRHI::RHICreateRayTracingShader(TArrayView<const uint8> Code, const FSHAHash& Hash, EShaderFrequency ShaderFrequency)
+FRayTracingShaderRHIRef FVulkanDynamicRHI::RHICreateRayTracingShader(const FRHICreateShaderDesc& CreateShaderDesc, EShaderFrequency ShaderFrequency)
 {
 	switch (ShaderFrequency)
 	{
 	case EShaderFrequency::SF_RayGen:
-		 return Device->GetShaderFactory().CreateRayTracingShader<SF_RayGen>(Code, Device);
+		 return Device->GetShaderFactory().CreateRayTracingShader<SF_RayGen>(CreateShaderDesc.Code, Device);
 
 	case EShaderFrequency::SF_RayMiss:
-		return Device->GetShaderFactory().CreateRayTracingShader<SF_RayMiss>(Code, Device);
+		return Device->GetShaderFactory().CreateRayTracingShader<SF_RayMiss>(CreateShaderDesc.Code, Device);
 
 	case EShaderFrequency::SF_RayCallable:
-		return Device->GetShaderFactory().CreateRayTracingShader<SF_RayCallable>(Code, Device);
+		return Device->GetShaderFactory().CreateRayTracingShader<SF_RayCallable>(CreateShaderDesc.Code, Device);
 
 	case EShaderFrequency::SF_RayHitGroup:
-		return Device->GetShaderFactory().CreateRayTracingShader<SF_RayHitGroup>(Code, Device);
+		return Device->GetShaderFactory().CreateRayTracingShader<SF_RayHitGroup>(CreateShaderDesc.Code, Device);
 
 	default:
 		check(false);
diff --git a/Engine/Source/Runtime/Windows/D3D11RHI/Private/D3D11RHIPrivate.h b/Engine/Source/Runtime/Windows/D3D11RHI/Private/D3D11RHIPrivate.h
--- a/Engine/Source/Runtime/Windows/D3D11RHI/Private/D3D11RHIPrivate.h
+++ b/Engine/Source/Runtime/Windows/D3D11RHI/Private/D3D11RHIPrivate.h
@@ -230,10 +230,10 @@ public:
 	D3D11RHI_API virtual FDepthStencilStateRHIRef RHICreateDepthStencilState(const FDepthStencilStateInitializerRHI& Initializer) final override;
 	D3D11RHI_API virtual FBlendStateRHIRef RHICreateBlendState(const FBlendStateInitializerRHI& Initializer) final override;
 	D3D11RHI_API virtual FVertexDeclarationRHIRef RHICreateVertexDeclaration(const FVertexDeclarationElementList& Elements) final override;
-	D3D11RHI_API virtual FPixelShaderRHIRef RHICreatePixelShader(TArrayView<const uint8> Code, const FSHAHash& Hash) final override;
-	D3D11RHI_API virtual FVertexShaderRHIRef RHICreateVertexShader(TArrayView<const uint8> Code, const FSHAHash& Hash) final override;
-	D3D11RHI_API virtual FGeometryShaderRHIRef RHICreateGeometryShader(TArrayView<const uint8> Code, const FSHAHash& Hash) final override;
-	D3D11RHI_API virtual FComputeShaderRHIRef RHICreateComputeShader(TArrayView<const uint8> Code, const FSHAHash& Hash) final override;
+	D3D11RHI_API virtual FPixelShaderRHIRef RHICreatePixelShader(const FRHICreateShaderDesc& CreateShaderDesc) final override;
+	D3D11RHI_API virtual FVertexShaderRHIRef RHICreateVertexShader(const FRHICreateShaderDesc& CreateShaderDesc) final override;
+	D3D11RHI_API virtual FGeometryShaderRHIRef RHICreateGeometryShader(const FRHICreateShaderDesc& CreateShaderDesc) final override;
+	D3D11RHI_API virtual FComputeShaderRHIRef RHICreateComputeShader(const FRHICreateShaderDesc& CreateShaderDesc) final override;
 	D3D11RHI_API virtual FStagingBufferRHIRef RHICreateStagingBuffer() final override;
 	D3D11RHI_API virtual void RHICopyToStagingBuffer(FRHIBuffer* SourceBuffer, FRHIStagingBuffer* DestinationStagingBuffer, uint32 Offset, uint32 NumBytes) final override;
 	D3D11RHI_API virtual FGPUFenceRHIRef RHICreateGPUFence(const FName& Name) final override;
diff --git a/Engine/Source/Runtime/Windows/D3D11RHI/Private/D3D11Shaders.cpp b/Engine/Source/Runtime/Windows/D3D11RHI/Private/D3D11Shaders.cpp
--- a/Engine/Source/Runtime/Windows/D3D11RHI/Private/D3D11Shaders.cpp
+++ b/Engine/Source/Runtime/Windows/D3D11RHI/Private/D3D11Shaders.cpp
@@ -119,12 +119,12 @@ static void ResetVendorExtensions(ID3D11Device* Direct3DDevice)
 #endif
 }
 
-FVertexShaderRHIRef FD3D11DynamicRHI::RHICreateVertexShader(TArrayView<const uint8> Code, const FSHAHash& Hash)
+FVertexShaderRHIRef FD3D11DynamicRHI::RHICreateVertexShader(const FRHICreateShaderDesc& CreateShaderDesc)
 {
-	FShaderCodeReader ShaderCode(Code);
+	FShaderCodeReader ShaderCode(CreateShaderDesc.Code);
 	FD3D11VertexShader* Shader = new FD3D11VertexShader;
 
-	FMemoryReaderView Ar( Code, true );
+	FMemoryReaderView Ar( CreateShaderDesc.Code, true );
 	Shader->SerializeShaderResourceTable(Ar);
 	int32 Offset = Ar.Tell();
 
@@ -144,18 +144,18 @@ FVertexShaderRHIRef FD3D11DynamicRHI::RHICreateVertexShader(TArrayView<const uin
 	}
 	
 	// TEMP
-	Shader->Code = Code;
+	Shader->Code = CreateShaderDesc.Code;
 	Shader->Offset = Offset;
 
 	return Shader;
 }
 
-FGeometryShaderRHIRef FD3D11DynamicRHI::RHICreateGeometryShader(TArrayView<const uint8> Code, const FSHAHash& Hash)
+FGeometryShaderRHIRef FD3D11DynamicRHI::RHICreateGeometryShader(const FRHICreateShaderDesc& CreateShaderDesc)
 { 
-	FShaderCodeReader ShaderCode(Code);
+	FShaderCodeReader ShaderCode(CreateShaderDesc.Code);
 	FD3D11GeometryShader* Shader = new FD3D11GeometryShader;
 
-	FMemoryReaderView Ar( Code, true );
+	FMemoryReaderView Ar( CreateShaderDesc.Code, true );
 	Shader->SerializeShaderResourceTable(Ar);
 	int32 Offset = Ar.Tell();
 
@@ -177,12 +177,12 @@ FGeometryShaderRHIRef FD3D11DynamicRHI::RHICreateGeometryShader(TArrayView<const
 	return Shader;
 }
 
-FPixelShaderRHIRef FD3D11DynamicRHI::RHICreatePixelShader(TArrayView<const uint8> Code, const FSHAHash& Hash)
+FPixelShaderRHIRef FD3D11DynamicRHI::RHICreatePixelShader(const FRHICreateShaderDesc& CreateShaderDesc)
 {
-	FShaderCodeReader ShaderCode(Code);
+	FShaderCodeReader ShaderCode(CreateShaderDesc.Code);
 	FD3D11PixelShader* Shader = new FD3D11PixelShader;
 
-	FMemoryReaderView Ar( Code, true );
+	FMemoryReaderView Ar( CreateShaderDesc.Code, true );
 	Shader->SerializeShaderResourceTable(Ar);
 	int32 Offset = Ar.Tell();
 
@@ -204,12 +204,12 @@ FPixelShaderRHIRef FD3D11DynamicRHI::RHICreatePixelShader(TArrayView<const uint8
 	return Shader;
 }
 
-FComputeShaderRHIRef FD3D11DynamicRHI::RHICreateComputeShader(TArrayView<const uint8> Code, const FSHAHash& Hash)
+FComputeShaderRHIRef FD3D11DynamicRHI::RHICreateComputeShader(const FRHICreateShaderDesc& CreateShaderDesc)
 { 
-	FShaderCodeReader ShaderCode(Code);
+	FShaderCodeReader ShaderCode(CreateShaderDesc.Code);
 	FD3D11ComputeShader* Shader = new FD3D11ComputeShader;
 
-	FMemoryReaderView Ar( Code, true );
+	FMemoryReaderView Ar( CreateShaderDesc.Code, true );
 	Shader->SerializeShaderResourceTable(Ar);
 	int32 Offset = Ar.Tell();
 
