Removing assert from ShaderPlatformConfig

#Synced-CL 51242114

diff --git a/Engine/Source/Runtime/RenderCore/Private/ShaderPlatformConfig.cpp b/Engine/Source/Runtime/RenderCore/Private/ShaderPlatformConfig.cpp
--- a/Engine/Source/Runtime/RenderCore/Private/ShaderPlatformConfig.cpp
+++ b/Engine/Source/Runtime/RenderCore/Private/ShaderPlatformConfig.cpp
@@ -138,7 +138,6 @@ void FShaderPlatformConfig::Initialize()
 #if WITH_EDITOR
 				const ERHIFeatureLevel::Type PreviewFeatureLevel = FDataDrivenShaderPlatformInfo::GetMaxFeatureLevel(ShaderPlatform);
 				const EShaderPlatform RuntimePlatform = FDataDrivenShaderPlatformInfo::GetPreviewRuntimePlatform(ShaderPlatform);
-				check(RuntimePlatform != SP_NumPlatforms);
 
 				if (RuntimePlatform < SP_NumPlatforms && PlatformConfigs[RuntimePlatform].IsValid())
 				{
