[Shaders] Temporarily disable unmounting of individual shader libraries.

The proper fix requires that shadermaps are re-assigned to new shader libraries when they are still referenced by the time we unmount their current shader library instance.
This temporary workaround unblocks QA and brings us back to how unmounting was supported previously (only when the entire plugin is unmounted).

#rnx

#Synced-CL 51256130

diff --git a/Engine/Source/Runtime/RenderCore/Private/ShaderLibrary/ShaderCodeLibrary.cpp b/Engine/Source/Runtime/RenderCore/Private/ShaderLibrary/ShaderCodeLibrary.cpp
--- a/Engine/Source/Runtime/RenderCore/Private/ShaderLibrary/ShaderCodeLibrary.cpp
+++ b/Engine/Source/Runtime/RenderCore/Private/ShaderLibrary/ShaderCodeLibrary.cpp
@@ -2215,7 +2215,8 @@ void FShaderCodeLibrary::PreInit()
 {
 	// add a callback for opening later chunks
 	UE::ShaderLibrary::Private::OnPakFileMountedDelegateHandle = FCoreDelegates::GetOnPakFileMounted2().AddStatic(&FShaderLibraryPakFileMountedCallback);
-	UE::ShaderLibrary::Private::OnPakFileUnmountedDelegateHandle = FCoreDelegates::GetOnPakFileUnmounted().AddStatic(&FShaderLibraryPakFileUnmountedCallback);
+	//@lh-todo - Temporarily disabled until shadermap resources can be properly re-assigned to new shader libraries when others are unmounted
+	//UE::ShaderLibrary::Private::OnPakFileUnmountedDelegateHandle = FCoreDelegates::GetOnPakFileUnmounted().AddStatic(&FShaderLibraryPakFileUnmountedCallback);
 }
 
 void FShaderCodeLibrary::InitForRuntime(EShaderPlatform ShaderPlatform)
@@ -2289,11 +2290,11 @@ void FShaderCodeLibrary::Shutdown()
 		FCoreDelegates::GetOnPakFileMounted2().Remove(UE::ShaderLibrary::Private::OnPakFileMountedDelegateHandle);
 		UE::ShaderLibrary::Private::OnPakFileMountedDelegateHandle.Reset();
 	}
-	if (UE::ShaderLibrary::Private::OnPakFileUnmountedDelegateHandle.IsValid())
+	/*if (UE::ShaderLibrary::Private::OnPakFileUnmountedDelegateHandle.IsValid())
 	{
 		FCoreDelegates::GetOnPakFileUnmounted().Remove(UE::ShaderLibrary::Private::OnPakFileUnmountedDelegateHandle);
 		UE::ShaderLibrary::Private::OnPakFileUnmountedDelegateHandle.Reset();
-	}
+	}*/
 	if (UE::ShaderLibrary::Private::OnPluginMountedDelegateHandle.IsValid())
 	{
 		IPluginManager::Get().OnNewPluginMounted().Remove(UE::ShaderLibrary::Private::OnPluginMountedDelegateHandle);
