FIoStoreShaderCodeArchive - Separate shader group data outside of PreloadEntry lifetime
This way the ownership of that memory is shared amonst its users and can be released whenever nobody needs it

#tests Local run & preflight

#Synced-CL 51373489

diff --git a/Engine/Source/Runtime/RenderCore/Private/ShaderLibrary/ShaderCodeArchive.cpp b/Engine/Source/Runtime/RenderCore/Private/ShaderLibrary/ShaderCodeArchive.cpp
--- a/Engine/Source/Runtime/RenderCore/Private/ShaderLibrary/ShaderCodeArchive.cpp
+++ b/Engine/Source/Runtime/RenderCore/Private/ShaderLibrary/ShaderCodeArchive.cpp
@@ -3,9 +3,6 @@
 #include "ShaderCodeArchive.h"
 
 #include "Async/ParallelFor.h"
-#if USE_MMAPPED_SHADERARCHIVE
-#include "Async/MappedFileHandle.h"
-#endif
 #include "Compression/OodleDataCompression.h"
 #include "DataDrivenShaderPlatformInfo.h"
 #include "Misc/Compression.h"
@@ -2524,8 +2521,8 @@ bool FIoStoreShaderCodeArchive::PreloadShaderGroup(int32 ShaderGroupIndex, FGrap
 			if (UNLIKELY(AttachShaderReadRequestFuncPtr == nullptr))
 			{
 #if USE_MMAPPED_SHADERARCHIVE
-				PreloadEntry.MappedRegionStatus = IoDispatcher.OpenMapped(Header.ShaderGroupIoHashes[ShaderGroupIndex], FIoReadOptions());
-				if (PreloadEntry.MappedRegionStatus.IsOk())
+				PreloadEntry.ShaderGroupData->MappedRegionStatus = IoDispatcher.OpenMapped(Header.ShaderGroupIoHashes[ShaderGroupIndex], FIoReadOptions());
+				if (PreloadEntry.ShaderGroupData->MappedRegionStatus.IsOk())
 				{
 					PreloadEntry.PreloadEvent->DispatchSubsequents();
 					PreloadEntry.PreloadEvent.SafeRelease();
@@ -2534,7 +2531,7 @@ bool FIoStoreShaderCodeArchive::PreloadShaderGroup(int32 ShaderGroupIndex, FGrap
 #endif
 				{
 					FIoBatch IoBatch = IoDispatcher.NewBatch();
-					PreloadEntry.IoRequest = IoBatch.Read(Header.ShaderGroupIoHashes[ShaderGroupIndex], FIoReadOptions(), GetShaderCodeArchivePriority());
+					PreloadEntry.ShaderGroupData->IoRequest = IoBatch.Read(Header.ShaderGroupIoHashes[ShaderGroupIndex], FIoReadOptions(), GetShaderCodeArchivePriority());
 					IoBatch.IssueAndDispatchSubsequents(PreloadEntry.PreloadEvent);
 				}
 			}
@@ -2542,17 +2539,17 @@ bool FIoStoreShaderCodeArchive::PreloadShaderGroup(int32 ShaderGroupIndex, FGrap
 			{
 				FCoreDelegates::FShaderReadRequestResult Result = (*AttachShaderReadRequestFuncPtr)(Header.ShaderGroupIoHashes[ShaderGroupIndex], PreloadEntry.PreloadEvent);
 #if USE_MMAPPED_SHADERARCHIVE
-				PreloadEntry.IoRequest = Result.Key;
-				PreloadEntry.MappedRegionStatus = Result.Value;
+				PreloadEntry.ShaderGroupData->IoRequest = Result.Key;
+				PreloadEntry.ShaderGroupData->MappedRegionStatus = Result.Value;
 #else
-				PreloadEntry.IoRequest = Result;
+				PreloadEntry.ShaderGroupData->IoRequest = Result;
 #endif
 			}
 		};
 
-	if (PreloadEntry.IoRequest.Status() == FIoStatus::Invalid
+	if (PreloadEntry.ShaderGroupData->IoRequest.Status() == FIoStatus::Invalid
 #if USE_MMAPPED_SHADERARCHIVE
-		&& !PreloadEntry.MappedRegionStatus.IsOk()
+		&& !PreloadEntry.ShaderGroupData->MappedRegionStatus.IsOk()
 #endif
 		)
 	{
@@ -2566,15 +2563,15 @@ bool FIoStoreShaderCodeArchive::PreloadShaderGroup(int32 ShaderGroupIndex, FGrap
 		CsvStatPreloadedShaderMB->Add((float)ShaderGroupSize / (1024.0f * 1024.0f));
 #endif
 	}
-	else if (PreloadEntry.PreloadEvent.IsValid() && PreloadEntry.PreloadEvent->IsCompleted() && PreloadEntry.IoRequest.GetResult() == nullptr
+	else if (PreloadEntry.PreloadEvent.IsValid() && PreloadEntry.PreloadEvent->IsCompleted() && PreloadEntry.ShaderGroupData->IoRequest.GetResult() == nullptr
 #if USE_MMAPPED_SHADERARCHIVE
-		&& !PreloadEntry.MappedRegionStatus.IsOk()
+		&& !PreloadEntry.ShaderGroupData->MappedRegionStatus.IsOk()
 #endif
 			 )
 	{
 		// failed request, re-kick in-place, without changing stats
-		PreloadEntry.IoRequest.Cancel();
-		PreloadEntry.IoRequest = FIoRequest();
+		PreloadEntry.ShaderGroupData->IoRequest.Cancel();
+		PreloadEntry.ShaderGroupData->IoRequest = FIoRequest();
 		PreloadEntry.PreloadEvent.SafeRelease();
 
 		KickOffIoRequest();
@@ -2809,16 +2806,7 @@ void FIoStoreShaderCodeArchive::ReleasePreloadEntry(int32 ShaderGroupIndex
 			uint32 ShaderGroupLoadBytes = 0;
 			if (!PreloadEntry->bNeverToBePreloaded)
 			{
-				PreloadEntry->IoRequest.Cancel();
-				PreloadEntry->IoRequest = FIoRequest();
-#if USE_MMAPPED_SHADERARCHIVE
-				if (PreloadEntry->MappedRegionStatus.IsOk())
-				{
-					FIoMappedRegion MappedRegion = PreloadEntry->MappedRegionStatus.ConsumeValueOrDie();
-					delete MappedRegion.MappedFileRegion;
-					delete MappedRegion.MappedFileHandle;
-				}
-#endif
+				PreloadEntry->ShaderGroupData.Reset();
 				PreloadEntry->PreloadEvent.SafeRelease();
 				const FIoStoreShaderGroupEntry& GroupEntry = Header.ShaderGroupEntries[ShaderGroupIndex];
 				ShaderGroupLoadBytes = GroupEntry.CompressedSize + sizeof(FShaderGroupPreloadEntry);
@@ -2938,13 +2926,13 @@ TRefCountPtr<FRHIShader> FIoStoreShaderCodeArchive::CreateShader(int32 ShaderInd
 		PreloadEntryPtr = FindExistingPreloadEntry(GroupIndex);
 	}
 #if USE_MMAPPED_SHADERARCHIVE
-	if (!PreloadEntryPtr->MappedRegionStatus.IsOk())
+	if (!PreloadEntryPtr->ShaderGroupData->MappedRegionStatus.IsOk())
 #endif
 	{
 		// raise the prio if still ongoing
-		if (!PreloadEntryPtr->IoRequest.Status().IsCompleted())
+		if (!PreloadEntryPtr->ShaderGroupData->IoRequest.Status().IsCompleted())
 		{
-			PreloadEntryPtr->IoRequest.UpdatePriority(IoDispatcherPriority_Max);
+			PreloadEntryPtr->ShaderGroupData->IoRequest.UpdatePriority(IoDispatcherPriority_Max);
 		}
 		FGraphEventRef Event = PreloadEntryPtr->PreloadEvent;
 		
@@ -2980,7 +2968,7 @@ TRefCountPtr<FRHIShader> FIoStoreShaderCodeArchive::CreateShader(int32 ShaderInd
 		}
 		
 		// we should not have this situation, but if we do, log a little more info about it before dying
-		if (UNLIKELY(PreloadEntryPtr->IoRequest.GetResult() == nullptr))
+		if (UNLIKELY(PreloadEntryPtr->ShaderGroupData->IoRequest.GetResult() == nullptr))
 		{
 			UE_LOG(LogShaderLibrary, Warning, TEXT("IoRequest backing buffer is not there when we need it, about to crash. Library '%s', shader hash '%s' (freq=%d), group %d IoChunkId '%s', started preloading at frame %d, current frame is %d, preloading %s"),
 				*GetName(),
@@ -2998,22 +2986,22 @@ TRefCountPtr<FRHIShader> FIoStoreShaderCodeArchive::CreateShader(int32 ShaderInd
 	const uint8* ShaderCode = nullptr;
 	uint32 DataSize = 0;
 #if USE_MMAPPED_SHADERARCHIVE
-	if (PreloadEntryPtr->MappedRegionStatus.IsOk())
+	if (PreloadEntryPtr->ShaderGroupData->MappedRegionStatus.IsOk())
 	{
-		IMappedFileRegion* MappedFileRegion = PreloadEntryPtr->MappedRegionStatus.ValueOrDie().MappedFileRegion;
+		IMappedFileRegion* MappedFileRegion = PreloadEntryPtr->ShaderGroupData->MappedRegionStatus.ValueOrDie().MappedFileRegion;
 		ShaderCode = MappedFileRegion->GetMappedPtr();
 		DataSize = MappedFileRegion->GetMappedSize();
 	}
 	else
 #endif
 	{
-		const FIoBuffer& IoBuffer = PreloadEntryPtr->IoRequest.GetResultOrDie();
+		const FIoBuffer& IoBuffer = PreloadEntryPtr->ShaderGroupData->IoRequest.GetResultOrDie();
 		ShaderCode = IoBuffer.Data();
 		DataSize = IoBuffer.DataSize();
 		
 		// Verbose logging from what IoDispatcher backend this shader chunk was loaded. If this is null, GetResultOrDie() will have already crashed.
 		UE_LOG(LogShaderLibrary, VeryVerbose, TEXT("Received shader chunk '%s' in library '%s' from '%s'"),
-			*LexToString(Header.ShaderGroupIoHashes[GroupIndex]), *GetName(), GetIoRequestBackendName(PreloadEntryPtr->IoRequest));
+			*LexToString(Header.ShaderGroupIoHashes[GroupIndex]), *GetName(), GetIoRequestBackendName(PreloadEntryPtr->ShaderGroupData->IoRequest));
 	}
 	
 	FMemStackBase& MemStack = FMemStack::Get();
@@ -3037,20 +3025,15 @@ TRefCountPtr<FRHIShader> FIoStoreShaderCodeArchive::CreateShader(int32 ShaderInd
 	{
 		ShaderCode += ShaderEntry.UncompressedOffsetInGroup;
 	}
-	FRHICreateShaderDesc CreateShaderDesc = IsGroupCompressed ? FRHICreateShaderDesc(MakeArrayView(ShaderCode, Header.GetShaderUncompressedSize(ShaderIndex))) 
-															  : FRHICreateShaderDesc(MakeArrayView(ShaderCode, Header.GetShaderUncompressedSize(ShaderIndex)),
-															  [this, GroupIndex
-#if UE_SCA_DEBUG_PRELOADING
-															  , Callsite
-#endif
-															  ](void*)
-															  {
-															  	ReleasePreloadEntry(GroupIndex
-#if UE_SCA_DEBUG_PRELOADING
-																					, Callsite
-#endif
-																					);
-	});
+	FRHICreateShaderDesc CreateShaderDesc = IsGroupCompressed
+		? FRHICreateShaderDesc(MakeArrayView(ShaderCode, Header.GetShaderUncompressedSize(ShaderIndex)))
+		: FRHICreateShaderDesc(
+			MakeArrayView(ShaderCode, Header.GetShaderUncompressedSize(ShaderIndex)),
+			[ShaderGroupDataSharedPtr = PreloadEntryPtr->ShaderGroupData](void*)
+			{
+				// ShaderGroupDataSharedPtr will be released when this lambda is destroyed
+			}
+		);
 	
 	switch (ShaderEntry.Frequency)
 	{
@@ -3075,14 +3058,13 @@ TRefCountPtr<FRHIShader> FIoStoreShaderCodeArchive::CreateShader(int32 ShaderInd
 	DebugVisualizer.MarkCreatedForVisualization(ShaderIndex);
 
 	PreloadEntryPtr = nullptr;	
-	if (IsGroupCompressed)
-	{
-		ReleasePreloadEntry(GroupIndex
+
+	ReleasePreloadEntry(GroupIndex
 #if UE_SCA_DEBUG_PRELOADING
-							, Callsite
+						, Callsite
 #endif
-							);
-	}
+						);
+	
 	
 	if (Shader)
 	{
diff --git a/Engine/Source/Runtime/RenderCore/Public/ShaderCodeArchive.h b/Engine/Source/Runtime/RenderCore/Public/ShaderCodeArchive.h
--- a/Engine/Source/Runtime/RenderCore/Public/ShaderCodeArchive.h
+++ b/Engine/Source/Runtime/RenderCore/Public/ShaderCodeArchive.h
@@ -15,6 +15,10 @@
 #include "IO/IoDispatcher.h"
 #include "Misc/AssertionMacros.h"
 #include "Misc/CoreDelegates.h"
+// USE_MMAPPED_SHADERARCHIVE is defined in Misc/CoreDelegates.h
+#if USE_MMAPPED_SHADERARCHIVE
+#include "Async/MappedFileHandle.h"
+#endif
 #include "Misc/MemoryReadStream.h"
 #include "Misc/SecureHash.h"
 #include "RHI.h"
@@ -905,14 +909,33 @@ public:
 
 private:
 	static constexpr uint32 CurrentVersion = 1;
-
-	struct FShaderGroupPreloadEntry
+	
+	struct FShaderGroupData
 	{
-		FGraphEventRef PreloadEvent;
 		FIoRequest IoRequest;
 #if USE_MMAPPED_SHADERARCHIVE
 		TIoStatusOr<FIoMappedRegion> MappedRegionStatus;
 #endif
+		
+		~FShaderGroupData()
+		{
+			IoRequest.Cancel();
+#if USE_MMAPPED_SHADERARCHIVE
+			if (MappedRegionStatus.IsOk())
+			{
+				FIoMappedRegion MappedRegion = MappedRegionStatus.ConsumeValueOrDie();
+				delete MappedRegion.MappedFileRegion;
+				delete MappedRegion.MappedFileHandle;
+			}
+#endif
+		}
+	};
+
+	struct FShaderGroupPreloadEntry
+	{
+		FGraphEventRef PreloadEvent;
+		TSharedPtr<FShaderGroupData> ShaderGroupData;
+
 		uint32 FramePreloadStarted = ~0u;
 		uint32 NumRefs : 31;
 		uint32 bNeverToBePreloaded : 1;
@@ -922,7 +945,8 @@ private:
 #endif
 
 		FShaderGroupPreloadEntry()
-			: NumRefs(0)
+			: ShaderGroupData(MakeShared<FShaderGroupData>())
+			, NumRefs(0)
 			, bNeverToBePreloaded(0)
 		{
 		}
