Reverted CL 32329009 that replaced thread_local with TLS usage.
After additional thorough testing thread_local was proven to be faster

#Synced-CL 51337004

#Synced-CL 51337004

diff --git a/Engine/Source/Runtime/RenderCore/Private/RenderGraphAllocator.cpp b/Engine/Source/Runtime/RenderCore/Private/RenderGraphAllocator.cpp
--- a/Engine/Source/Runtime/RenderCore/Private/RenderGraphAllocator.cpp
+++ b/Engine/Source/Runtime/RenderCore/Private/RenderGraphAllocator.cpp
@@ -8,7 +8,7 @@
 thread_local int32 FRDGAllocator::NumAccessesTLS = 0;
 #endif
 
-uint32 FRDGAllocator::AllocatorTLSSlot = FPlatformTLS::AllocTlsSlot();
+thread_local FRDGAllocator* FRDGAllocator::AllocatorTLS = nullptr;
 
 FRDGAllocator::FRDGAllocator()
 #if !RDG_USE_MALLOC
@@ -50,9 +50,8 @@ FRDGAllocator::~FRDGAllocator()
 
 FRDGAllocator& FRDGAllocator::GetTLS()
 {
-	void* Allocator = FPlatformTLS::GetTlsValue(AllocatorTLSSlot);
-	checkf(Allocator, TEXT("Attempted to access RDG allocator outside of FRDGAllocatorScope"));
-	return *(FRDGAllocator*)Allocator;
+	checkf(AllocatorTLS, TEXT("Attempted to access RDG allocator outside of FRDGAllocatorScope"));
+	return *AllocatorTLS;
 }
 
 #if RDG_ALLOCATOR_DEBUG
@@ -97,14 +96,14 @@ void FRDGAllocator::ReleaseAll()
 }
 
 FRDGAllocatorScope::FRDGAllocatorScope(FRDGAllocator& Allocator)
-	: AllocatorToRestore(FPlatformTLS::GetTlsValue(FRDGAllocator::AllocatorTLSSlot))
+	: AllocatorToRestore(FRDGAllocator::AllocatorTLS)
 {
-	FPlatformTLS::SetTlsValue(FRDGAllocator::AllocatorTLSSlot, &Allocator);
+	FRDGAllocator::AllocatorTLS = &Allocator;
 }
 
 FRDGAllocatorScope::~FRDGAllocatorScope()
 {
-	FPlatformTLS::SetTlsValue(FRDGAllocator::AllocatorTLSSlot, AllocatorToRestore);
+	FRDGAllocator::AllocatorTLS = AllocatorToRestore;
 }
 
 FORCENOINLINE void UE::RenderCore::Private::OnInvalidRDGAllocatorNum(int32 NewNum, SIZE_T NumBytesPerElement)
diff --git a/Engine/Source/Runtime/RenderCore/Public/RenderGraphAllocator.h b/Engine/Source/Runtime/RenderCore/Public/RenderGraphAllocator.h
--- a/Engine/Source/Runtime/RenderCore/Public/RenderGraphAllocator.h
+++ b/Engine/Source/Runtime/RenderCore/Public/RenderGraphAllocator.h
@@ -139,7 +139,7 @@ private:
 #endif
 
 	friend class FRDGAllocatorScope;
-	static uint32 AllocatorTLSSlot;
+	static thread_local FRDGAllocator* AllocatorTLS;
 };
 
 class FRDGAllocatorScope
@@ -149,7 +149,7 @@ public:
 	RENDERCORE_API ~FRDGAllocatorScope();
 
 private:
-	void* AllocatorToRestore;
+	FRDGAllocator* AllocatorToRestore;
 };
 
 #define RDG_FRIEND_ALLOCATOR_FRIEND(Type) friend class FRDGAllocator::TObject<Type>
