[Shaders] Move shadermap resources from unmounted containers to new shader library instance.

When a Pak container is unmounted, all shader library instances associated with that container must also be unmounted.
Otherwise, the shader library instances still have their own TOC, but the shader data is no longer available from the unmounted container.
The actual deletion of the shader library instance must happen at the end of a frame, since some RHI resources might still hold on to their shaders.
This is implemented with a singleton that is hooked up to the OnEndFrame delegate and removes itself afterwards.

Also re-enable unmounting individual shader library instances when their containers are unmounted.

Remove old diagnostics function `DiagnoseLeakingShaderMapResourceRefs()`, which wasn't really helpful.

#rnx

#Synced-CL 51487274

diff --git a/Engine/Source/Runtime/RenderCore/Internal/ShaderLibrary/ShaderCodeLibraryInternal.h b/Engine/Source/Runtime/RenderCore/Internal/ShaderLibrary/ShaderCodeLibraryInternal.h
new file mode 100644
--- /dev/null
+++ b/Engine/Source/Runtime/RenderCore/Internal/ShaderLibrary/ShaderCodeLibraryInternal.h
@@ -0,0 +1,18 @@
+// Copyright Epic Games, Inc. All Rights Reserved.
+
+/*=============================================================================
+	ShaderCodeLibraryInternal.h: Extended functions to FShaderCodeLibrary for internal use only
+=============================================================================*/
+
+#pragma once
+
+#include "CoreTypes.h"
+
+class FShaderCodeLibraryInternal
+{
+public:
+	/** Tries to find a new shader library instance for the specified resource.
+		This must only be called when a shader library instance is being closed, when all the mutexes are already locked.
+		See destructor of FShaderLibraryInstance. */
+	static RENDERCORE_API bool MoveShaderMapResourceOwnership(class FShaderMapResource_SharedCode* Resource);
+};
diff --git a/Engine/Source/Runtime/RenderCore/Private/ShaderLibrary/ShaderCodeLibrary.cpp b/Engine/Source/Runtime/RenderCore/Private/ShaderLibrary/ShaderCodeLibrary.cpp
--- a/Engine/Source/Runtime/RenderCore/Private/ShaderLibrary/ShaderCodeLibrary.cpp
+++ b/Engine/Source/Runtime/RenderCore/Private/ShaderLibrary/ShaderCodeLibrary.cpp
@@ -41,6 +41,7 @@ ShaderCodeLibrary.cpp: Bound shader state cache implementation.
 #include "ShaderPipelineCache.h"
 #include "ShaderLibrary/EditorShaderCodeArchive.h"
 #include "ShaderLibrary/EditorShaderStableInfo.h"
+#include "ShaderLibrary/ShaderCodeLibraryInternal.h"
 #include "ShaderLibrary/ShaderCodeLibraryUtilities.h"
 #include "ShaderLibrary/ShaderLibraryInstance.h"
 #include "String/ParseTokens.h"
@@ -252,7 +253,7 @@ namespace UE::ShaderLibrary::Private
 		bool OpenShaderCode(const FString& ShaderCodeDir, FString const& Library, int32 ChunkId, int32 Priority = 0);
 		void CloseShaderCode(const FString& ShaderCodeDir, FString const& Library, int32 ChunkId);
 		void ForEachShaderLibraryWithShaderMap(const FSHAHash& Hash, TFunctionRef<void(FShaderLibraryInstance* InLibraryInstance, int32 InShaderMapIndex)> InPerShaderMapCallback);
-		FShaderLibraryInstance* FindShaderLibraryForShaderMap(const FSHAHash& Hash, int32& OutShaderMapIndex);
+		FShaderLibraryInstance* FindShaderLibraryForShaderMap(const FSHAHash& Hash, int32& OutShaderMapIndex, const FShaderLibraryInstance* SkipLibraryInstance = nullptr);
 		FShaderLibraryInstance* FindShaderLibraryForShader(const FSHAHash& Hash, int32& OutShaderIndex);
 		uint32 GetShaderCount();
 #if UE_SHADERLIB_WITH_INTROSPECTION
@@ -1345,13 +1346,13 @@ public:
 		}
 	}
 
-	FShaderLibraryInstance* FindShaderLibraryForShaderMap(const FSHAHash& Hash, int32& OutShaderMapIndex)
+	FShaderLibraryInstance* FindShaderLibraryForShaderMap(const FSHAHash& Hash, int32& OutShaderMapIndex, const FShaderLibraryInstance* SkipLibraryInstance = nullptr)
 	{
 		FRWScopeLock ReadLock(NamedLibrariesMutex, SLT_ReadOnly);
 
 		for (TTuple<FString, TUniquePtr<UE::ShaderLibrary::Private::FNamedShaderLibrary>>& NamedLibraryPair : NamedLibrariesStack)
 		{
-			FShaderLibraryInstance* Instance = NamedLibraryPair.Value->FindShaderLibraryForShaderMap(Hash, OutShaderMapIndex);
+			FShaderLibraryInstance* Instance = NamedLibraryPair.Value->FindShaderLibraryForShaderMap(Hash, OutShaderMapIndex, SkipLibraryInstance);
 			if (Instance)
 			{
 				return Instance;
@@ -2104,27 +2105,21 @@ public:
 			Invoke(Forward<F>(Func), LogicalName, *NamedShaderLibrary);
 		}
 	}
-};
 
-// Used to diagnose leaking references to shader map resources in FShaderMapResource_SharedCode::ReleaseRHI()
-void DiagnoseLeakingShaderMapResourceRefs(const FShaderMapResource_SharedCode* InResource)
-{
-	TArray<FString> LibsWithRefToResource;
-	FShaderLibrariesCollection::Impl->GatherShaderMapResourceReferences(InResource, LibsWithRefToResource);
-	FString Info;
-	for (const FString& Name : LibsWithRefToResource)
+	bool MoveShaderMapResourceOwnership(FShaderMapResource_SharedCode* Resource)
 	{
-		if (!Info.IsEmpty())
+		check(Resource);
+
+		// Find new shader library for this resource.
+		int32 NewShaderMapIndex = INDEX_NONE;
+		if (FShaderLibraryInstance* NewLibraryInstance = FindShaderLibraryForShaderMap(Resource->GetShaderMapHash(), NewShaderMapIndex, Resource->LibraryInstance))
 		{
-			Info += TEXT(", ");
+			Resource->AssignToNewShaderLibrary(NewLibraryInstance, NewShaderMapIndex);
+			return true;
 		}
-		Info += Name;
-	}
-	if (!Info.IsEmpty())
-	{
-		UE_LOG(LogShaderLibrary, Warning, TEXT("Leaking reference to shader map resource detected; Resource \"%s\" still referenced in: %s"), *InResource->GetFriendlyName(), *Info);
+		return false;
 	}
-}
+};
 
 static FSharedShaderCodeRequest OnSharedShaderCodeRequest;
 
@@ -2215,8 +2210,7 @@ void FShaderCodeLibrary::PreInit()
 {
 	// add a callback for opening later chunks
 	UE::ShaderLibrary::Private::OnPakFileMountedDelegateHandle = FCoreDelegates::GetOnPakFileMounted2().AddStatic(&FShaderLibraryPakFileMountedCallback);
-	//@lh-todo - Temporarily disabled until shadermap resources can be properly re-assigned to new shader libraries when others are unmounted
-	//UE::ShaderLibrary::Private::OnPakFileUnmountedDelegateHandle = FCoreDelegates::GetOnPakFileUnmounted().AddStatic(&FShaderLibraryPakFileUnmountedCallback);
+	UE::ShaderLibrary::Private::OnPakFileUnmountedDelegateHandle = FCoreDelegates::GetOnPakFileUnmounted().AddStatic(&FShaderLibraryPakFileUnmountedCallback);
 }
 
 void FShaderCodeLibrary::InitForRuntime(EShaderPlatform ShaderPlatform)
@@ -2290,11 +2284,11 @@ void FShaderCodeLibrary::Shutdown()
 		FCoreDelegates::GetOnPakFileMounted2().Remove(UE::ShaderLibrary::Private::OnPakFileMountedDelegateHandle);
 		UE::ShaderLibrary::Private::OnPakFileMountedDelegateHandle.Reset();
 	}
-	/*if (UE::ShaderLibrary::Private::OnPakFileUnmountedDelegateHandle.IsValid())
+	if (UE::ShaderLibrary::Private::OnPakFileUnmountedDelegateHandle.IsValid())
 	{
 		FCoreDelegates::GetOnPakFileUnmounted().Remove(UE::ShaderLibrary::Private::OnPakFileUnmountedDelegateHandle);
 		UE::ShaderLibrary::Private::OnPakFileUnmountedDelegateHandle.Reset();
-	}*/
+	}
 	if (UE::ShaderLibrary::Private::OnPluginMountedDelegateHandle.IsValid())
 	{
 		IPluginManager::Get().OnNewPluginMounted().Remove(UE::ShaderLibrary::Private::OnPluginMountedDelegateHandle);
@@ -2912,6 +2906,11 @@ void FShaderCodeLibrary::OpenPluginShaderLibrary(IPlugin& Plugin, bool bMonolith
 	}
 }
 
+bool FShaderCodeLibraryInternal::MoveShaderMapResourceOwnership(FShaderMapResource_SharedCode* Resource)
+{
+	return FShaderLibrariesCollection::Impl->MoveShaderMapResourceOwnership(Resource);
+}
+
 // FNamedShaderLibrary methods
 
 static void LogShaderLibraryInstanceState(const FShaderLibraryInstance* InLibraryInstance, const TCHAR* InLibraryName, const TCHAR* InStateInfo)
@@ -3000,6 +2999,63 @@ bool UE::ShaderLibrary::Private::FNamedShaderLibrary::OpenShaderCode(const FStri
 	return true;
 }
 
+// Singleton to defer the release of ShaderLibraryInstance components to the end of a frame.
+class FDeferredShaderLibraryInstanceDeleter
+{
+	FCriticalSection Lock;
+	FDelegateHandle OnEndFrameDelegateHandle;
+	TArray<FShaderLibraryInstancePriorityPair> ComponentsEnqueuedForRelease;
+
+	void OnEndFrame()
+	{
+		checkf(IsInGameThread(), TEXT("Shader library instances must be released on the game thread"));
+		FScopeLock ScopeLock(&Lock);
+
+		// Reset the container. This will flush the delete operations.
+		ComponentsEnqueuedForRelease.Reset();
+
+		// Remove OnEndFrame delegate handle to avoid being invoked in further frames
+		if (OnEndFrameDelegateHandle.IsValid())
+		{
+			if (FCoreDelegates::OnEndFrame.Remove(OnEndFrameDelegateHandle))
+			{
+				OnEndFrameDelegateHandle.Reset();
+			}
+			else
+			{
+				UE_LOG(LogShaderLibrary, Warning, TEXT("Failed to remove `OnEndFrame` delegate to release shader library instances"));
+			}
+		}
+	}
+
+	void RegisterDelegate()
+	{
+		if (!OnEndFrameDelegateHandle.IsValid())
+		{
+			OnEndFrameDelegateHandle = FCoreDelegates::OnEndFrame.AddRaw(this, &FDeferredShaderLibraryInstanceDeleter::OnEndFrame);
+		}
+	}
+
+	FDeferredShaderLibraryInstanceDeleter() = default;
+
+public:
+	FDeferredShaderLibraryInstanceDeleter(const FDeferredShaderLibraryInstanceDeleter&) = delete;
+	FDeferredShaderLibraryInstanceDeleter& operator = (const FDeferredShaderLibraryInstanceDeleter&) = delete;
+
+	void EnqueueRelease(FShaderLibraryInstancePriorityPair&& Component)
+	{
+		FScopeLock ScopeLock(&Lock);
+		RegisterDelegate();
+		ComponentsEnqueuedForRelease.Add(MoveTemp(Component));
+	}
+
+	static FDeferredShaderLibraryInstanceDeleter& Get()
+	{
+		static FDeferredShaderLibraryInstanceDeleter Instance;
+		return Instance;
+	}
+};
+
 void UE::ShaderLibrary::Private::FNamedShaderLibrary::CloseShaderCode(const FString& ShaderCodeDir, FString const& Library, int32 ChunkId)
 {
 	LLM_SCOPE(ELLMTag::Shaders);
@@ -3013,6 +3069,9 @@ void UE::ShaderLibrary::Private::FNamedShaderLibrary::CloseShaderCode(const FStr
 			{
 				LogShaderLibraryInstanceState(LibraryInstance, *Library, TEXT("Released"));
 
+				// Defer the actual deletion of the shader library instance to the end of a frame, since some RHI resources might still hold on to their shaders in the current frame.
+				FDeferredShaderLibraryInstanceDeleter::Get().EnqueueRelease(MoveTemp(Components[Index]));
+
 				// Remove component. Ensure we preserve the order since non-IAD components need to be iterated before IAD components.
 				Components.RemoveAt(Index);
 
@@ -3039,17 +3098,20 @@ void UE::ShaderLibrary::Private::FNamedShaderLibrary::ForEachShaderLibraryWithSh
 	}
 }
 
-FShaderLibraryInstance* UE::ShaderLibrary::Private::FNamedShaderLibrary::FindShaderLibraryForShaderMap(const FSHAHash& Hash, int32& OutShaderMapIndex)
+FShaderLibraryInstance* UE::ShaderLibrary::Private::FNamedShaderLibrary::FindShaderLibraryForShaderMap(const FSHAHash& Hash, int32& OutShaderMapIndex, const FShaderLibraryInstance* SkipLibraryInstance)
 {
 	FRWScopeLock ReadLock(ComponentsMutex, SLT_ReadOnly);
 
 	for (const FShaderLibraryInstancePriorityPair& Instance : Components)
 	{
-		const int32 ShaderMapIndex = Instance.Value->Library->FindShaderMapIndex(Hash);
-		if (ShaderMapIndex != INDEX_NONE)
+		if (Instance.Value.Get() != SkipLibraryInstance)
 		{
-			OutShaderMapIndex = ShaderMapIndex;
-			return Instance.Value.Get();
+			const int32 ShaderMapIndex = Instance.Value->Library->FindShaderMapIndex(Hash);
+			if (ShaderMapIndex != INDEX_NONE)
+			{
+				OutShaderMapIndex = ShaderMapIndex;
+				return Instance.Value.Get();
+			}
 		}
 	}
 	return nullptr;
diff --git a/Engine/Source/Runtime/RenderCore/Private/ShaderLibrary/ShaderLibraryInstance.cpp b/Engine/Source/Runtime/RenderCore/Private/ShaderLibrary/ShaderLibraryInstance.cpp
--- a/Engine/Source/Runtime/RenderCore/Private/ShaderLibrary/ShaderLibraryInstance.cpp
+++ b/Engine/Source/Runtime/RenderCore/Private/ShaderLibrary/ShaderLibraryInstance.cpp
@@ -13,6 +13,7 @@
 #include "ProfilingDebugging/DiagnosticTable.h"
 #include "RHIShaderLibrary.h"
 #include "ShaderCodeArchive.h"
+#include "ShaderLibrary/ShaderCodeLibraryInternal.h"
 #include "ShaderLibrary/ShaderCodeLibraryUtilities.h"
 #include "UObject/NameTypes.h"
 
@@ -90,18 +91,40 @@ FShaderLibraryInstance* FShaderLibraryInstance::Create(EShaderPlatform InShaderP
 
 FShaderLibraryInstance::~FShaderLibraryInstance()
 {
+	// Ensure RHI shaders are no longer used before attempting to move them to a new shader library instance
+	checkf(IsInGameThread(), TEXT("Shader library closure is expected to happen only on the game thread, at the \'top\' of the pipeline"));
+	FlushRenderingCommands();
+
 	// release RHI resources on all of the resources (note: this has nothing to do with their own lifetime and refcount)
 	for (FShaderMapResource_SharedCode* Resource : Resources)
 	{
 		if (Resource)
 		{
-			BeginReleaseResource(Resource);
+			const int32 NumRefs = Resource->GetNumRefs();
+			if (NumRefs > 0)
+			{
+				// Try to find new shader library instance to take ownership
+				if (!FShaderCodeLibraryInternal::MoveShaderMapResourceOwnership(Resource))
+				{
+					UE_LOG(LogShaderLibrary, Warning, TEXT("Failed to move shadermap resource '%s' (ShaderMapIndex=%d, Owner='%s') to a new library instance, but its resource is still referenced (Refcount=%d)."),
+						*LexToString(Resource->GetShaderMapHash()), Resource->ShaderMapIndex, *Resource->GetOwnerName().ToString(), NumRefs);
+
+					// If no new shader library was found, something is still holding on to this resource which it shouldn't.
+					// If the reference was leaked, we'll only see the warning above. Otherwise, we'll see another warning or failure later on.
+					// Either way, this resource must be released or it could attempt to create shaders from a library that's being destroyed.
+					BeginReleaseResource(Resource);
+				}
+			}
+			else
+			{
+				BeginReleaseResource(Resource);
+			}
 		}
 	}
 
-	// if rendering thread is active, the actual teardown may happen later, so flush the rendering commands here
-	checkf(IsInGameThread(), TEXT("Shader library closure is expected to happen only on the game thread, at the \'top\' of the pipeline"));
-	FlushRenderingCommands();	// this will also flush pending deletes
+	// if rendering thread is active, the actual teardown may happen later, so flush the rendering commands here.
+	// this will also flush pending deletes
+	FlushRenderingCommands();
 
 	Library->Teardown();
 	DEC_DWORD_STAT_BY(STAT_Shaders_ShaderResourceMemory, GetSizeBytes());
@@ -235,16 +258,48 @@ TRefCountPtr<FShaderMapResource_SharedCode> FShaderLibraryInstance::AddResource(
 	return OutResource;
 }
 
-bool FShaderLibraryInstance::TryRemoveResource(FShaderMapResource_SharedCode* Resource)
+bool FShaderLibraryInstance::TryInsertResource(FShaderMapResource_SharedCode* InResource)
 {
-	FRWScopeLock Locker(ResourceLock, SLT_Write);
+	check(InResource);
+	check(InResource->LibraryInstance == this);
+
+	// Check if this entry is already occupied
+	{
+		FRWScopeLock Locker(ResourceLock, SLT_ReadOnly);
+		check(InResource->ShaderMapIndex < Resources.Num());
+		if (Resources[InResource->ShaderMapIndex] != nullptr)
+		{
+			return false;
+		}
+	}
+
+	// Insert resource entry
+	{
+		FRWScopeLock Locker(ResourceLock, SLT_Write);
+		// Check again since the slot might have been lost between the previous read-lock and this write-lock.
+		if (Resources[InResource->ShaderMapIndex] != nullptr)
+		{
+			return false;
+		}
+		Resources[InResource->ShaderMapIndex] = InResource;
+	}
+
+	return true;
+}
 
-	if (Resource->GetNumRefs() == 0)
+bool FShaderLibraryInstance::TryRemoveResource(FShaderMapResource_SharedCode* Resource, bool bForceRemoval)
+{
+	if (Resource->GetNumRefs() == 0 || bForceRemoval)
 	{
+		FRWScopeLock Locker(ResourceLock, SLT_Write);
+
+		// If this resource is not registered in this library instance, it must have been transferred from another library for a slot that was already occupied
 		const int32 ShaderMapIndex = Resource->ShaderMapIndex;
-		check(Resources[ShaderMapIndex] == Resource);
-		Resources[ShaderMapIndex] = nullptr;
-		return true;
+		if (Resources[ShaderMapIndex] == Resource)
+		{
+			Resources[ShaderMapIndex] = nullptr;
+			return true;
+		}
 	}
 
 	// Another thread found the resource after ref-count was decremented to zero
diff --git a/Engine/Source/Runtime/RenderCore/Private/ShaderLibrary/ShaderLibraryInstance.h b/Engine/Source/Runtime/RenderCore/Private/ShaderLibrary/ShaderLibraryInstance.h
--- a/Engine/Source/Runtime/RenderCore/Private/ShaderLibrary/ShaderLibraryInstance.h
+++ b/Engine/Source/Runtime/RenderCore/Private/ShaderLibrary/ShaderLibraryInstance.h
@@ -43,7 +43,10 @@ public:
 
 	TRefCountPtr<FShaderMapResource_SharedCode> AddResource(int32 ShaderMapIndex, FArchive* Ar);
 
-	bool TryRemoveResource(FShaderMapResource_SharedCode* Resource);
+	/** Tries to insert the specified shadermap resource into this library instance. */
+	bool TryInsertResource(FShaderMapResource_SharedCode* Resource);
+
+	bool TryRemoveResource(FShaderMapResource_SharedCode* Resource, bool bForceRemoval = false);
 
 	TRefCountPtr<FRHIShader> GetOrCreateShader(int32 ShaderIndex, bool bRequired = true);
 
diff --git a/Engine/Source/Runtime/RenderCore/Private/ShaderLibrary/ShaderMapResourceSharedCode.cpp b/Engine/Source/Runtime/RenderCore/Private/ShaderLibrary/ShaderMapResourceSharedCode.cpp
--- a/Engine/Source/Runtime/RenderCore/Private/ShaderLibrary/ShaderMapResourceSharedCode.cpp
+++ b/Engine/Source/Runtime/RenderCore/Private/ShaderLibrary/ShaderMapResourceSharedCode.cpp
@@ -79,6 +79,43 @@ FSHAHash FShaderMapResource_SharedCode::GetShaderMapHash() const
 	return LibraryInstance->Library->GetShaderMapHash(ShaderMapIndex);
 }
 
+void FShaderMapResource_SharedCode::AssignToNewShaderLibrary(FShaderLibraryInstance* NewLibraryInstance, int32 NewShaderMapIndex)
+{
+	const int32 OldShaderMapIndex = this->ShaderMapIndex;
+	FShaderLibraryInstance* OldLibraryInstance = this->LibraryInstance;
+
+	if (NewLibraryInstance == OldLibraryInstance)
+	{
+		checkf(NewShaderMapIndex == OldShaderMapIndex, TEXT("Shader library assignment to shadermap resource is unchanged, but shadermap indices don't match: Old(%d) != New(%d)"), OldShaderMapIndex, NewShaderMapIndex);
+		return;
+	}
+
+	// Release references to old RHI shaders in resource
+	ReleaseRHIShadersFromLibrary();
+
+	// Remove shadermap entry from old library
+	OldLibraryInstance->TryRemoveResource(this, /*bForceRemoval:*/ true);
+
+	// Set reference to new shader library in resource
+	this->LibraryInstance = NewLibraryInstance;
+	this->ShaderMapIndex = NewShaderMapIndex;
+
+	// Try to insert the new resource. If the specified slot is already occupied, the resource is considered detached from the library and we skip removing it in FShaderLibraryInstance::TryRemoveResource().
+	// Such resources will only live for as long as they are still in use and the shader library does not keep it cached as another resource has taken that spot.
+	const bool bIsResourceAttached = NewLibraryInstance->TryInsertResource(this);
+
+	UE_LOGF(LogShaderLibrary, Verbose, "Moved shadermap '%ls' (Refcount=%d, Owner='%ls') from '%ls' (Index=%d) to '%ls' (Index=%d) with %hs resource"
+		, *LexToString(this->GetShaderMapHash())
+		, this->GetNumRefs()
+		, *this->GetOwnerName().ToString()
+		, *OldLibraryInstance->Library->GetName()
+		, OldShaderMapIndex
+		, *NewLibraryInstance->Library->GetName()
+		, NewShaderMapIndex
+		, bIsResourceAttached ? "attached" : "detached"
+	);
+}
+
 void FShaderMapResource_SharedCode::ReleasePreloadedShaderCode(int32 ShaderIndex)
 {
 	SCOPED_LOADTIMER(FShaderMapResource_SharedCode_InitRHI);	// part of shader initialization in a way
@@ -118,9 +155,7 @@ void FShaderMapResource_SharedCode::PreloadShaderMap(FGraphEventArray& OutComple
 	}
 }
 
-void DiagnoseLeakingShaderMapResourceRefs(const FShaderMapResource_SharedCode* InResource);
-
-void FShaderMapResource_SharedCode::ReleaseRHI()
+void FShaderMapResource_SharedCode::ReleaseRHIShadersFromLibrary()
 {
 	if (LibraryInstance && ensureMsgf(LibraryInstance->Library, TEXT("LibraryInstance->Library pointer is expected to be valid as long as library's FShaderMapResource are alive.")))
 	{
@@ -164,15 +199,23 @@ void FShaderMapResource_SharedCode::ReleaseRHI()
 	bEntireShaderMapPreloaded = false;
 
 	FShaderMapResource::ReleaseRHI();
+}
+
+void FShaderMapResource_SharedCode::ReleaseRHI()
+{
+	ReleaseRHIShadersFromLibrary();
 
 	if (GetNumRefs() > 0)
 	{
-		// There shouldn't be any references left to this resource at this point. Diagnose what shader libraries are still holding a reference to this shader map resource.
-		DiagnoseLeakingShaderMapResourceRefs(this);
-
-		ensureMsgf(false, TEXT("FShaderMapResource_SharedCode::ReleaseRHI is still referenced (Num of references %d, owner %s). Invoking OnSharedShaderMapResourceExplicitRelease delegate."), GetNumRefs(), *GetOwnerName().ToString());
-		UE_LOG(LogShaderLibrary, Warning, TEXT("FShaderMapResource_SharedCode::ReleaseRHI is still referenced (Num of references %d, owner %s). Invoking OnSharedShaderMapResourceExplicitRelease delegate."),
-			GetNumRefs(), *GetOwnerName().ToString());
+		// Only print the shadermap index, not the hash since we cannot access it from the library anymore
+		ensureMsgf(false, TEXT("FShaderMapResource_SharedCode::ReleaseRHI is still referenced (Refcount=%d, Owner='%s')"), GetNumRefs(), *GetOwnerName().ToString());
+		UE_LOGF(LogShaderLibrary, Warning, "FShaderMapResource_SharedCode::ReleaseRHI is still referenced (Refcount=%d, ShaderMapHash='%ls', ShaderMapIndex=%d, Library='%ls', Owner='%ls')"
+			, GetNumRefs()
+			, LibraryInstance ? *LexToString(GetShaderMapHash()) : TEXT("<Null>")
+			, ShaderMapIndex
+			, LibraryInstance ? *LibraryInstance->Library->GetName() : TEXT("<Null>")
+			, *GetOwnerName().ToString()
+		);
 
 		// Invoke delegate to notify shader map resource needs to be forced released
 		OnSharedShaderMapResourceExplicitRelease.ExecuteIfBound(this);
diff --git a/Engine/Source/Runtime/RenderCore/Private/ShaderLibrary/ShaderMapResourceSharedCode.h b/Engine/Source/Runtime/RenderCore/Private/ShaderLibrary/ShaderMapResourceSharedCode.h
--- a/Engine/Source/Runtime/RenderCore/Private/ShaderLibrary/ShaderMapResourceSharedCode.h
+++ b/Engine/Source/Runtime/RenderCore/Private/ShaderLibrary/ShaderMapResourceSharedCode.h
@@ -10,6 +10,8 @@
 
 class FShaderMapResource_SharedCode final : public FShaderMapResource
 {
+	void ReleaseRHIShadersFromLibrary();
+
 public:
 	FShaderMapResource_SharedCode(class FShaderLibraryInstance* InLibraryInstance, int32 InShaderMapIndex);
 	virtual ~FShaderMapResource_SharedCode();
@@ -37,6 +39,9 @@ public:
 	virtual uint32 GetShaderSizeBytes(int32 ShaderIndex) const override;
 	virtual FSHAHash GetShaderMapHash() const override;
 
+	// Assigns this resource to the new shader library with its new shadermap index. See FShaderLibrariesCollection::MoveShaderMapResourceOwnership().
+	void AssignToNewShaderLibrary(FShaderLibraryInstance* NewLibraryInstance, int32 NewShaderMapIndex);
+
 	class FShaderLibraryInstance* LibraryInstance;
 	int32 ShaderMapIndex;
 	bool bEntireShaderMapPreloaded;
diff --git a/Engine/Source/Runtime/RenderCore/Private/ShaderResource.cpp b/Engine/Source/Runtime/RenderCore/Private/ShaderResource.cpp
--- a/Engine/Source/Runtime/RenderCore/Private/ShaderResource.cpp
+++ b/Engine/Source/Runtime/RenderCore/Private/ShaderResource.cpp
@@ -487,15 +487,17 @@ void FShaderMapResource::Release()
 
 void FShaderMapResource::ReleaseShaders()
 {
-	if (RHIShaders)
+	if (bAtLeastOneRHIShaderCreated)
 	{
+		// Note that this function must not reset the RHIShaders pointer nor NumRHIShaders. This will be destroyed automatically in the destructor.
+		// This function could also be used by a subclass to releases its RHI shaders to allow them to be re-created after they have been moved to a new shader library when the old one was unmounted.
 		FScopeLock ScopeLock(&RHIShadersCreationGuard);
 
 		int NumReleaseShaders = 0;
 
 		for (uint32 Idx = 0; Idx < NumRHIShaders; ++Idx)
 		{
-			if (FRHIShader* Shader = RHIShaders[Idx].load(std::memory_order_acquire))
+			if (FRHIShader* Shader = ExchangeShader(Idx, nullptr))
 			{
 				Shader->Release();
 				NumReleaseShaders++;
@@ -508,14 +510,11 @@ void FShaderMapResource::ReleaseShaders()
 		CsvStatNumShadersCreated->Sub(NumReleaseShaders);
 #endif
 
-		RHIShaders = nullptr;
-		NumRHIShaders = 0;
-		if (bAtLeastOneRHIShaderCreated)
 		{
 			DEC_DWORD_STAT(STAT_Shaders_NumShaderMapsUsedForRendering);
 
 #if (CSV_PROFILER_STATS && !UE_BUILD_SHIPPING) 
-			if(CsvStatNumShaderMapsUsedForRendering == nullptr)
+			if (CsvStatNumShaderMapsUsedForRendering == nullptr)
 			{
 				CsvStatNumShaderMapsUsedForRendering = FCsvProfiler::Get()->GetOrCreatePersistentCustomStatInt(TEXT("NumShaderMapsUsedForRendering"), CSV_CATEGORY_INDEX(Shaders));
 			}
diff --git a/Engine/Source/Runtime/RenderCore/Public/Shader.h b/Engine/Source/Runtime/RenderCore/Public/Shader.h
--- a/Engine/Source/Runtime/RenderCore/Public/Shader.h
+++ b/Engine/Source/Runtime/RenderCore/Public/Shader.h
@@ -471,8 +471,15 @@ protected:
 
 	virtual bool TryRelease() { return true; }
 
+	/** Releases all RHI shaders. This retains the internal array so that new RHI shaders can be created again, for example when a shadermap resource is being transferred to a new shader library instance. */
 	RENDERCORE_API void ReleaseShaders();
 
+	/** Exchanges an RHI shader. This is meant to be used by a shadermap resource that is transferred to a new shader library instance. */
+	inline FRHIShader* ExchangeShader(int32 ShaderIndex, FRHIShader* Shader)
+	{
+		return RHIShaders[ShaderIndex].exchange(Shader, std::memory_order_acquire);
+	}
+
 private:
 
 	/** Creates an entry in RHIShaders array and registers it among the raytracing libs if needed. Created shader is returned. */
