Track pso precache request result in pso file cache so it can be used at runtime to decide what to compile from the bundled pso cache

#Synced-CL 51591995

diff --git a/Engine/Source/Runtime/RHI/Private/PipelineFileCache.cpp b/Engine/Source/Runtime/RHI/Private/PipelineFileCache.cpp
--- a/Engine/Source/Runtime/RHI/Private/PipelineFileCache.cpp
+++ b/Engine/Source/Runtime/RHI/Private/PipelineFileCache.cpp
@@ -96,12 +96,13 @@ enum class EPipelineCacheFileFormatVersions : uint32
 	AddingDepthBounds = 28,
 	AddRTPSOShaderBindingLayout = 29,
 	VariableRateShading = 30,
+	PSOPrecacheResult = 31,
 };
 
 const uint64 FPipelineCacheFileFormatMagic = 0x5049504543414348; // PIPECACH
 const uint64 FPipelineCacheTOCFileFormatMagic = 0x544F435354415232; // TOCSTAR2
 const uint64 FPipelineCacheEOFFileFormatMagic = 0x454F462D4D41524B; // EOF-MARK
-const RHI_API uint32 FPipelineCacheFileFormatCurrentVersion = (uint32)EPipelineCacheFileFormatVersions::VariableRateShading;
+const RHI_API uint32 FPipelineCacheFileFormatCurrentVersion = (uint32)EPipelineCacheFileFormatVersions::PSOPrecacheResult;
 const int32  FPipelineCacheGraphicsDescPartsNum = 67; // parser will expect this number of parts in a description string
 
 /**
@@ -129,11 +130,11 @@ static TAutoConsoleVariable<int32> CVarPSOFileCacheReportPSO(
 														   ECVF_Default | ECVF_RenderThreadSafe
 														   );
 
-static int32 GPSOExcludePrecachePSOsInFileCache = 1;
+static int32 GPSOExcludePrecachePSOsInFileCache = 0;
 static FAutoConsoleVariableRef CVarPSOFileCacheExcludePrecachePSO(
 														   TEXT("r.ShaderPipelineCache.ExcludePrecachePSO"),
 														   GPSOExcludePrecachePSOsInFileCache,
-														   TEXT("1 excludes saving runtime-precached graphics & compute PSOs in the file cache which were successfully precached (default), 0 includes all of them. Excluding precached PSOs requires PSO precaching to be enabled."),
+														   TEXT("1 excludes saving runtime-precached graphics & compute PSOs in the file cache which were successfully precached, 0 includes all of them (default). Excluding precached PSOs requires PSO precaching to be enabled."),
 														   ECVF_ReadOnly
 														   );
 
@@ -970,7 +971,9 @@ FString FPipelineCacheFileFormatPSO::ToStringReadable() const
 	TReadableStringBuilder Builder;
 
 	Builder << TEXT("PSO hash ");
-	Builder << GetTypeHash(*this);
+	Builder << GetTypeHash(*this);	
+	Builder << TEXT(" PSOPrecacheResult ");
+	Builder << LexToString(PSOPrecacheResult);
 #if PSO_COOKONLY_DATA
 	Builder << TEXT(" mask ");
 	Builder << UsageMask;
@@ -1241,7 +1244,10 @@ bool FPipelineCacheFileFormatPSO::Verify() const
 /*friend*/ FArchive& operator<<( FArchive& Ar, FPipelineCacheFileFormatPSO& Info )
 {
 	Ar << Info.Type;
-	
+	if (Ar.GameNetVer() >= (uint32)EPipelineCacheFileFormatVersions::PSOPrecacheResult)
+	{
+		Ar << Info.PSOPrecacheResult;
+	}
 #if PSO_COOKONLY_DATA
 	/* Ignore: Ar << Info.UsageMask; during serialization */
 	/* Ignore: Ar << Info.BindCoun during  serialization*/
@@ -1420,11 +1426,13 @@ FPipelineCacheFileFormatPSO::FPipelineCacheFileFormatPSO()
 {
 }
 
-/*static*/ bool FPipelineCacheFileFormatPSO::Init(FPipelineCacheFileFormatPSO& PSO, FComputePipelineStateInitializer const& Init)
+/*static*/ bool FPipelineCacheFileFormatPSO::Init(FPipelineCacheFileFormatPSO& PSO, FComputePipelineStateInitializer const& Init, EPSOPrecacheResult PSOPrecacheResult)
 {
 	check(Init.ComputeShader != nullptr);
 
 	PSO.Type = DescriptorType::Compute;
+	PSO.PSOPrecacheResult = PSOPrecacheResult;
+
 #if PSO_COOKONLY_DATA
 	PSO.UsageMask = 0;
 	PSO.BindCount = 0;
@@ -1444,11 +1452,13 @@ FPipelineCacheFileFormatPSO::FPipelineCacheFileFormatPSO()
 	return bOK;
 }
 
-/*static*/ bool FPipelineCacheFileFormatPSO::Init(FPipelineCacheFileFormatPSO& PSO, FGraphicsPipelineStateInitializer const& Init)
+/*static*/ bool FPipelineCacheFileFormatPSO::Init(FPipelineCacheFileFormatPSO& PSO, FGraphicsPipelineStateInitializer const& Init, EPSOPrecacheResult PSOPrecacheResult)
 {
 	bool bOK = true;
 	
 	PSO.Type = DescriptorType::Graphics;
+	PSO.PSOPrecacheResult = PSOPrecacheResult;
+
 #if PSO_COOKONLY_DATA
 	PSO.UsageMask = 0;
 	PSO.BindCount = 0;
@@ -1596,6 +1606,8 @@ bool FPipelineCacheFileFormatPSO::operator==(const FPipelineCacheFileFormatPSO&
 	{
 		bSame = Type == Other.Type;
 
+		// ignore pso precache result for equals op
+
 #if PSO_COOKONLY_DATA
 		/* Ignore: [UsageMask == UsageMask] in this test. */
 		/* Ignore: [BindCount == BindCount] in this test. */
@@ -1658,7 +1670,7 @@ bool FPipelineCacheFileFormatPSO::operator==(const FPipelineCacheFileFormatPSO&
 }
 
 FPipelineCacheFileFormatPSO::FPipelineCacheFileFormatPSO(const FPipelineCacheFileFormatPSO& Other)
-: Type(Other.Type)
+: Type(Other.Type), PSOPrecacheResult(Other.PSOPrecacheResult)
 #if PSO_COOKONLY_DATA
 , UsageMask(Other.UsageMask)
 , BindCount(Other.BindCount)
@@ -1696,6 +1708,7 @@ FPipelineCacheFileFormatPSO& FPipelineCacheFileFormatPSO::operator=(const FPipel
 	if(this != &Other)
 	{
 		Type = Other.Type;
+		PSOPrecacheResult = Other.PSOPrecacheResult;
 #if PSO_COOKONLY_DATA
 		UsageMask = Other.UsageMask;
 		BindCount = Other.BindCount;
@@ -3566,7 +3579,7 @@ void FPipelineFileCacheManager::CacheGraphicsPSO(uint32 RunTimeHash, FGraphicsPi
 			if(PSOUsage == nullptr)
 			{
 				FPipelineCacheFileFormatPSO NewEntry;
-				bool bOK = FPipelineCacheFileFormatPSO::Init(NewEntry, Initializer);
+				bool bOK = FPipelineCacheFileFormatPSO::Init(NewEntry, Initializer, PSOPrecacheResult);
 				check(bOK);
 					
 				uint32 PSOHash = GetTypeHash(NewEntry);
@@ -3659,7 +3672,7 @@ void FPipelineFileCacheManager::CacheComputePSO(uint32 RunTimeHash, FComputePipe
 				if(PSOUsage == nullptr)
 				{
 					FPipelineCacheFileFormatPSO NewEntry;
-					bool bOK = FPipelineCacheFileFormatPSO::Init(NewEntry, Initializer);
+					bool bOK = FPipelineCacheFileFormatPSO::Init(NewEntry, Initializer, PSOPrecacheResult);
 					check(bOK);
 					
 					uint32 PSOHash = GetTypeHash(NewEntry);
@@ -3835,7 +3848,7 @@ void FPipelineFileCacheManager::RegisterPSOCompileFailure(uint32 RunTimeHash, FG
 			if(PSOUsage == nullptr)
 			{
 				FPipelineCacheFileFormatPSO ShouldBeExistingEntry;
-				bool bOK = FPipelineCacheFileFormatPSO::Init(ShouldBeExistingEntry, Initializer);
+				bool bOK = FPipelineCacheFileFormatPSO::Init(ShouldBeExistingEntry, Initializer, EPSOPrecacheResult::Untracked);
 				check(bOK);
 					
 				uint32 PSOHash = GetTypeHash(ShouldBeExistingEntry);
diff --git a/Engine/Source/Runtime/RHI/Public/PipelineFileCache.h b/Engine/Source/Runtime/RHI/Public/PipelineFileCache.h
--- a/Engine/Source/Runtime/RHI/Public/PipelineFileCache.h
+++ b/Engine/Source/Runtime/RHI/Public/PipelineFileCache.h
@@ -8,6 +8,7 @@
 #include "Containers/StringView.h"
 #include "RHI.h"
 #include "RHIShaderBindingLayout.h"
+#include "PipelineStateCache.h"
 
 DECLARE_STATS_GROUP(TEXT("ShaderPipelineCache"),STATGROUP_PipelineStateCache, STATCAT_Advanced);
 
@@ -213,6 +214,7 @@ struct FPipelineCacheFileFormatPSO
 	ComputeDescriptor ComputeDesc;
 	GraphicsDescriptor GraphicsDesc;
 	FPipelineFileCacheRayTracingDesc RayTracingDesc;
+	EPSOPrecacheResult PSOPrecacheResult = EPSOPrecacheResult::Unknown;
 
 #if PSO_COOKONLY_DATA
 	uint64 UsageMask;
@@ -230,8 +232,8 @@ struct FPipelineCacheFileFormatPSO
 	friend RHI_API uint32 GetTypeHash(const FPipelineCacheFileFormatPSO &Key);
 	friend RHI_API FArchive& operator<<( FArchive& Ar, FPipelineCacheFileFormatPSO& Info );
 	
-	static bool Init(FPipelineCacheFileFormatPSO& PSO, FComputePipelineStateInitializer const& Init);
-	static bool Init(FPipelineCacheFileFormatPSO& PSO, FGraphicsPipelineStateInitializer const& Init);
+	static bool Init(FPipelineCacheFileFormatPSO& PSO, FComputePipelineStateInitializer const& Init, EPSOPrecacheResult PSOPrecacheResult);
+	static bool Init(FPipelineCacheFileFormatPSO& PSO, FGraphicsPipelineStateInitializer const& Init, EPSOPrecacheResult PSOPrecacheResult);
 	static bool Init(FPipelineCacheFileFormatPSO & PSO, FPipelineFileCacheRayTracingDesc const& Desc);
 
 	RHI_API static void SortVertexElements(FVertexDeclarationElementList& VertexElements);
diff --git a/Engine/Source/Runtime/RenderCore/Private/ShaderPipelineCache.cpp b/Engine/Source/Runtime/RenderCore/Private/ShaderPipelineCache.cpp
--- a/Engine/Source/Runtime/RenderCore/Private/ShaderPipelineCache.cpp
+++ b/Engine/Source/Runtime/RenderCore/Private/ShaderPipelineCache.cpp
@@ -37,6 +37,11 @@ DECLARE_DWORD_COUNTER_STAT(TEXT("# Pipelines Pre-Compiled"), STAT_PreCompileShad
 DECLARE_DWORD_ACCUMULATOR_STAT(TEXT("Total Batches Pre-Compiled"), STAT_PreCompileBatchTotal, STATGROUP_PipelineStateCache);
 DECLARE_DWORD_COUNTER_STAT(TEXT("# Batches Pre-Compiled"), STAT_PreCompileBatchNum, STATGROUP_PipelineStateCache);
 
+DECLARE_DWORD_COUNTER_STAT(TEXT("Graphic PSOs Compiled"), STAT_GraphicPSOsCompiled, STATGROUP_PipelineStateCache);
+DECLARE_DWORD_COUNTER_STAT(TEXT("Graphic PSOs Skipped"), STAT_GraphicPSOsSkipped, STATGROUP_PipelineStateCache);
+DECLARE_DWORD_COUNTER_STAT(TEXT("Compute PSOs Compiled"), STAT_ComputePSOsCompiled, STATGROUP_PipelineStateCache);
+DECLARE_DWORD_COUNTER_STAT(TEXT("Compute PSOs Skipped"), STAT_ComputePSOsSkipped, STATGROUP_PipelineStateCache);
+
 namespace FShaderPipelineCacheConstants
 {
 	static TCHAR const* SectionHeading = TEXT("ShaderPipelineCache.CacheFile");
@@ -183,6 +188,17 @@ static TAutoConsoleVariable<int32> CVarPSOFileCacheOnlyOpenUserCache(
 																ECVF_ReadOnly
                                                                 );
 
+static int32 GPSOPrecacheIncludeMode = 2;
+static FAutoConsoleVariableRef CVarPSOFileCachePSOPrecacheIncludeMode(
+	TEXT("r.ShaderPipelineCache.PSOPrecacheIncludeMode"),
+	GPSOPrecacheIncludeMode,
+	TEXT("Defines which PSOs from the file cache are compiled depending on the PSOPrecacheResult which is stored during collection:\n") \
+	TEXT("0: Include all PSOs independent of PSO precache result\n") \
+	TEXT("1: Only include PSOs which were not successfully compiled with PSO precaching (Missed, Untracked & TooLate)\n") \
+	TEXT("2: Same as 1, but always precache all PSOs coming from global shaders (default)"),
+	ECVF_ReadOnly
+);
+
 static bool GetShaderPipelineCacheSaveBoundPSOLog()
 {
 	static bool bOnce = false;
@@ -1036,6 +1052,31 @@ bool FShaderPipelineCacheTask::Precompile(FRHICommandListImmediate& RHICmdList,
 			return bGlobalShader && FShaderCodeLibrary::ContainsShaderCode(Hash, GlobalLibraryName);
 		};
 
+		auto ShouldCompilePSO = [PSO](bool bGlobalShader)
+		{
+			// Check if we should skip compiling this PSO due to PSO precaching state
+			if (GPSOPrecacheIncludeMode > 0)
+			{
+				bool bIsPrecached = PSO.PSOPrecacheResult == EPSOPrecacheResult::Complete;
+
+				// Include all global shaders when mode > 1
+				if (GPSOPrecacheIncludeMode > 1)
+				{
+					if (bIsPrecached && !bGlobalShader)
+					{
+						return false;
+					}
+				}
+				// Otherwise ignore all PSOs which are precached correctly
+				else if (bIsPrecached)
+				{
+					return false;
+				}
+			}
+
+			return true;
+		};
+
 		if(FPipelineCacheFileFormatPSO::DescriptorType::Graphics == PSO.Type)
 		{
 			TRACE_CPUPROFILER_EVENT_SCOPE(FShaderPipelineCache::PrecompileGraphics);
@@ -1147,31 +1188,42 @@ bool FShaderPipelineCacheTask::Precompile(FRHICommandListImmediate& RHICmdList,
 				}
 #endif
 			}
-			// Use SetGraphicsPipelineState to call down into PipelineStateCache and also handle the fallback case used by OpenGL.
-			// we lose track of our PSO precompile jobs through here. they may not be RHI dependencies and could persist beyond the next tick.
-			// GetNumActivePipelinePrecompileTasks() is tracking non-rhi dependent precompiles.
-			// It's being used for rate limiting and ensuring tasks are done before calling OnPrecompilationComplete delegate.
-			if (UsePSOPrecachingThreadPool())
-			{	
-				// Remove file cache flag because it's going over precaching code path				
-				GraphicsInitializer.bFromPSOFileCache = false;
-				GraphicsInitializer.StatePrecachePSOHash = RHIComputeStatePrecachePSOHash(GraphicsInitializer);
-
-				static FString BundledGlobalGraphicsPSOCollectorName = TEXT("BundledGlobalGraphicsPSO");
-				static FString BundledMeshPassGraphicsPSOCollectorName = TEXT("BundledMeshPassGraphicsPSO");
-
-				// Async compile requests
-				FPSOPrecacheRequestParams PrecacheRequestParams;
-				PrecacheRequestParams.PrecachePriority = bOnlyGlobalShaders ? EPSOPrecachePriority::High : EPSOPrecachePriority::Medium;
-				PrecacheRequestParams.InMemoryMode = bOnlyGlobalShaders ? EPSOPrecacheInMemoryMode::Cached : EPSOPrecacheInMemoryMode::Uncached;
-				PrecacheRequestParams.CollectorName = bOnlyGlobalShaders ? *BundledGlobalGraphicsPSOCollectorName : *BundledMeshPassGraphicsPSOCollectorName;
-				PipelineStateCache::PrecacheGraphicsPipelineState(GraphicsInitializer, PrecacheRequestParams);
+
+			if (ShouldCompilePSO(bOnlyGlobalShaders))
+			{
+				// Use SetGraphicsPipelineState to call down into PipelineStateCache and also handle the fallback case used by OpenGL.
+				// we lose track of our PSO precompile jobs through here. they may not be RHI dependencies and could persist beyond the next tick.
+				// GetNumActivePipelinePrecompileTasks() is tracking non-rhi dependent precompiles.
+				// It's being used for rate limiting and ensuring tasks are done before calling OnPrecompilationComplete delegate.
+				if (UsePSOPrecachingThreadPool())
+				{
+					// Remove file cache flag because it's going over precaching code path				
+					GraphicsInitializer.bFromPSOFileCache = false;
+					GraphicsInitializer.StatePrecachePSOHash = RHIComputeStatePrecachePSOHash(GraphicsInitializer);
+
+					static FString BundledGlobalGraphicsPSOCollectorName = TEXT("BundledGlobalGraphicsPSO");
+					static FString BundledMeshPassGraphicsPSOCollectorName = TEXT("BundledMeshPassGraphicsPSO");
+
+					// Async compile requests
+					FPSOPrecacheRequestParams PrecacheRequestParams;
+					PrecacheRequestParams.PrecachePriority = bOnlyGlobalShaders ? EPSOPrecachePriority::High : EPSOPrecachePriority::Medium;
+					PrecacheRequestParams.InMemoryMode = bOnlyGlobalShaders ? EPSOPrecacheInMemoryMode::Cached : EPSOPrecacheInMemoryMode::Uncached;
+					PrecacheRequestParams.CollectorName = bOnlyGlobalShaders ? *BundledGlobalGraphicsPSOCollectorName : *BundledMeshPassGraphicsPSOCollectorName;
+					PipelineStateCache::PrecacheGraphicsPipelineState(GraphicsInitializer, PrecacheRequestParams);
+				}
+				else
+				{
+					QUICK_SCOPE_CYCLE_COUNTER(STAT_FShaderPipelineCache_SetGraphicsPipelineState);
+					SetGraphicsPipelineState(RHICmdList, GraphicsInitializer, 0, EApplyRendertargetOption::DoNothing, false);
+				}
+
+				INC_DWORD_STAT(STAT_GraphicPSOsCompiled);
 			}
 			else
 			{
-				QUICK_SCOPE_CYCLE_COUNTER(STAT_FShaderPipelineCache_SetGraphicsPipelineState);
-				SetGraphicsPipelineState(RHICmdList, GraphicsInitializer, 0, EApplyRendertargetOption::DoNothing, false);
+				INC_DWORD_STAT(STAT_GraphicPSOsSkipped);
 			}
+
 			bOk = true;
 		}
 		else if(FPipelineCacheFileFormatPSO::DescriptorType::Compute == PSO.Type)
@@ -1185,26 +1237,36 @@ bool FShaderPipelineCacheTask::Precompile(FRHICommandListImmediate& RHICmdList,
 				{
 					bool bIsGlobalShader = IsGlobalShader(true, PSO.ComputeDesc.ComputeShader);
 
-					FComputePipelineStateInitializer ComputeInitializer(ComputeShaderRef);
-					if (UsePSOPrecachingThreadPool())
-					{									
-						static FString BundledGlobalComputePSOCollectorName = TEXT("BundledGlobalComputePSO");
-						static FString BundledMeshPassComputePSOCollectorName = TEXT("BundledMeshPassComputePSO");
+					if (ShouldCompilePSO(bIsGlobalShader))
+					{
+						FComputePipelineStateInitializer ComputeInitializer(ComputeShaderRef);
+						if (UsePSOPrecachingThreadPool())
+						{
+							static FString BundledGlobalComputePSOCollectorName = TEXT("BundledGlobalComputePSO");
+							static FString BundledMeshPassComputePSOCollectorName = TEXT("BundledMeshPassComputePSO");
 						
-						FPSOPrecacheRequestParams PrecacheRequestParams;
-						PrecacheRequestParams.PrecachePriority = bIsGlobalShader ? EPSOPrecachePriority::High : EPSOPrecachePriority::Medium;
-						PrecacheRequestParams.InMemoryMode = bIsGlobalShader ? EPSOPrecacheInMemoryMode::Cached : EPSOPrecacheInMemoryMode::Uncached;
-						PrecacheRequestParams.CollectorName = bIsGlobalShader ? *BundledGlobalComputePSOCollectorName : *BundledMeshPassComputePSOCollectorName;
+							FPSOPrecacheRequestParams PrecacheRequestParams;
+							PrecacheRequestParams.PrecachePriority = bIsGlobalShader ? EPSOPrecachePriority::High : EPSOPrecachePriority::Medium;
+							PrecacheRequestParams.InMemoryMode = bIsGlobalShader ? EPSOPrecacheInMemoryMode::Cached : EPSOPrecacheInMemoryMode::Uncached;
+							PrecacheRequestParams.CollectorName = bIsGlobalShader ? *BundledGlobalComputePSOCollectorName : *BundledMeshPassComputePSOCollectorName;
 
-						PipelineStateCache::PrecacheComputePipelineState(ComputeInitializer, PrecacheRequestParams);
-						bOk = true;
+							PipelineStateCache::PrecacheComputePipelineState(ComputeInitializer, PrecacheRequestParams);
+							bOk = true;
+						}
+						else
+						{
+							ComputeInitializer.bFromPSOFileCache = true;
+
+							FComputePipelineState* ComputeResult = PipelineStateCache::GetAndOrCreateComputePipelineState(RHICmdList, ComputeInitializer);
+							bOk = ComputeResult != nullptr;
+						}
+
+						INC_DWORD_STAT(STAT_ComputePSOsCompiled);
 					}
 					else
 					{
-						ComputeInitializer.bFromPSOFileCache = true;
-
-						FComputePipelineState* ComputeResult = PipelineStateCache::GetAndOrCreateComputePipelineState(RHICmdList, ComputeInitializer);
-						bOk = ComputeResult != nullptr;
+						INC_DWORD_STAT(STAT_ComputePSOsSkipped);
+						bOk = true;
 					}
 				}
 			}
