CookIndetermism: ShaderLibrary: Sort the cooked ShaderLibrary by ShaderHash when writing to disk. Enabled only in iostore builds.
#rnx
#Synced-CL 51754452

diff --git a/Engine/Source/Editor/UnrealEd/Classes/CookOnTheSide/CookOnTheFlyServer.h b/Engine/Source/Editor/UnrealEd/Classes/CookOnTheSide/CookOnTheFlyServer.h
--- a/Engine/Source/Editor/UnrealEd/Classes/CookOnTheSide/CookOnTheFlyServer.h
+++ b/Engine/Source/Editor/UnrealEd/Classes/CookOnTheSide/CookOnTheFlyServer.h
@@ -877,6 +877,8 @@ public:
 	 * than as loose files on disk)?
 	 */
 	UNREALED_API bool IsUsingZenStore() const;
+	/** Is the databuild using IoStore (cooked packages will be stored in iostore containers rather than pak files)? */
+	UNREALED_API bool IsUsingIoStore() const;
 
 	/**
 	 * Is the local CookOnTheFlyServer initialized to run CookOnTheFly (accept connections from game executables,
diff --git a/Engine/Source/Editor/UnrealEd/Private/CookOnTheFlyServer.cpp b/Engine/Source/Editor/UnrealEd/Private/CookOnTheFlyServer.cpp
--- a/Engine/Source/Editor/UnrealEd/Private/CookOnTheFlyServer.cpp
+++ b/Engine/Source/Editor/UnrealEd/Private/CookOnTheFlyServer.cpp
@@ -1231,6 +1231,12 @@ bool UCookOnTheFlyServer::IsUsingZenStore() const
 	return bZenStore;
 }
 
+bool UCookOnTheFlyServer::IsUsingIoStore() const
+{
+	const UProjectPackagingSettings* const PackagingSettings = GetDefault<UProjectPackagingSettings>();
+	return PackagingSettings->bUseIoStore;
+}
+
 bool UCookOnTheFlyServer::IsCookOnTheFlyMode() const
 {
 	return ::IsCookOnTheFlyMode(CurrentCookMode);
diff --git a/Engine/Source/Editor/UnrealEd/Private/Cooker/ShaderLibraryCooking.cpp b/Engine/Source/Editor/UnrealEd/Private/Cooker/ShaderLibraryCooking.cpp
--- a/Engine/Source/Editor/UnrealEd/Private/Cooker/ShaderLibraryCooking.cpp
+++ b/Engine/Source/Editor/UnrealEd/Private/Cooker/ShaderLibraryCooking.cpp
@@ -614,8 +614,10 @@ void UCookOnTheFlyServer::SaveShaderLibrary(const ITargetPlatform* TargetPlatfor
 	TArray<FString>& PlatformSCLCSVPaths = OutSCLCSVPaths.FindOrAdd(FName(TargetPlatform->PlatformName()));
 	FString ErrorString;
 	bool bHasData;
+	UE::Cook::EPackagingSystem PackagingSystem = IsUsingIoStore() ? UE::Cook::EPackagingSystem::IoStore
+		: UE::Cook::EPackagingSystem::Loose;
 	if (!FShaderLibraryCooker::SaveShaderLibraryWithoutChunking(TargetPlatform, Name, ShaderCodeDir, MetaDataPath,
-		PlatformSCLCSVPaths, ErrorString, bHasData))
+		PlatformSCLCSVPaths, PackagingSystem, ErrorString, bHasData))
 	{
 		// This is fatal - In this case we should cancel any launch on device operation or package write but we don't want to assert and crash the editor
 		LogCookerMessage(FString::Printf(TEXT("%s"), *ErrorString), EMessageSeverity::Error);
diff --git a/Engine/Source/Editor/UnrealEd/Private/ShaderLibraryChunkDataGenerator.cpp b/Engine/Source/Editor/UnrealEd/Private/ShaderLibraryChunkDataGenerator.cpp
--- a/Engine/Source/Editor/UnrealEd/Private/ShaderLibraryChunkDataGenerator.cpp
+++ b/Engine/Source/Editor/UnrealEd/Private/ShaderLibraryChunkDataGenerator.cpp
@@ -70,8 +70,10 @@ void FShaderLibraryChunkDataGenerator::GenerateChunkDataFiles(const int32 InChun
 		COTFS.GetShaderLibraryPaths(TargetPlatform, ShaderCodeDir, MetaDataPath, true /* bUseProjectDirForDLC */);
 
 		bool bHasData;
+		UE::Cook::EPackagingSystem PackagingSystem = COTFS.IsUsingIoStore() ? UE::Cook::EPackagingSystem::IoStore
+			: UE::Cook::EPackagingSystem::Loose;
 		FShaderLibraryCooker::SaveShaderLibraryChunk(InChunkId, InPackagesInChunk, TargetPlatform, ShaderCodeDir,
-			MetaDataPath, OutChunkFilenames, bHasData);
+			MetaDataPath, PackagingSystem, OutChunkFilenames, bHasData);
 	}
 }
 	
\ No newline at end of file
diff --git a/Engine/Source/Runtime/RenderCore/Internal/ShaderLibrary/EditorShaderCodeArchive.h b/Engine/Source/Runtime/RenderCore/Internal/ShaderLibrary/EditorShaderCodeArchive.h
--- a/Engine/Source/Runtime/RenderCore/Internal/ShaderLibrary/EditorShaderCodeArchive.h
+++ b/Engine/Source/Runtime/RenderCore/Internal/ShaderLibrary/EditorShaderCodeArchive.h
@@ -58,6 +58,7 @@ class FEditorShaderCodeArchive
 {
 	using ECookShaderLibrarySource = ShaderCodeArchive::ECookShaderLibrarySource;
 	using ESaveToDiskTarget = ShaderCodeArchive::ESaveToDiskTarget;
+	using ESaveToDiskSortOrder = ShaderCodeArchive::ESaveToDiskSortOrder;
 
 	/** Current library name, e.g. "Global" or "Lyra_Chunk0". */
 	FString LibraryName;
@@ -157,7 +158,8 @@ public:
 	 * Asset-shader associations to *.assetinfo.json file,
 	 * and shader type information to *.stinfo file. */
 	RENDERCORE_API bool SaveToDisk(const FString& OutputDir, const FString& MetaOutputDir,
-		ESaveToDiskTarget Target = ESaveToDiskTarget::Staging, TArray<FString>* OutputFilenames = nullptr);
+		ESaveToDiskTarget Target, ESaveToDiskSortOrder SortOrder,
+		TArray<FString>* OutputFilenames = nullptr);
 
 	/** Creates a native shader library for this archive, returns the output filenames,
 	 * and deletes the shader code and pipeline files as they are no longer needed. */
diff --git a/Engine/Source/Runtime/RenderCore/Internal/ShaderLibrary/ShaderCodeArchiveTypes.h b/Engine/Source/Runtime/RenderCore/Internal/ShaderLibrary/ShaderCodeArchiveTypes.h
--- a/Engine/Source/Runtime/RenderCore/Internal/ShaderLibrary/ShaderCodeArchiveTypes.h
+++ b/Engine/Source/Runtime/RenderCore/Internal/ShaderLibrary/ShaderCodeArchiveTypes.h
@@ -15,4 +15,22 @@ namespace ShaderCodeArchive
 		Staging = 0,
 		CookCache,
 	};
+
+	enum class ESaveToDiskSortOrder : uint8
+	{
+		/**
+		 * Shaders of a cooked shader library are sorted by the order in which they are
+		 * discovered during cook. This is better for performance when the library is loaded at runtime, since
+		 * runtime load order is correlated with cooker load order. But the correlation is not perfect, and cooker
+		 * load order is non-deterministic. This sort order is not beneficial at all if IoStore is used, because
+		 * the libraries are sorted into a different sort order when IoStore puts them into containers.
+		 */
+		PackageLoad,
+		/**
+		 * Shaders of a cooked shader library are sorted by their hash values. This is useful for cook determinism:
+		 * the bytes of the shader library in two different cooks are identical if shader inputs are unchanged between
+		 * the cooks.
+		 */
+		ShaderHash,
+	};
 } // namespace ShaderCodeArchive
diff --git a/Engine/Source/Runtime/RenderCore/Private/ShaderLibrary/EditorShaderCodeArchive.cpp b/Engine/Source/Runtime/RenderCore/Private/ShaderLibrary/EditorShaderCodeArchive.cpp
--- a/Engine/Source/Runtime/RenderCore/Private/ShaderLibrary/EditorShaderCodeArchive.cpp
+++ b/Engine/Source/Runtime/RenderCore/Private/ShaderLibrary/EditorShaderCodeArchive.cpp
@@ -899,7 +899,7 @@ void FEditorShaderCodeArchive::AssignFrom(FEditorShaderCodeArchive& InArchive)
 }
 
 bool FEditorShaderCodeArchive::SaveToDisk(const FString& OutputDir, const FString& MetaOutputDir,
-	ESaveToDiskTarget Target, TArray<FString>* OutputFilenames)
+	ESaveToDiskTarget Target, ESaveToDiskSortOrder SortOrder, TArray<FString>* OutputFilenames)
 {
 	check(LibraryName.Len() > 0);
 	checkf(!bHasCopiedAndCleared, TEXT("It is not valid to call SaveToDisk on an FEditorShaderCodeArchive that has sent its shaders to another process by calling CopyAndClear."));
@@ -951,11 +951,27 @@ bool FEditorShaderCodeArchive::SaveToDisk(const FString& OutputDir, const FStrin
 		constexpr bool bSaveOnlyAssetInfo = false; // A currently unused setting that we might want to bring back.
 		if (!bSaveOnlyAssetInfo)
 		{
+			switch (SortOrder)
+			{
+			case ESaveToDiskSortOrder::PackageLoad:
+				// Leave shaders arranged in the order in which they were created when loading and saving packages. This
+				// is better for runtime performance when the shader library we write is loaded directly from a loose
+				// file rather than being loaded through IoStore.
+				break;
+			case ESaveToDiskSortOrder::ShaderHash:
+				// Sort the shader library for deterministic output when no cook inputs have changed. Note that this
+				// depends on the ShaderHashes, so if they change the shader library sorting will change.
+				SerializedShaders.SortByHashes(ShaderCode);
+				break;
+			default:
+				checkNoEntry();
+				break;
+			}
+			SerializedShaders.Finalize();
+
 			const FString IntermediateFormatPath = GetShaderCodeFilename(IntermediatePathBase, LibraryName, FormatAndPlatformName, Target);
 			if (TUniquePtr<FArchive> FileWriter = TUniquePtr<FArchive>(IFileManager::Get().CreateFileWriter(*IntermediateFormatPath, FILEWRITE_NoFail)))
 			{
-				SerializedShaders.Finalize();
-
 				FSerializedShaderArchive::SerializeHeaderVersion(*FileWriter);
 
 				// Write shader library
@@ -1091,7 +1107,7 @@ bool FEditorShaderCodeArchive::CreatePatchLibrary(FName InFormatAndPlatformName,
 		{
 			FString Empty;
 			OutLibrary.FinishPopulate(OutDir, ECookShaderLibrarySource::PatchBase);
-			bOK = OutLibrary.SaveToDisk(OutDir, Empty);
+			bOK = OutLibrary.SaveToDisk(OutDir, Empty, ESaveToDiskTarget::Staging, ESaveToDiskSortOrder::PackageLoad);
 			UE_CLOG(!bOK, LogShaderLibrary, Error, TEXT("Failed to save %s shader patch library %s, %s, %s"), bNativeFormat ? TEXT("native") : TEXT(""), *InFormatAndPlatformName.ToString(), *LibraryName, *OutDir);
 
 			if (bOK && bNativeFormat && OutLibrary.SupportsShaderArchives())
diff --git a/Engine/Source/Runtime/RenderCore/Private/ShaderLibrary/ShaderCodeArchive.cpp b/Engine/Source/Runtime/RenderCore/Private/ShaderLibrary/ShaderCodeArchive.cpp
--- a/Engine/Source/Runtime/RenderCore/Private/ShaderLibrary/ShaderCodeArchive.cpp
+++ b/Engine/Source/Runtime/RenderCore/Private/ShaderLibrary/ShaderCodeArchive.cpp
@@ -27,6 +27,7 @@
 #include "IO/IoDispatcherInternal.h"
 
 #if WITH_EDITORONLY_DATA
+#include "Algo/Sort.h"
 #include "Misc/Optional.h"
 #include "Misc/StringBuilder.h"
 #include "Serialization/CompactBinarySerialization.h"
@@ -667,6 +668,133 @@ void FSerializedShaderArchive::DecompressShaderEntryAndAppend(const FShaderCodeE
 }
 
 #if !USE_MMAPPED_SHADERARCHIVE
+#if WITH_EDITORONLY_DATA
+void FSerializedShaderArchive::SortByHashes(TArray<FSharedBuffer>& InOutShaderCodeForShaderIndex)
+{
+	const uint32 NumShaders = ShaderEntries.Num();
+	const uint32 NumShaderMaps = ShaderMapEntries.Num();
+	check(ShaderHashes.Num() == NumShaders);
+	check(ShaderTypes.Num() == NumShaders);
+	check(InOutShaderCodeForShaderIndex.Num() == NumShaders);
+	check(ShaderMapHashes.Num() == NumShaderMaps);
+
+	// Find the SortOrder of Shaders
+	TArray<uint32> ShaderSortOrder;
+	ShaderSortOrder.SetNum(NumShaders);
+	for (uint32 Index = 0; Index < NumShaders; ++Index)
+	{
+		ShaderSortOrder[Index] = Index;
+	}
+	Algo::Sort(ShaderSortOrder, [this](uint32 AIndex, uint32 BIndex)
+		{
+			return ShaderHashes[AIndex] < ShaderHashes[BIndex];
+		});
+	TMap<uint32, uint32> ShaderIndexRemap;
+	ShaderIndexRemap.Reserve(NumShaders);
+	for (uint32 NewIndex = 0; NewIndex < NumShaders; ++NewIndex)
+	{
+		ShaderIndexRemap.Add(ShaderSortOrder[NewIndex], NewIndex);
+	}
+
+	// Remap all values that store an index into ShaderEntries (except ShaderHashTable, that's done below).
+	for (uint32& ShaderIndex : ShaderIndices)
+	{
+		ShaderIndex = ShaderIndexRemap.FindChecked(ShaderIndex);
+	}
+
+	// Shuffle Shaders into the new order
+	TArrayType<FShaderCodeEntry> NewShaderEntries;
+	TArrayType<FSHAHash> NewShaderHashes;
+	TArray<FShaderTypeHashes> NewShaderTypes;
+	TArrayType<FSharedBuffer> NewShaderCodeForShaderIndex;
+	NewShaderEntries.SetNum(NumShaders);
+	NewShaderHashes.SetNum(NumShaders);
+	NewShaderTypes.SetNum(NumShaders);
+	NewShaderCodeForShaderIndex.SetNum(NumShaders);
+	for (uint32 NewIndex = 0; NewIndex < NumShaders; ++NewIndex)
+	{
+		const uint32 OldIndex = ShaderSortOrder[NewIndex];
+		NewShaderEntries[NewIndex] = MoveTemp(ShaderEntries[OldIndex]);
+		NewShaderHashes[NewIndex] = MoveTemp(ShaderHashes[OldIndex]);
+		NewShaderTypes[NewIndex] = MoveTemp(ShaderTypes[OldIndex]);
+		NewShaderCodeForShaderIndex[NewIndex] = MoveTemp(InOutShaderCodeForShaderIndex[OldIndex]);
+	}
+	ShaderHashes = MoveTemp(NewShaderHashes);
+	ShaderEntries = MoveTemp(NewShaderEntries);
+	ShaderTypes = MoveTemp(NewShaderTypes);
+	InOutShaderCodeForShaderIndex = MoveTemp(NewShaderCodeForShaderIndex);
+
+	// Remap ShaderHashTable
+	ShaderHashTable.Clear(); // Resets values but does not deallocate
+	for (uint32 NewIndex = 0; NewIndex < NumShaders; ++NewIndex)
+	{
+		const uint32 Key = GetTypeHash(ShaderHashes[NewIndex]);
+		ShaderHashTable.Add(Key, NewIndex);
+	}
+
+	// PreloadEntries are recreated in deterministic order during Finalize, so we do not need to sort them.
+
+	// Sort each ShaderMap's internal data
+	TArrayView<uint32> ShaderIndicesView = TArrayView<uint32>(ShaderIndices);
+	for (FShaderMapEntry& Entry : ShaderMapEntries)
+	{
+		if (Entry.NumShaders > 0)
+		{
+			TArrayView<uint32> Indices = ShaderIndicesView.Mid(Entry.ShaderIndicesOffset, Entry.NumShaders);
+			Algo::Sort(Indices);
+		}
+		// PreloadIndices are recreated during Finalize, so we do not need to sort them.
+	}
+
+	// Find the SortOrder of ShaderMaps
+	TArray<uint32> ShaderMapSortOrder;
+	ShaderMapSortOrder.SetNum(NumShaderMaps);
+	for (uint32 Index = 0; Index < NumShaderMaps; ++Index)
+	{
+		ShaderMapSortOrder[Index] = Index;
+	}
+	Algo::Sort(ShaderMapSortOrder, [this](uint32 AIndex, uint32 BIndex)
+		{
+			return ShaderMapHashes[AIndex] < ShaderMapHashes[BIndex];
+		});
+
+	// Shuffle ShaderMaps into the new order
+	TArrayType<FShaderMapEntry> NewShaderMapEntries;
+	TArrayType<FSHAHash> NewShaderMapHashes;
+	TArrayType<uint32> NewShaderIndices;
+	NewShaderMapEntries.SetNum(NumShaderMaps);
+	NewShaderMapHashes.SetNum(NumShaderMaps);
+	NewShaderIndices.Reserve(ShaderIndices.Num()); // Reserve rather than SetNum so we can use Append.
+	for (uint32 NewIndex = 0; NewIndex < NumShaderMaps; ++NewIndex)
+	{
+		const uint32 OldIndex = ShaderMapSortOrder[NewIndex];
+		FShaderMapEntry& Entry = NewShaderMapEntries[NewIndex];
+		Entry = MoveTemp(ShaderMapEntries[OldIndex]);
+		NewShaderMapHashes[NewIndex] = MoveTemp(ShaderMapHashes[OldIndex]);
+
+		if (Entry.NumShaders > 0)
+		{
+			uint32 OldShaderIndicesOffset = Entry.ShaderIndicesOffset;
+			uint32 NewShaderIndicesOffset = NewShaderIndices.Num();
+			TArrayView<uint32> Indices = ShaderIndicesView.Mid(OldShaderIndicesOffset, Entry.NumShaders);
+			NewShaderIndices.Append(Indices);
+			Entry.ShaderIndicesOffset = NewShaderIndicesOffset;
+		}
+	}
+	ShaderMapEntries = MoveTemp(NewShaderMapEntries);
+	ShaderMapHashes = MoveTemp(NewShaderMapHashes);
+	ShaderIndices = MoveTemp(NewShaderIndices);
+
+	// Remap ShaderMapHashTable
+	ShaderMapHashTable.Clear(); // Resets values but does not deallocate
+	for (uint32 NewIndex = 0; NewIndex < NumShaderMaps; ++NewIndex)
+	{
+		const uint32 Key = GetTypeHash(ShaderMapHashes[NewIndex]);
+		ShaderMapHashTable.Add(Key, NewIndex);
+	}
+}
+#endif // WITH_EDITORONLY_DATA
+
 void FSerializedShaderArchive::Finalize()
 {
 	// Set the correct offsets
@@ -718,7 +846,7 @@ void FSerializedShaderArchive::Finalize()
 		check(CurrentPreloadEntry.Size == 0);
 	}
 }
-#endif
+#endif // !USE_MMAPPED_SHADERARCHIVE
 
 static void MakeDebugDirectoryTreeGuarded(const TCHAR* InPath)
 {
diff --git a/Engine/Source/Runtime/RenderCore/Private/ShaderLibrary/ShaderCodeLibrary.cpp b/Engine/Source/Runtime/RenderCore/Private/ShaderLibrary/ShaderCodeLibrary.cpp
--- a/Engine/Source/Runtime/RenderCore/Private/ShaderLibrary/ShaderCodeLibrary.cpp
+++ b/Engine/Source/Runtime/RenderCore/Private/ShaderLibrary/ShaderCodeLibrary.cpp
@@ -79,6 +79,7 @@ using namespace UE::ShaderLibrary::Private;
 namespace UE::ShaderLibrary::Private
 {
 	using ESaveToDiskTarget = ShaderCodeArchive::ESaveToDiskTarget;
+	using ESaveToDiskSortOrder = ShaderCodeArchive::ESaveToDiskSortOrder;
 
 	static FDelegateHandle OnPakFileMountedDelegateHandle;
 	static FDelegateHandle OnPakFileUnmountedDelegateHandle;
@@ -1728,7 +1729,8 @@ public:
 		}
 	}
 
-	bool SaveShaderCode(const FString& ShaderCodeDir, const FString& MetaOutputDir, const TArray<FName>& ShaderFormats, TArray<FString>& OutSCLCSVPath)
+	bool SaveShaderCode(const FString& ShaderCodeDir, const FString& MetaOutputDir, const TArray<FName>& ShaderFormats,
+		TArray<FString>& OutSCLCSVPath, ESaveToDiskSortOrder SortOrder)
 	{
 		using namespace UE::ShaderLibrary::Private;
 		bool bAllOk = ShaderFormats.Num() > 0;
@@ -1753,7 +1755,7 @@ public:
 						// always save shaders in our format even if the platform will use native one. This is needed for iterative cooks (Launch On et al)
 						// to reload previously cooked shaders
 						bool bOk = ReferencedAssetsArchive->SaveToDisk(ShaderCodeDir, MetaOutputDir,
-							ESaveToDiskTarget::Staging);
+							ESaveToDiskTarget::Staging, SortOrder);
 
 						bool bShouldWriteInNativeFormat = bOk && bNativeFormat && ReferencedAssetsArchive->SupportsShaderArchives();
 						if (bShouldWriteInNativeFormat)
@@ -1772,7 +1774,8 @@ public:
 				TUniquePtr<FEditorShaderCodeArchive> UnreferencedAssetsArchive(CodeArchive->CreateArchiveFromAssetsOnlyReferencedByOplog());
 				if (!UnreferencedAssetsArchive->IsEmpty())
 				{
-					bAllOk = UnreferencedAssetsArchive->SaveToDisk(ShaderCodeDir, MetaOutputDir, ESaveToDiskTarget::CookCache) && bAllOk;
+					bAllOk = UnreferencedAssetsArchive->SaveToDisk(ShaderCodeDir, MetaOutputDir,
+						ESaveToDiskTarget::CookCache, ESaveToDiskSortOrder::ShaderHash) && bAllOk;
 				}
 			}
 			FEditorShaderStableInfo* StableArchive = EditorShaderStableInfo[SPlatform];
@@ -1794,7 +1797,8 @@ public:
 	}
 
 	bool SaveShaderCodeChunk(int32 ChunkId, const TSet<FName>& InPackagesInChunk, const TArray<FName>& ShaderFormats,
-		const FString& SandboxDestinationPath, const FString& SandboxMetadataPath, TArray<FString>& OutChunkFilenames)
+		const FString& SandboxDestinationPath, const FString& SandboxMetadataPath,
+		ShaderCodeArchive::ESaveToDiskSortOrder SortOrder, TArray<FString>& OutChunkFilenames)
 	{
 		using namespace UE::ShaderLibrary::Private;
 
@@ -1833,7 +1837,7 @@ public:
 				// always save shaders in our format even if the platform will use native one. This is needed for iterative cooks (Launch On et al)
 				// to reload previously cooked shaders
 				bOk = PerChunkArchive->SaveToDisk(SandboxDestinationPath, SandboxMetadataPath,
-					ESaveToDiskTarget::Staging, &OutChunkFilenames) && bOk;
+					ESaveToDiskTarget::Staging, SortOrder, &OutChunkFilenames) && bOk;
 
 				bool bShouldWriteInNativeFormat = bOk && bNativeFormat && PerChunkArchive->SupportsShaderArchives();
 				if (bShouldWriteInNativeFormat)
@@ -2764,7 +2768,7 @@ bool FShaderLibraryCooker::MergeShaderCodeArchive(const TArray<FString>& CookedM
 	{
 		TArray<FString> CreatedFiles;
 		if (ShaderCodeArchivePair.Value.Key.SaveToDisk(OutShaderCodeDir, TEXT(""),
-			ESaveToDiskTarget::Staging, &CreatedFiles))
+			ESaveToDiskTarget::Staging, ESaveToDiskSortOrder::PackageLoad, &CreatedFiles))
 		{
 			OutWrittenFiles.Append(MoveTemp(CreatedFiles));
 		}
@@ -2775,7 +2779,7 @@ bool FShaderLibraryCooker::MergeShaderCodeArchive(const TArray<FString>& CookedM
 		}
 		CreatedFiles.Reset();
 		if (ShaderCodeArchivePair.Value.Key.SaveToDisk(OutContentDir, TEXT(""),
-			ESaveToDiskTarget::Staging, &CreatedFiles))
+			ESaveToDiskTarget::Staging, ESaveToDiskSortOrder::PackageLoad, &CreatedFiles))
 		{
 			OutWrittenFiles.Append(MoveTemp(CreatedFiles));
 		}
@@ -2803,10 +2807,22 @@ bool FShaderLibraryCooker::MergeShaderCodeArchive(const TArray<FString>& CookedM
 	return bSuccess;
 }
 
+static ShaderCodeArchive::ESaveToDiskSortOrder GetSortOrderForPackagingSystem(UE::Cook::EPackagingSystem PackagingSystem)
+{
+	using namespace ShaderCodeArchive;
+
+	// We prefer ShaderHash sorting for determinism, but with loose (or .pak file packaging of loose shaderfiles),
+	// the ordering from PackageLoad during the cook is important for performance.
+	return PackagingSystem != UE::Cook::EPackagingSystem::Loose ?
+		ESaveToDiskSortOrder::ShaderHash : ESaveToDiskSortOrder::PackageLoad;
+}
+
 bool FShaderLibraryCooker::SaveShaderLibraryWithoutChunking(const ITargetPlatform* TargetPlatform, FString const& Name,
-	FString const& SandboxDestinationPath, FString const& SandboxMetadataPath, TArray<FString>& PlatformSCLCSVPaths, FString& OutErrorMessage,
-	bool& bOutHasData)
+	FString const& SandboxDestinationPath, FString const& SandboxMetadataPath, TArray<FString>& PlatformSCLCSVPaths,
+	UE::Cook::EPackagingSystem PackagingSystem, FString& OutErrorMessage, bool& bOutHasData)
 {
+	checkf(TargetPlatform, TEXT("A valid TargetPlatform is expected"));
+
 	// note that shader formats can be shared across the target platforms
 	TArray<FName> ShaderFormats;
 	TargetPlatform->GetAllTargetedShaderFormats(ShaderFormats);
@@ -2817,13 +2833,15 @@ bool FShaderLibraryCooker::SaveShaderLibraryWithoutChunking(const ITargetPlatfor
 	}
 	bOutHasData = true;
 
+	ShaderCodeArchive::ESaveToDiskSortOrder SortOrder = GetSortOrderForPackagingSystem(PackagingSystem);
+
 	const FString& ShaderCodeDir = SandboxDestinationPath;
 	const FString& MetaDataPath = SandboxMetadataPath;
 
 	checkf(FShaderLibrariesCollection::Impl != nullptr, TEXT("FShaderLibraryCooker was not initialized properly"));
-	checkf(TargetPlatform, TEXT("A valid TargetPlatform is expected"));
 
-	const bool bSaved = FShaderLibrariesCollection::Impl->SaveShaderCode(ShaderCodeDir, MetaDataPath, ShaderFormats, PlatformSCLCSVPaths);
+	const bool bSaved = FShaderLibrariesCollection::Impl->SaveShaderCode(ShaderCodeDir, MetaDataPath, ShaderFormats,
+		PlatformSCLCSVPaths, SortOrder);
 	if (UNLIKELY(!bSaved))
 	{
 		OutErrorMessage = FString::Printf(TEXT("Saving shared material shader code library failed for %s."),
@@ -2834,9 +2852,12 @@ bool FShaderLibraryCooker::SaveShaderLibraryWithoutChunking(const ITargetPlatfor
 	return true;
 }
 
-bool FShaderLibraryCooker::SaveShaderLibraryChunk(int32 ChunkId, const TSet<FName>& InPackagesInChunk, const ITargetPlatform* TargetPlatform,
-	const FString& SandboxDestinationPath, const FString& SandboxMetadataPath, TArray<FString>& OutChunkFilenames, bool& bOutHasData)
+bool FShaderLibraryCooker::SaveShaderLibraryChunk(int32 ChunkId, const TSet<FName>& InPackagesInChunk,
+	const ITargetPlatform* TargetPlatform, const FString& SandboxDestinationPath, const FString& SandboxMetadataPath,
+	UE::Cook::EPackagingSystem PackagingSystem, TArray<FString>& OutChunkFilenames, bool& bOutHasData)
 {
+	checkf(TargetPlatform, TEXT("A valid TargetPlatform is expected"));
+
 	TArray<FName> ShaderFormats;
 	TargetPlatform->GetAllTargetedShaderFormats(ShaderFormats);
 	if (ShaderFormats.IsEmpty())
@@ -2846,10 +2867,11 @@ bool FShaderLibraryCooker::SaveShaderLibraryChunk(int32 ChunkId, const TSet<FNam
 	}
 	bOutHasData = true;
 
+	ShaderCodeArchive::ESaveToDiskSortOrder SortOrder = GetSortOrderForPackagingSystem(PackagingSystem);
+
 	checkf(FShaderLibrariesCollection::Impl != nullptr, TEXT("FShaderLibraryCooker was not initialized properly"));
-	checkf(TargetPlatform, TEXT("A valid TargetPlatform is expected"));
 	return FShaderLibrariesCollection::Impl->SaveShaderCodeChunk(ChunkId, InPackagesInChunk, ShaderFormats,
-		SandboxDestinationPath, SandboxMetadataPath, OutChunkFilenames);
+		SandboxDestinationPath, SandboxMetadataPath, SortOrder, OutChunkFilenames);
 }
 
 void FShaderLibraryCooker::DeleteShaderLibraryChunks(const ITargetPlatform* TargetPlatform,
diff --git a/Engine/Source/Runtime/RenderCore/Public/ShaderCodeArchive.h b/Engine/Source/Runtime/RenderCore/Public/ShaderCodeArchive.h
--- a/Engine/Source/Runtime/RenderCore/Public/ShaderCodeArchive.h
+++ b/Engine/Source/Runtime/RenderCore/Public/ShaderCodeArchive.h
@@ -417,6 +417,10 @@ public:
 	RENDERCORE_API bool FindOrAddShaderMapNonEditor(const FSHAHash& Hash, int32& OutIndex);
 #endif
 	RENDERCORE_API void RemoveLastAddedShader();
+#if WITH_EDITORONLY_DATA
+	/** Sorts Shaders and ShaderMaps by their Hash values. Used to make cooked data deterministic. */
+	RENDERCORE_API void SortByHashes(TArray<FSharedBuffer>& InOutShaderCodeForShaderIndex);
+#endif
 	RENDERCORE_API void Finalize();
 #endif // !USE_MMAPPED_SHADERARCHIVE
 
diff --git a/Engine/Source/Runtime/RenderCore/Public/ShaderCodeLibrary.h b/Engine/Source/Runtime/RenderCore/Public/ShaderCodeLibrary.h
--- a/Engine/Source/Runtime/RenderCore/Public/ShaderCodeLibrary.h
+++ b/Engine/Source/Runtime/RenderCore/Public/ShaderCodeLibrary.h
@@ -291,6 +291,17 @@ struct FShaderCodeLibrary
 #if WITH_EDITOR
 class ITargetPlatform;
 
+namespace UE::Cook
+{
+
+enum class EPackagingSystem
+{
+	IoStore,
+	Loose,
+};
+
+} // namespace UE::Cook
+
 struct FShaderLibraryCooker
 {
 	// Initialize the library cooker
@@ -380,13 +391,17 @@ struct FShaderLibraryCooker
 	 * @param Name shader library name
 	 * @param SandboxDestinationPath where to put the .ushaderbytecode file(s)
 	 * @param SandboxMetadataPath path for the metadata (not a part of the build itself, but produced together with the build)
-	 * @param PlatformSCLCSVPaths path where to put the information about the shader hashes
+	 * @param PlatformSCLCSVPaths path for where to put the information about the shader hashes
+	 * @param PackagingSystem Which packaging system is being used by the project's data build. This impacts
+	 *        some internal format storage decisions for ShaderLibraries due to runtime performance requirements.
 	 * @param OutErrorMessage used to return the details of the failure (if failed)
 	 * @param bOutHasData Reports whether any files were written to PlatformSCLCSVPaths
 	 * @return true if successful
 	 */
-	static RENDERCORE_API bool SaveShaderLibraryWithoutChunking(const ITargetPlatform* TargetPlatform, FString const& Name, FString const& SandboxDestinationPath,
-		FString const& SandboxMetadataPath, TArray<FString>& PlatformSCLCSVPaths, FString& OutErrorMessage, bool& bOutHasData);
+	static RENDERCORE_API bool SaveShaderLibraryWithoutChunking(const ITargetPlatform* TargetPlatform,
+		FString const& Name, FString const& SandboxDestinationPath, FString const& SandboxMetadataPath,
+		TArray<FString>& PlatformSCLCSVPaths, UE::Cook::EPackagingSystem PackagingSystem,
+		FString& OutErrorMessage, bool& bOutHasData);
 
 	/** 
 	 * Saves a single chunk of the collected shader code (per shader platform). Does not save SCL.CSV info.
@@ -397,12 +412,17 @@ struct FShaderLibraryCooker
 	 * @param InPackagesInChunk packages that belong to the chunk
 	 * @param TargetPlatform target platform
 	 * @param SandboxDestinationPath where to put the .ushaderbytecode file(s)
+	 * @param SandboxMetadataPath path for the metadata (not a part of the build itself, but produced together with the build)
+	 * @param PackagingSystem Which packaging system is being used by the project's data build. This impacts
+	 *        some internal format storage decisions for ShaderLibraries due to runtime performance requirements.
 	 * @param OutChunkFilenames array where the function will append the full paths of the written files
 	 * @param bOutHasData Reports whether any files were written to OutChunkFilenames
 	 * @return true if successful
 	 */
-	static RENDERCORE_API bool SaveShaderLibraryChunk(int32 ChunkId, const TSet<FName>& InPackagesInChunk, const ITargetPlatform* TargetPlatform,
-		const FString& SandboxDestinationPath, const FString& SandboxMetadataPath, TArray<FString>& OutChunkFilenames, bool& bOutHasData);
+	static RENDERCORE_API bool SaveShaderLibraryChunk(int32 ChunkId, const TSet<FName>& InPackagesInChunk,
+		const ITargetPlatform* TargetPlatform, const FString& SandboxDestinationPath,
+		const FString& SandboxMetadataPath, UE::Cook::EPackagingSystem PackagingSystem,
+		TArray<FString>& OutChunkFilenames, bool& bOutHasData);
 
 	/**
 	 * Delete the chunk-specific shaderlibrary files that are still present on disk from a previous incremental cook.
