Fix UnifiedBuffer Incorrect Dispatch Sizes For Copies

Fix the dispatch size calculations for byte buffers (e.g. single uint32 elements). The thread count needs to be rounded up to properly handle the tail of non-divisible by four element counts. This made itself obvious during startup of Lyra due to some single element buffer resizes that would turn into 0 sized dispatches and so d3ddebug would issue a warning.

Additionally, ScatterCopyResource() would issue a 0 sized dispatch if Params.GetNumScatters returned 0, again resulting in a warning message. Now it skips the dispatch entirely if NumScatters is 0.

#Synced-CL 51782356

diff --git a/Engine/Source/Runtime/RenderCore/Private/UnifiedBuffer.cpp b/Engine/Source/Runtime/RenderCore/Private/UnifiedBuffer.cpp
--- a/Engine/Source/Runtime/RenderCore/Private/UnifiedBuffer.cpp
+++ b/Engine/Source/Runtime/RenderCore/Private/UnifiedBuffer.cpp
@@ -384,7 +384,7 @@ void MemsetResource(FRDGBuilder& GraphBuilder, FRDGUnorderedAccessView* UAV, con
 		RDG_EVENT_NAME("MemsetResource (%s)", Resource->Name),
 		ComputeShader,
 		PassParameters,
-		FComputeShaderUtils::GetGroupCountWrapped(Params.Count / Divisor, 64));
+		FComputeShaderUtils::GetGroupCountWrapped(FMath::DivideAndRoundUp(Params.Count, Divisor), 64));
 }
 
 void MemsetResource(FRDGBuilder& GraphBuilder, FRDGBufferUAV* UAV, const FMemsetResourceParams& Params)
@@ -470,7 +470,7 @@ void MemcpyResource(FRDGBuilder& GraphBuilder, FRDGUnorderedAccessView* UAV, FRD
 			RDG_EVENT_NAME("Memcpy %s -> %s", DstResource->Name, SrcResource->Name),
 			ComputeShader,
 			PassParameters,
-			FComputeShaderUtils::GetGroupCountWrapped(Params.Count / Divisor, 64));
+			FComputeShaderUtils::GetGroupCountWrapped(FMath::DivideAndRoundUp(Params.Count, Divisor), 64));
 	}
 }
 
@@ -804,13 +804,16 @@ void ScatterCopyResource(FRDGBuilder& GraphBuilder, FRDGBuffer* DstResource, FRD
 		[PassParameters, ComputeShader, Config, GetNumScatters = Params.GetNumScatters](FRDGAsyncTask, FRHIComputeCommandList& RHICmdList)
 		{
 			PassParameters->NumScatters = GetNumScatters();
-			const FScatterUploadDispatchConfig DispatchConfig = GetScatterUploadDispatchConfig(Config, PassParameters->NumScatters);
+			if (PassParameters->NumScatters)
+			{
+				const FScatterUploadDispatchConfig DispatchConfig = GetScatterUploadDispatchConfig(Config, PassParameters->NumScatters);
 
-			FComputeShaderUtils::Dispatch(
-				RHICmdList,
-				ComputeShader,
-				*PassParameters,
-				DispatchConfig.GroupCountWrapped);
+				FComputeShaderUtils::Dispatch(
+					RHICmdList,
+					ComputeShader,
+					*PassParameters,
+					DispatchConfig.GroupCountWrapped);
+			}
 		}
 	);
 }
