Fix TMap race bug in RenderUtils.cpp

#rnx

#Synced-CL 51845578

diff --git a/Engine/Source/Runtime/RenderCore/Private/RenderUtils.cpp b/Engine/Source/Runtime/RenderCore/Private/RenderUtils.cpp
--- a/Engine/Source/Runtime/RenderCore/Private/RenderUtils.cpp
+++ b/Engine/Source/Runtime/RenderCore/Private/RenderUtils.cpp
@@ -2295,6 +2295,9 @@ namespace Substrate
 		// Typically there is only one platform. In editor many can be evaluated for multiple preview in a session. So a Map should be enough.
 		static TMap<EShaderPlatform, FSubstratePerMaterialLevelQualityClosureCount> CachedPlatformResults;
 
+		static FCriticalSection CS;
+		FScopeLock Lock(&CS);
+
 		FSubstratePerMaterialLevelQualityClosureCount* FoundResult = CachedPlatformResults.Find(InPlatform);
 		if (!FoundResult)
 		{
