[Shaders] Reduce lock contention in shader library, specifically `OpenLibrary()` and `PakFileStateChangeInternal()` by iterating over copies of sets and arrays instead of holding the locks while performing I/O operators.

While no deadlock could be found, the heavy I/O operations while holding locks was identified as a potential lock contention, especially when the engine is loading many plugins in parallel.
`KnownPakFilesAccessLock` was held while executing `OnPakFileMounted()` and `NamedLibrariesMutex` was held while executing `OnPakFileMounted()`.
This change allocates copies of the respective arrays in order to release the locks as soon as possible.
This is safe, because it follows the same pattern already used in `OpenLibrary()` and `CloseLibrary()` is not expected to race with pak mount callbacks for the same library.

The crashes reported in the bug ticket all refer to runs in the cloud, where I/O operations have usually higher latency, which exacerbates the lock contention while performing I/O.

#test Tested locally in a client and checked UnrealInsight to verify timings for OpenShaderLibrary() have not changed much; The improvements are expected to benefit mostly machines with higher I/O latency.
#rnx
#Synced-CL 51943062

diff --git a/Engine/Source/Runtime/RenderCore/Private/ShaderLibrary/ShaderCodeLibrary.cpp b/Engine/Source/Runtime/RenderCore/Private/ShaderLibrary/ShaderCodeLibrary.cpp
--- a/Engine/Source/Runtime/RenderCore/Private/ShaderLibrary/ShaderCodeLibrary.cpp
+++ b/Engine/Source/Runtime/RenderCore/Private/ShaderLibrary/ShaderCodeLibrary.cpp
@@ -1028,7 +1028,7 @@ public:
 			FNamedShaderLibrary* Library = nullptr;
 			{
 				// scope of this lock should be as limited as possible - particularly, OpenShaderCode will start async work which is not a good idea to do while holding a lock
-				FRWScopeLock WriteLock(NamedLibrariesMutex, SLT_Write);
+				FRWScopeLock ReadLock(NamedLibrariesMutex, SLT_ReadOnly);
 
 				// create a named library if one didn't exist
 				TUniquePtr<FNamedShaderLibrary>* LibraryPtr = NamedLibrariesStack.Find(Name);
@@ -1060,13 +1060,21 @@ public:
 			else if (!bMonolithicOnly) // attempt to open a chunked library
 			{
 				TSet<int32> PrevComponentSet = Library->GetPresentChunks();
+
+				// Copy the known pak files under the lock, then release it before iterating.
+				// OnPakFileMounted calls OpenShaderCode which performs disk I/O, so holding
+				// KnownPakFilesAccessLock for the entire iteration blocks pak mount callbacks
+				// (FShaderLibraryPakFileMountedCallback) that need the same lock.
+				TSet<FMountedPakFileInfo> KnownPakFilesCopy;
 				{
 					FScopeLock KnownPakFilesLocker(&FMountedPakFileInfo::KnownPakFilesAccessLock);
-					for (TSet<FMountedPakFileInfo>::TConstIterator Iter(FMountedPakFileInfo::KnownPakFiles); Iter; ++Iter)
-					{
-						Library->OnPakFileMounted(*Iter, Directory);
-					}
+					KnownPakFilesCopy = FMountedPakFileInfo::KnownPakFiles;
 				}
+				for (TSet<FMountedPakFileInfo>::TConstIterator Iter(KnownPakFilesCopy); Iter; ++Iter)
+				{
+					Library->OnPakFileMounted(*Iter, Directory);
+				}
+
 				NewComponentIDs = Library->GetPresentChunks().Difference(PrevComponentSet);
 				bResult = !NewComponentIDs.IsEmpty();
 
@@ -1250,24 +1258,35 @@ public:
 		if (IsLibraryInitializedForRuntime())
 		{
 			LLM_SCOPE(ELLMTag::Shaders);
-			TArray<TUniqueFunction<void()>> ShaderLibraryStateChanges;
+
+			// Collect library pointers under the lock, then release it before invoking the callback.
+			// The callback may call OpenShaderCode which performs disk I/O, so holding NamedLibrariesMutex
+			// for the entire iteration can starve writers (e.g. OpenLibrary on the GameThread) for the
+			// duration of all that I/O, causing hang detection to fire.
+			TArray<UE::ShaderLibrary::Private::FNamedShaderLibrary*> Libraries;
 			{
 				FRWScopeLock ReadLock(NamedLibrariesMutex, SLT_ReadOnly);
+				Libraries.Reserve(NamedLibrariesStack.Num());
 				for (TTuple<FString, TUniquePtr<UE::ShaderLibrary::Private::FNamedShaderLibrary>>& NamedLibraryPair : NamedLibrariesStack)
 				{
-					const TSet<int32> PrevComponentSet = NamedLibraryPair.Value->GetPresentChunks();
-					InStateChangeCallback(NamedLibraryPair.Value.Get());
-					const TSet<int32> NewComponentSet = NamedLibraryPair.Value->GetPresentChunks();
-					const TSet<int32> ChangedComponentSet = NewComponentSet.Num() > PrevComponentSet.Num()
-						? NewComponentSet.Difference(PrevComponentSet)
-						: PrevComponentSet.Difference(NewComponentSet);
-					for (int32 ComponentID : ChangedComponentSet)
-					{
-						// Defer these for outside of the lock, ShaderPipelineCache may want to inspect shader library.
-						ShaderLibraryStateChanges.Add([StateChange = InStateChange, ShaderPlatform = NamedLibraryPair.Value->ShaderPlatform, LogicalName = NamedLibraryPair.Value->LogicalName, ComponentID]() {
-							FShaderPipelineCache::ShaderLibraryStateChanged(StateChange, ShaderPlatform, LogicalName, ComponentID);
-						});
-					}
+					Libraries.Add(NamedLibraryPair.Value.Get());
+				}
+			}
+
+			TArray<TUniqueFunction<void()>> ShaderLibraryStateChanges;
+			for (UE::ShaderLibrary::Private::FNamedShaderLibrary* NamedShaderLibrary : Libraries)
+			{
+				const TSet<int32> PrevComponentSet = NamedShaderLibrary->GetPresentChunks();
+				InStateChangeCallback(NamedShaderLibrary);
+				const TSet<int32> NewComponentSet = NamedShaderLibrary->GetPresentChunks();
+				const TSet<int32> ChangedComponentSet = NewComponentSet.Num() > PrevComponentSet.Num()
+					? NewComponentSet.Difference(PrevComponentSet)
+					: PrevComponentSet.Difference(NewComponentSet);
+				for (int32 ComponentID : ChangedComponentSet)
+				{
+					ShaderLibraryStateChanges.Add([StateChange = InStateChange, ShaderPlatform = NamedShaderLibrary->ShaderPlatform, LogicalName = NamedShaderLibrary->LogicalName, ComponentID]() {
+						FShaderPipelineCache::ShaderLibraryStateChanged(StateChange, ShaderPlatform, LogicalName, ComponentID);
+					});
 				}
 			}
 			for (TUniqueFunction<void()>& OnShaderLibraryStateChange : ShaderLibraryStateChanges)
@@ -2961,7 +2980,7 @@ bool UE::ShaderLibrary::Private::FNamedShaderLibrary::OpenShaderCode(const FStri
 	LLM_SCOPE(ELLMTag::Shaders);
 	// check if any of the components has this content
 	{
-		FRWScopeLock WriteLock(ComponentsMutex, SLT_Write);
+		FRWScopeLock ReadLock(ComponentsMutex, SLT_ReadOnly);
 		for (const FShaderLibraryInstancePriorityPair& Component : Components)
 		{
 			if (Component.Value->HasContentFrom(ShaderCodeDir, Library))
diff --git a/Engine/Source/Runtime/RenderCore/Private/ShaderLibrary/ShaderLibraryInstance.cpp b/Engine/Source/Runtime/RenderCore/Private/ShaderLibrary/ShaderLibraryInstance.cpp
--- a/Engine/Source/Runtime/RenderCore/Private/ShaderLibrary/ShaderLibraryInstance.cpp
+++ b/Engine/Source/Runtime/RenderCore/Private/ShaderLibrary/ShaderLibraryInstance.cpp
@@ -430,12 +430,12 @@ void FShaderLibraryInstance::ReleasePreloadedPackageShaderMap(int32 ShaderMapInd
 	}
 }
 
-bool FShaderLibraryInstance::HasContentFrom(const FString& ShaderCodeDir, FString const& InLibraryName)
+bool FShaderLibraryInstance::HasContentFrom(const FString& ShaderCodeDir, FString const& InLibraryName) const
 {
 	if (Library->GetName() == InLibraryName)
 	{
 		// IoStore-based libraries don't care about the directory, so name collision alone is enough to say yes.
-		return ShaderCodeDirectory.IsEmpty() ? true : (ShaderCodeDirectory == ShaderCodeDir);
+		return ShaderCodeDirectory.IsEmpty() || (ShaderCodeDirectory == ShaderCodeDir);
 	}
 
 	return false;
diff --git a/Engine/Source/Runtime/RenderCore/Private/ShaderLibrary/ShaderLibraryInstance.h b/Engine/Source/Runtime/RenderCore/Private/ShaderLibrary/ShaderLibraryInstance.h
--- a/Engine/Source/Runtime/RenderCore/Private/ShaderLibrary/ShaderLibraryInstance.h
+++ b/Engine/Source/Runtime/RenderCore/Private/ShaderLibrary/ShaderLibraryInstance.h
@@ -60,7 +60,7 @@ public:
 
 	void ReleasePreloadedPackageShaderMap(int32 ShaderMapIndex);
 
-	bool HasContentFrom(const FString& ShaderCodeDir, FString const& InLibraryName);
+	bool HasContentFrom(const FString& ShaderCodeDir, FString const& InLibraryName) const;
 
 	/** Returns chunk id for this instance, if it is a part of a chunked named library. INDEX_NONE is returned for a monolithic lib. */
 	int32 GetChunkId() const { return ChunkId; }
