Add prerequisite task support for AddDispatchPass. This allows sub-tasks to sync with RDG tasks so RDG resources can be referenced within them.

#Synced-CL 51965579

diff --git a/Engine/Source/Runtime/RenderCore/Private/RenderGraphBuilder.cpp b/Engine/Source/Runtime/RenderCore/Private/RenderGraphBuilder.cpp
--- a/Engine/Source/Runtime/RenderCore/Private/RenderGraphBuilder.cpp
+++ b/Engine/Source/Runtime/RenderCore/Private/RenderGraphBuilder.cpp
@@ -3086,6 +3086,22 @@ void FRDGBuilder::SetupDispatchPassExecute()
 			DispatchPass->LaunchDispatchPassTasks(DispatchPassBuilder);
 			DispatchPassBuilder.Finish();
 
+			for (const UE::Tasks::FTask& Task : DispatchPassBuilder.Prerequisites)
+			{
+				if (ParallelExecute.TasksAsync)
+				{
+					ParallelExecute.TasksAsync->AddPrerequisites(Task);
+				}
+				else if (ParallelExecute.TasksAwait)
+				{
+					ParallelExecute.TasksAwait->AddPrerequisites(Task);
+				}
+				else
+				{
+					Task.Wait();
+				}
+			}
+
 			IF_RDG_ENABLE_DEBUG(UserValidation.ValidateExecutePassEnd(DispatchPass));
 		}
 
diff --git a/Engine/Source/Runtime/RenderCore/Private/RenderGraphPass.cpp b/Engine/Source/Runtime/RenderCore/Private/RenderGraphPass.cpp
--- a/Engine/Source/Runtime/RenderCore/Private/RenderGraphPass.cpp
+++ b/Engine/Source/Runtime/RenderCore/Private/RenderGraphPass.cpp
@@ -155,6 +155,14 @@ FRHICommandList* FRDGDispatchPassBuilder::CreateCommandList()
 	return RHICmdList;
 }
 
+void FRDGDispatchPassBuilder::AddPrerequisite(const UE::Tasks::FTask& Task)
+{
+	if (Task.IsValid())
+	{
+		Prerequisites.Add(Task);
+	}
+}
+
 void FRDGDispatchPassBuilder::Finish()
 {
 	// With serial execution the pass commands are embedded in the immediate command list instead.
diff --git a/Engine/Source/Runtime/RenderCore/Public/RenderGraphPass.h b/Engine/Source/Runtime/RenderCore/Public/RenderGraphPass.h
--- a/Engine/Source/Runtime/RenderCore/Public/RenderGraphPass.h
+++ b/Engine/Source/Runtime/RenderCore/Public/RenderGraphPass.h
@@ -741,6 +741,9 @@ public:
 	/** Create a new command list to record into and inserts it. Call FinishRecording() on the task when done. */
 	RENDERCORE_API FRHICommandList* CreateCommandList();
 
+	/** Adds a task that the graph will wait on during execution, alongside other async execution tasks. */
+	RENDERCORE_API void AddPrerequisite(const UE::Tasks::FTask& Task);
+
 private:
 	FRDGDispatchPassBuilder(FRDGDispatchPass* InPass)
 		: Pass(InPass)
@@ -759,7 +762,8 @@ private:
 	TSharedPtr<FRHIParallelRenderPassInfo> RenderPassInfo;
 
 	TArray<FRHISubCommandList*> SubCommandLists;
-	
+	TArray<UE::Tasks::FTask> Prerequisites;
+
 	friend FRDGBuilder;
 };
 
