Fixed TSAN error due to bExternal and bExtracted being in the same bitfield as other members accessed from the async setup queue.

#Synced-CL 51975396

diff --git a/Engine/Source/Runtime/RenderCore/Private/RenderGraphResources.cpp b/Engine/Source/Runtime/RenderCore/Private/RenderGraphResources.cpp
--- a/Engine/Source/Runtime/RenderCore/Private/RenderGraphResources.cpp
+++ b/Engine/Source/Runtime/RenderCore/Private/RenderGraphResources.cpp
@@ -43,8 +43,6 @@ FRDGBuffer::FRDGBuffer(const TCHAR* InName, const FRDGBufferDesc& InDesc, ERDGBu
 FRDGViewableResource::FRDGViewableResource(const TCHAR* InName, const ERDGViewableResourceType InType, bool bSkipTracking, bool bInSplitFirstTransition)
 	: FRDGResource(InName)
 	, Type(InType)
-	, bExternal(0)
-	, bExtracted(0)
 	, bTransient(0)
 	, bForceNonTransient(0)
 	, bSkipLastTransition(0)
diff --git a/Engine/Source/Runtime/RenderCore/Public/RenderGraphResources.h b/Engine/Source/Runtime/RenderCore/Public/RenderGraphResources.h
--- a/Engine/Source/Runtime/RenderCore/Public/RenderGraphResources.h
+++ b/Engine/Source/Runtime/RenderCore/Public/RenderGraphResources.h
@@ -373,12 +373,6 @@ protected:
 		External
 	};
 
-	/** Whether this is an externally registered resource. */
-	uint8 bExternal : 1;
-
-	/** Whether this is an extracted resource. */
-	uint8 bExtracted : 1;
-
 	/** Whether this resource is allocated through the transient resource allocator. */
 	uint8 bTransient : 1;
 
@@ -406,9 +400,17 @@ protected:
 	uint8 bClobbered : 1;
 #endif
 
+	// Use explicit booleans due to access from the render thread while the async setup queue task is running.
+
 	/** Whether any sub-resource has been used for write by a pass. Separate from the bitfield because it is assigned by the async setup queue. */
 	bool bProduced = false;
 
+	/** Whether this is an externally registered resource. */
+	bool bExternal = false;
+
+	/** Whether this is an extracted resource. */
+	bool bExtracted = false;
+
 	/** Whether this resource is allowed to be both transient and extracted. */
 	ETransientExtractionHint TransientExtractionHint;
 
