[shaders] version bump to SHK header; format changed with hashing changes

#Synced-CL 52103695

diff --git a/Engine/Source/Runtime/RenderCore/Private/PipelineCacheUtilities.cpp b/Engine/Source/Runtime/RenderCore/Private/PipelineCacheUtilities.cpp
--- a/Engine/Source/Runtime/RenderCore/Private/PipelineCacheUtilities.cpp
+++ b/Engine/Source/Runtime/RenderCore/Private/PipelineCacheUtilities.cpp
@@ -38,7 +38,7 @@ namespace Private
 
 		enum class EVersion : int32
 		{
-			Current = 1
+			Current = 2
 		};
 
 		/** Magic to reject other files */
