Pre-warm shader source file hash cache to speed up editor startup. Instead of taking a lock and calculating the hash for each file individually on the game thread, this CL performs a parallel task up front to calculate all the hashes.

#rnx
#Synced-CL 52172106

diff --git a/Engine/Source/Runtime/Engine/Private/ShaderCompiler/ShaderCompiler.cpp b/Engine/Source/Runtime/Engine/Private/ShaderCompiler/ShaderCompiler.cpp
--- a/Engine/Source/Runtime/Engine/Private/ShaderCompiler/ShaderCompiler.cpp
+++ b/Engine/Source/Runtime/Engine/Private/ShaderCompiler/ShaderCompiler.cpp
@@ -5585,6 +5585,8 @@ void CompileGlobalShaderMap(EShaderPlatform Platform, const ITargetPlatform* Tar
 #if WITH_EDITOR
 		if (!bLoadedFromCacheFile && bAllowShaderCompiling)
 		{
+			PrimeGlobalShaderSourceHashCache(Platform, TargetPlatform);
+			
 			// Ensure we've generated AutogenShaderHeaders.ush
 			FShaderCompileUtilities::GenerateBrdfHeaders(Platform);
 
@@ -5656,7 +5658,10 @@ void CompileGlobalShaderMap(EShaderPlatform Platform, const ITargetPlatform* Tar
 							GlobalShaderMapLoads[int32(Response.UserData)].ReadFromRecord(Response.Record);
 						}
 					});
-					BlockingOwner.Wait();
+					{
+						TRACE_CPUPROFILER_EVENT_SCOPE(WaitForDDC);
+						BlockingOwner.Wait();
+					}
 				}
 
 				BufferIndex = 0;
diff --git a/Engine/Source/Runtime/RenderCore/Private/GlobalShader.cpp b/Engine/Source/Runtime/RenderCore/Private/GlobalShader.cpp
--- a/Engine/Source/Runtime/RenderCore/Private/GlobalShader.cpp
+++ b/Engine/Source/Runtime/RenderCore/Private/GlobalShader.cpp
@@ -15,6 +15,7 @@
 #include "Serialization/MemoryWriter.h"
 #include "ShaderCodeLibrary.h"
 #include "ShaderSerialization.h"
+#include "Tasks/Task.h"
 
 /** The global shader map. */
 FGlobalShaderMap* GGlobalShaderMap[SP_NumPlatforms] = {};
@@ -585,6 +586,43 @@ FGlobalShaderMap* GetGlobalShaderMap(EShaderPlatform Platform)
 	return GGlobalShaderMap[Platform];
 }
 
+void PrimeGlobalShaderSourceHashCache(EShaderPlatform Platform, const ITargetPlatform* TargetPlatform)
+{
+#if WITH_EDITOR
+	FPlatformTypeLayoutParameters LayoutParams;
+	LayoutParams.InitializeForPlatform(TargetPlatform);
+	const EShaderPermutationFlags PermutationFlags = GetShaderPermutationFlags(LayoutParams);
+
+	TArray<FGlobalShaderType*> ShaderTypesToHash;
+	for (TLinkedList<FShaderType*>::TIterator ShaderTypeIt(FShaderType::GetTypeList()); ShaderTypeIt; ShaderTypeIt.Next())
+	{
+		FGlobalShaderType* GlobalShaderType = ShaderTypeIt->GetGlobalShaderType();
+		if (!GlobalShaderType)
+		{
+			continue;
+		}
+
+		for (int32 PermutationId = 0; PermutationId < GlobalShaderType->GetPermutationCount(); PermutationId++)
+		{
+			if (GlobalShaderType->ShouldCompilePermutation(Platform, PermutationId, PermutationFlags))
+			{
+				ShaderTypesToHash.Add(GlobalShaderType);
+				break;
+			}
+		}
+	}
+
+	TArray<FString> ShaderFilenames;
+	ShaderFilenames.Reserve(ShaderTypesToHash.Num());
+	for (FGlobalShaderType* GlobalShaderType : ShaderTypesToHash)
+	{
+		ShaderFilenames.Add(GlobalShaderType->GetShaderFilename());
+	}
+	
+	PrimeShaderFileHashCache(ShaderFilenames, Platform);
+#endif
+}
+
 FGlobalShaderMapSection* FGlobalShaderMapSection::CreateFromArchive(FArchive& Ar)
 {
 	FGlobalShaderMapSection* Section = new FGlobalShaderMapSection();
diff --git a/Engine/Source/Runtime/RenderCore/Private/ShaderCore.cpp b/Engine/Source/Runtime/RenderCore/Private/ShaderCore.cpp
--- a/Engine/Source/Runtime/RenderCore/Private/ShaderCore.cpp
+++ b/Engine/Source/Runtime/RenderCore/Private/ShaderCore.cpp
@@ -3475,6 +3475,90 @@ const FSHAHash& GetShaderFilesHash(const TArray<FString>& VirtualFilePaths, ESha
 	}
 }
 
+void PrimeShaderFileHashCache(TArrayView<const FString> VirtualFilePaths, EShaderPlatform ShaderPlatform)
+{
+	STAT(double HashTime = 0);
+	{
+		SCOPE_SECONDS_COUNTER(HashTime);
+
+		// Single read-lock pass: collect all paths not yet present in the cache.
+		TArray<const FString*> UncachedPaths;
+		UncachedPaths.Reserve(VirtualFilePaths.Num());
+		{
+			FRWScopeLock ShaderHashAccessLock(GShaderHashAccessRWLock, SLT_ReadOnly);
+			for (const FString& VirtualFilePath : VirtualFilePaths)
+			{
+				if (!GShaderHashCache.FindHash(ShaderPlatform, VirtualFilePath))
+				{
+					UncachedPaths.Add(&VirtualFilePath);
+				}
+			}
+		}
+
+		if (UncachedPaths.IsEmpty())
+		{
+			return;
+		}
+
+		// Hold the calculation guard for the entire batch so no other thread redundantly
+		// computes any of the same hashes while we are working through the list.
+		FScopeLock FileHashCalculationAccessLock(&GShaderFileHashCalculationGuard);
+
+		// Pre-allocate one slot per path so parallel tasks write to disjoint indices with no contention.
+		// FilePath == nullptr serves as the sentinel for "skipped or failed".
+		struct FComputedHash
+		{
+			const FString* FilePath = nullptr;
+			FSHAHash Hash;
+		};
+		TArray<FComputedHash> ComputedHashes;
+		ComputedHashes.SetNum(UncachedPaths.Num());
+
+		// Compute hashes in parallel; each task owns its slot in ComputedHashes.
+		// FindHash reads below are safe without RWLock: GShaderFileHashCalculationGuard
+		// ensures no other thread writes new cache entries during this section.
+		ParallelFor(UncachedPaths.Num(), [&](int32 Index)
+		{
+			const FString* VirtualFilePath = UncachedPaths[Index];
+
+			// Double-check: a thread that held the guard before us may have already populated this entry.
+			if (GShaderHashCache.FindHash(ShaderPlatform, *VirtualFilePath))
+			{
+				return;
+			}
+
+			FSHA1 HashState;
+			FString ErrorMessage;
+			if (TryUpdateSingleShaderFilehash(HashState, **VirtualFilePath, ShaderPlatform, &ErrorMessage))
+			{
+				HashState.Final();
+				ComputedHashes[Index].FilePath = VirtualFilePath;
+				HashState.GetHash(&ComputedHashes[Index].Hash.Hash[0]);
+			}
+			else
+			{
+				UE_LOG(LogShaders, Warning, TEXT("PrimeShaderFileHashCache: failed to hash '%s': %s"), **VirtualFilePath, *ErrorMessage);
+			}
+		});
+
+		// Single write-lock pass to insert all newly computed hashes at once.
+		FRWScopeLock ShaderHashAccessLock(GShaderHashAccessRWLock, SLT_Write);
+		for (const FComputedHash& Entry : ComputedHashes)
+		{
+			if (Entry.FilePath)
+			{
+				FSHAHash& NewHash = GShaderHashCache.AddHash(ShaderPlatform, *Entry.FilePath);
+				NewHash = Entry.Hash;
+#if WITH_EDITOR && !(UE_BUILD_SHIPPING || UE_BUILD_TEST)
+				UE_LOG(LogShaders, Verbose, TEXT("Final hash for file %s, %s"), *(*Entry.FilePath), *BytesToHex(&NewHash.Hash[0], 20));
+#endif
+			}
+		}
+
+		INC_FLOAT_STAT_BY(STAT_ShaderCompiling_HashingShaderFiles, (float)HashTime);
+	}
+}
+
 #if WITH_EDITOR
 void BuildShaderFileToUniformBufferMap(TMap<FString, TArray<const TCHAR*> >& ShaderFileToUniformBufferVariables)
 {
diff --git a/Engine/Source/Runtime/RenderCore/Public/GlobalShader.h b/Engine/Source/Runtime/RenderCore/Public/GlobalShader.h
--- a/Engine/Source/Runtime/RenderCore/Public/GlobalShader.h
+++ b/Engine/Source/Runtime/RenderCore/Public/GlobalShader.h
@@ -332,6 +332,14 @@ public:
  */
 extern RENDERCORE_API FGlobalShaderMap* GetGlobalShaderMap(EShaderPlatform Platform);
 
+RENDERCORE_API void PrimeGlobalShaderSourceHashCache(EShaderPlatform Platform, const ITargetPlatform* TargetPlatform);
+
+inline void PrimeGlobalShaderSourceHashCache()
+{
+	EShaderPlatform Platform = GShaderPlatformForFeatureLevel[GMaxRHIFeatureLevel];
+	PrimeGlobalShaderSourceHashCache(Platform, nullptr);
+}
+
 /**
   * Overload for the above GetGlobalShaderMap which takes a feature level and translates to the appropriate shader platform
   *
diff --git a/Engine/Source/Runtime/RenderCore/Public/ShaderCore.h b/Engine/Source/Runtime/RenderCore/Public/ShaderCore.h
--- a/Engine/Source/Runtime/RenderCore/Public/ShaderCore.h
+++ b/Engine/Source/Runtime/RenderCore/Public/ShaderCore.h
@@ -1516,6 +1516,15 @@ extern RENDERCORE_API const FSHAHash* TryGetShaderFileHash(const TCHAR* VirtualF
  */
 extern RENDERCORE_API const class FSHAHash& GetShaderFilesHash(const TArray<FString>& VirtualFilePaths, EShaderPlatform ShaderPlatform);
 
+/**
+ * Primes the shader file hash cache for a batch of virtual file paths, acquiring each lock only
+ * once for the entire batch rather than once per file. Prefer this over repeated TryGetShaderFileHash
+ * calls when pre-warming the cache for a known set of files.
+ * @param VirtualFilePaths - List of virtual shader file paths to prime.
+ * @param ShaderPlatform   - Shader platform the hashes are for.
+ */
+extern RENDERCORE_API void PrimeShaderFileHashCache(TArrayView<const FString> VirtualFilePaths, EShaderPlatform ShaderPlatform);
+
 /**
  * Flushes the shader file and CRC cache, and regenerates the binary shader files if necessary.
  * Allows shader source files to be re-read properly even if they've been modified since startup.
