[PCG] Add engine infrastructure for PCG GPU inline ray tracing

Adds engine-level support needed by the PCG World Raycast GPU backend:

- FRaytracingVFAttributeData type and GetRaytracingVFAttributeData() virtual on FPrimitiveSceneProxy for querying per-primitive UV data on the render thread. FStaticMeshSceneProxy override returns UV SRV from the ray tracing proxy vertex factory.
- SHADER_PARAMETER_RESOURCE_COLLECTION macro and UBMT_RESOURCE_COLLECTION plumbing through ShaderParameterStruct, ShaderParameterParser, and RenderGraphUtils to support binding FRHIResourceCollection in shader parameter structs.
- BufferFromCollection() helper in BindlessResources.ush for indexing into resource collections from shaders.
- ESceneViewExtensionFlags::RequiresHardwareInlineRayTracing flag and PostTLASBuild callback gate widened to fire for inline-only platforms (not just full RT shader platforms).
- GetRayTracingSceneViewRDG() accessor in FXRenderingUtils for obtaining the RDG ray tracing scene view handle.
- GetInlineRayTracingBindingDataBuffer() override on FViewInfo.

#rnx

#Synced-CL 52169849

diff --git a/Engine/Plugins/Runtime/ComputeFramework/Source/ComputeFramework/Private/ComputeViewExtension.cpp b/Engine/Plugins/Runtime/ComputeFramework/Source/ComputeFramework/Private/ComputeViewExtension.cpp
--- a/Engine/Plugins/Runtime/ComputeFramework/Source/ComputeFramework/Private/ComputeViewExtension.cpp
+++ b/Engine/Plugins/Runtime/ComputeFramework/Source/ComputeFramework/Private/ComputeViewExtension.cpp
@@ -59,7 +59,7 @@ ESceneViewExtensionFlags FComputeViewExtension::GetFlags() const
 		{
 			if (ComputeTaskWorker->HasWork(ComputeTaskExecutionGroup::PostTLASBuild))
 			{
-				return ESceneViewExtensionFlags::SubscribesToPostTLASBuild;
+				return ESceneViewExtensionFlags::SubscribesToPostTLASBuild | ESceneViewExtensionFlags::RequiresHardwareInlineRayTracing;
 			}
 		}
 	}
diff --git a/Engine/Shaders/Public/BindlessResources.ush b/Engine/Shaders/Public/BindlessResources.ush
--- a/Engine/Shaders/Public/BindlessResources.ush
+++ b/Engine/Shaders/Public/BindlessResources.ush
@@ -146,6 +146,15 @@ uint GetResourceIndexFromResourceCollection(FResourceCollection ResourceCollecti
 			uint ResourceIndex = GetResourceIndexFromResourceCollection(ResourceCollection, CollectionIndex); \
 			return RESOURCE_COLLECTION_HEAP_NAME(TType)[ResourceIndex]; \
 		}
+
+	#define BUFFER_FROM_COLLECTION_T(TType) \
+		typedef TType SafeType##TType; \
+		SafeType##TType RESOURCE_COLLECTION_HEAP_NAME(TType)[]; \
+		SafeType##TType BufferFromCollection_##TType(FResourceCollection ResourceCollection, uint CollectionIndex) \
+		{ \
+			uint ResourceIndex = GetResourceIndexFromResourceCollection(ResourceCollection, CollectionIndex); \
+			return RESOURCE_COLLECTION_HEAP_NAME(TType)[ResourceIndex]; \
+		}
 #else
 	#define TEXTURE_FROM_COLLECTION_T(TType) \
 		TType TextureFromCollection_##TType(FResourceCollection ResourceCollection, uint CollectionIndex) \
@@ -153,6 +162,13 @@ uint GetResourceIndexFromResourceCollection(FResourceCollection ResourceCollecti
 			uint ResourceIndex = GetResourceIndexFromResourceCollection(ResourceCollection, CollectionIndex); \
 			return GetSRVFromHeap(TType, ResourceIndex); \
 		}
+
+	#define BUFFER_FROM_COLLECTION_T(TType) \
+		TType BufferFromCollection_##TType(FResourceCollection ResourceCollection, uint CollectionIndex) \
+		{ \
+			uint ResourceIndex = GetResourceIndexFromResourceCollection(ResourceCollection, CollectionIndex); \
+			return GetSRVFromHeap(TType, ResourceIndex); \
+		}
 #endif
 
 TEXTURE_FROM_COLLECTION_T(Texture2D)
@@ -161,6 +177,9 @@ TEXTURE_FROM_COLLECTION_T(Texture2DArray)
 TEXTURE_FROM_COLLECTION_T(TextureCubeArray)
 TEXTURE_FROM_COLLECTION_T(Texture3D)
 
+typedef Buffer<float2> Buffer_float2;
+BUFFER_FROM_COLLECTION_T(Buffer_float2)
+
 #endif
 
 #if defined(INCLUDE_PLATFORM_BINDLESS_RESOURCES_USH)
diff --git a/Engine/Source/Runtime/Engine/Private/StaticMeshSceneProxy.cpp b/Engine/Source/Runtime/Engine/Private/StaticMeshSceneProxy.cpp
--- a/Engine/Source/Runtime/Engine/Private/StaticMeshSceneProxy.cpp
+++ b/Engine/Source/Runtime/Engine/Private/StaticMeshSceneProxy.cpp
@@ -2127,6 +2127,34 @@ RayTracing::FGeometryGroupHandle FStaticMeshSceneProxy::GetRayTracingGeometryGro
 	return RayTracingGeometryGroupHandle;
 }
 
+FRaytracingVFAttributeData FStaticMeshSceneProxy::GetRaytracingVFAttributeData(const FSceneView* View)
+{
+	FRaytracingVFAttributeData Out;
+
+	if (!View)
+	{
+		return Out;
+	}
+
+	if (!bSupportRayTracing || !bHasRayTracingRepresentation || !IsRayTracingRelevant())
+	{
+		return Out;
+	}
+
+	if (!RenderData || !RenderData->RayTracingProxy || !RenderData->RayTracingProxy->LODVertexFactories)
+	{
+		return Out;
+	}
+
+	const int32 LODIndex = RenderData->RayTracingProxy->bUsingRenderingLODs ? FMath::Max(GetLOD(View), (int32)GetCurrentFirstLODIdx_RenderThread()) : 0;
+	const FStaticMeshVertexFactories& VertexFactories = (*RenderData->RayTracingProxy->LODVertexFactories)[LODIndex];
+	FRaytracingVFSectionData& Section = Out.Sections.Emplace_GetRef();
+	Section.UV.BufferSRV = VertexFactories.VertexFactory.GetTextureCoordinatesSRV();
+	Section.UV.NumCoordinates = VertexFactories.VertexFactory.GetNumTexcoords();
+
+	return Out;
+}
+
 void FStaticMeshSceneProxy::GetDynamicRayTracingInstances(FRayTracingInstanceCollector& Collector)
 {
 #if DO_CHECK
diff --git a/Engine/Source/Runtime/Engine/Public/PrimitiveSceneProxy.h b/Engine/Source/Runtime/Engine/Public/PrimitiveSceneProxy.h
--- a/Engine/Source/Runtime/Engine/Public/PrimitiveSceneProxy.h
+++ b/Engine/Source/Runtime/Engine/Public/PrimitiveSceneProxy.h
@@ -21,6 +21,7 @@
 #include "SceneDefinitions.h"
 #include "MeshDrawCommandStatsDefines.h"
 #include "InstanceDataTypes.h"
+#include "RaytracingVFAttributeData.h"
 #include "PSOPrecacheFwd.h"
 #include "RayTracingGeometryManagerTypes.h"
 #include "Runtime/Engine/Internal/MaterialCache/MaterialCacheDescriptor.h"
@@ -478,6 +479,8 @@ public:
 
 	/** @return The handle of the ray tracing geometry group used by this primitive */
 	ENGINE_API ERayTracingProxyType GetRayTracingProxyType() const;
+
+	virtual FRaytracingVFAttributeData GetRaytracingVFAttributeData(const FSceneView* View) { return {}; }
 #endif // RHI_RAYTRACING
 
 	virtual void GetStreamableRenderAssetInfo(const FBoxSphereBounds& PrimitiveBounds, TArray<struct FStreamingRenderAssetPrimitiveInfo>& StreamableRenderAssets) const { check(IsSupportingStreamableRenderAssetsGathering()); }
diff --git a/Engine/Source/Runtime/Engine/Public/RaytracingVFAttributeData.h b/Engine/Source/Runtime/Engine/Public/RaytracingVFAttributeData.h
new file mode 100644
--- /dev/null
+++ b/Engine/Source/Runtime/Engine/Public/RaytracingVFAttributeData.h
@@ -0,0 +1,23 @@
+// Copyright Epic Games, Inc. All Rights Reserved.
+
+#pragma once
+
+#include "Containers/Array.h"
+#include "RendererInterface.h"
+
+struct FRaytracingVFUVData
+{
+	FRHIShaderResourceView* BufferSRV = nullptr;
+
+	uint32 NumCoordinates = 0;
+};
+
+struct FRaytracingVFSectionData
+{
+	FRaytracingVFUVData UV;
+};
+
+struct FRaytracingVFAttributeData
+{
+	TArray<FRaytracingVFSectionData, TInlineAllocator<4>> Sections;
+};
diff --git a/Engine/Source/Runtime/Engine/Public/SceneView.h b/Engine/Source/Runtime/Engine/Public/SceneView.h
--- a/Engine/Source/Runtime/Engine/Public/SceneView.h
+++ b/Engine/Source/Runtime/Engine/Public/SceneView.h
@@ -2171,6 +2171,9 @@ public:
 		return bIsSceneCaptureCube && bIsReflectionCapture;
 	}
 
+	/** Buffer to the shader binding data used for inline raytracing - only valid when inline raytracing is enabled. */
+	virtual FRDGBufferRef GetInlineRayTracingBindingDataBuffer() const { return nullptr; }
+
 protected:
 	friend class FViewSnapshotCache;
 
diff --git a/Engine/Source/Runtime/Engine/Public/SceneViewExtension.h b/Engine/Source/Runtime/Engine/Public/SceneViewExtension.h
--- a/Engine/Source/Runtime/Engine/Public/SceneViewExtension.h
+++ b/Engine/Source/Runtime/Engine/Public/SceneViewExtension.h
@@ -109,7 +109,8 @@ enum class ESceneViewExtensionFlags : uint32
 {
 	None = 0,
 
-	SubscribesToPostTLASBuild = 1 << 0
+	SubscribesToPostTLASBuild = 1 << 0,
+	RequiresHardwareInlineRayTracing = 1 << 1
 };
 ENUM_CLASS_FLAGS(ESceneViewExtensionFlags);
 
diff --git a/Engine/Source/Runtime/Engine/Public/StaticMeshSceneProxy.h b/Engine/Source/Runtime/Engine/Public/StaticMeshSceneProxy.h
--- a/Engine/Source/Runtime/Engine/Public/StaticMeshSceneProxy.h
+++ b/Engine/Source/Runtime/Engine/Public/StaticMeshSceneProxy.h
@@ -182,6 +182,7 @@ public:
 		return bIsStaticInstance && !HasViewDependentDPG() && !(bAllowStaticLighting && HasStaticLighting() && !HasValidSettingsForStaticLighting());
 	}
 	ENGINE_API virtual RayTracing::FGeometryGroupHandle GetRayTracingGeometryGroupHandle() const override;
+	ENGINE_API virtual FRaytracingVFAttributeData GetRaytracingVFAttributeData(const FSceneView* View) override;
 #endif // RHI_RAYTRACING
 
 	ENGINE_API virtual void GetLCIs(FLCIArray& LCIs) override;
diff --git a/Engine/Source/Runtime/RenderCore/Private/RenderGraphUtils.cpp b/Engine/Source/Runtime/RenderCore/Private/RenderGraphUtils.cpp
--- a/Engine/Source/Runtime/RenderCore/Private/RenderGraphUtils.cpp
+++ b/Engine/Source/Runtime/RenderCore/Private/RenderGraphUtils.cpp
@@ -42,7 +42,8 @@ void ClearUnusedGraphResourcesImpl(
 			Type == UBMT_RDG_TEXTURE_NON_PIXEL_SRV ||
 			Type == UBMT_RDG_TEXTURE_UAV ||
 			Type == UBMT_RDG_BUFFER_SRV ||
-			Type == UBMT_RDG_BUFFER_UAV)
+			Type == UBMT_RDG_BUFFER_UAV ||
+			Type == UBMT_RESOURCE_COLLECTION)
 		{
 			const TMemoryImageArray<FShaderParameterBindings::FResourceParameter>& ResourceParameters = ShaderBindings.ResourceParameters;
 			const int32 ShaderResourceCount = ResourceParameters.Num();
@@ -127,7 +128,8 @@ void ClearUnusedGraphResourcesImpl(
 			Type == UBMT_RDG_TEXTURE_NON_PIXEL_SRV ||
 			Type == UBMT_RDG_TEXTURE_UAV ||
 			Type == UBMT_RDG_BUFFER_SRV ||
-			Type == UBMT_RDG_BUFFER_UAV)
+			Type == UBMT_RDG_BUFFER_UAV ||
+			Type == UBMT_RESOURCE_COLLECTION)
 		{
 			for (int32 Index = 0; Index < ShaderBindingsList.Num(); ++Index)
 			{
diff --git a/Engine/Source/Runtime/RenderCore/Private/ShaderParameterParser.cpp b/Engine/Source/Runtime/RenderCore/Private/ShaderParameterParser.cpp
--- a/Engine/Source/Runtime/RenderCore/Private/ShaderParameterParser.cpp
+++ b/Engine/Source/Runtime/RenderCore/Private/ShaderParameterParser.cpp
@@ -163,6 +163,7 @@ static const FStringView s_AllSRVTypes[] =
 	TEXTVIEW("ByteAddressBuffer"),
 	TEXTVIEW("StructuredBuffer"),
 	TEXTVIEW("RaytracingAccelerationStructure"),
+	TEXTVIEW("FResourceCollection"),
 };
 
 static const FStringView s_AllUAVTypes[] =
diff --git a/Engine/Source/Runtime/RenderCore/Private/ShaderParameterStruct.cpp b/Engine/Source/Runtime/RenderCore/Private/ShaderParameterStruct.cpp
--- a/Engine/Source/Runtime/RenderCore/Private/ShaderParameterStruct.cpp
+++ b/Engine/Source/Runtime/RenderCore/Private/ShaderParameterStruct.cpp
@@ -87,7 +87,8 @@ struct FShaderParameterStructBindingContext
 				BaseType == UBMT_SRV ||
 				BaseType == UBMT_UAV ||
 				BaseType == UBMT_SAMPLER ||
-				BaseType == UBMT_DESCRIPTOR_RANGE);
+				BaseType == UBMT_DESCRIPTOR_RANGE ||
+				BaseType == UBMT_RESOURCE_COLLECTION);
 
 			const bool bIsRDGResource =
 				IsRDGResourceReferenceShaderParameterType(BaseType);
@@ -531,6 +532,7 @@ void ValidateShaderParameters(const TShaderRef<FShader>& Shader, const FShaderPa
 			case UBMT_SRV:
 			case UBMT_UAV:
 			case UBMT_SAMPLER:
+			case UBMT_RESOURCE_COLLECTION:
 			{
 				const FRHIResource* Resource = Reader.Read<const FRHIResource*>(Parameter);
 				if (!Resource && !ExclusionList.Contains(Parameter.BaseIndex))
@@ -603,7 +605,8 @@ void ValidateShaderParameterResourcesRHI(const void* Contents, const FRHIUniform
 			Parameter.MemberType == UBMT_RDG_TEXTURE_UAV ||
 			Parameter.MemberType == UBMT_RDG_TEXTURE_NON_PIXEL_SRV ||
 			Parameter.MemberType == UBMT_RDG_BUFFER_SRV ||
-			Parameter.MemberType == UBMT_RDG_BUFFER_UAV;
+			Parameter.MemberType == UBMT_RDG_BUFFER_UAV ||
+			Parameter.MemberType == UBMT_RESOURCE_COLLECTION;
 
 		// Allow null SRV's in uniform buffers for feature levels that don't support SRV's in shaders
 		// Mobile does not set some UAVs to not break RDG pass auto-merging
@@ -690,6 +693,12 @@ FRHIShaderParameterResource ExtractShaderParameterResource(FShaderParameterReade
 		checkSlow(DescriptorRange);
 		return FRHIShaderParameterResource(DescriptorRange, GetParameterIndex(Parameter));
 	}
+	case UBMT_RESOURCE_COLLECTION:
+	{
+		FRHIResourceCollection* Collection = Reader.Read<FRHIResourceCollection*>(Parameter);
+		checkSlow(Collection);
+		return FRHIShaderParameterResource(Collection, GetParameterIndex(Parameter));
+	}
 	case UBMT_RDG_TEXTURE:
 	{
 		FRDGTexture* RDGTexture = Reader.Read<FRDGTexture*>(Parameter);
diff --git a/Engine/Source/Runtime/RenderCore/Public/ShaderParameterMacros.h b/Engine/Source/Runtime/RenderCore/Public/ShaderParameterMacros.h
--- a/Engine/Source/Runtime/RenderCore/Public/ShaderParameterMacros.h
+++ b/Engine/Source/Runtime/RenderCore/Public/ShaderParameterMacros.h
@@ -1652,6 +1652,14 @@ private:
 #define SHADER_PARAMETER_SRV_ARRAY(ShaderType,MemberName, ArrayDecl) \
 	INTERNAL_SHADER_PARAMETER_EXPLICIT(UBMT_SRV, TShaderResourceParameterTypeInfo<FRHIShaderResourceView* ArrayDecl>, FRHIShaderResourceView*,MemberName,ArrayDecl,,EShaderPrecisionModifier::Float,TEXT(#ShaderType),false)
 
+/** Adds a shader resource collection.
+ *
+ * Example:
+ *	SHADER_PARAMETER_RESOURCE_COLLECTION(FResourceCollection, MyResourceCollection)
+ */
+#define SHADER_PARAMETER_RESOURCE_COLLECTION(ShaderType,MemberName) \
+	INTERNAL_SHADER_PARAMETER_EXPLICIT(UBMT_RESOURCE_COLLECTION, TShaderResourceParameterTypeInfo<FRHIResourceCollection*>, FRHIResourceCollection*,MemberName,, = nullptr,EShaderPrecisionModifier::Float,TEXT(#ShaderType),false)
+
 /** Adds an unordered access view.
  *
  * Example:
diff --git a/Engine/Source/Runtime/Renderer/Private/DeferredShadingRenderer.cpp b/Engine/Source/Runtime/Renderer/Private/DeferredShadingRenderer.cpp
--- a/Engine/Source/Runtime/Renderer/Private/DeferredShadingRenderer.cpp
+++ b/Engine/Source/Runtime/Renderer/Private/DeferredShadingRenderer.cpp
@@ -1149,6 +1149,14 @@ void FDeferredShadingSceneRenderer::SetupRayTracingRenderingData(FRDGBuilder& Gr
 
 	bool bAnyInlineHardwareRayTracingPassEnabled = false;
 
+	for (int32 ViewExt = 0; ViewExt < ViewFamily.ViewExtensions.Num(); ++ViewExt)
+	{
+		if (EnumHasAnyFlags(ViewFamily.ViewExtensions[ViewExt]->GetFlags(), ESceneViewExtensionFlags::RequiresHardwareInlineRayTracing))
+		{
+			bAnyInlineHardwareRayTracingPassEnabled = true;
+		}
+	}
+
 	for (FViewInfo& View : Views)
 	{
 		if (View.bHasAnyRayTracingPass)
@@ -3282,7 +3290,7 @@ void FDeferredShadingSceneRenderer::Render(FRDGBuilder& GraphBuilder, const FSce
 		// End shadow and fog after base pass
 
 #if RHI_RAYTRACING
-		if (IsRayTracingEnabled(ViewFamily.GetShaderPlatform()) && GRHISupportsRayTracingShaders)
+		if (IsRayTracingEnabled(ViewFamily.GetShaderPlatform()) && (GRHISupportsInlineRayTracing || GRHISupportsRayTracingShaders))
 		{
 			for (int32 ViewExt = 0; ViewExt < ViewFamily.ViewExtensions.Num(); ++ViewExt)
 			{
diff --git a/Engine/Source/Runtime/Renderer/Private/FXRenderingUtils.cpp b/Engine/Source/Runtime/Renderer/Private/FXRenderingUtils.cpp
--- a/Engine/Source/Runtime/Renderer/Private/FXRenderingUtils.cpp
+++ b/Engine/Source/Runtime/Renderer/Private/FXRenderingUtils.cpp
@@ -322,6 +322,17 @@ FRHIShaderResourceView* UE::FXRenderingUtils::RayTracing::GetRayTracingSceneView
 	return nullptr;
 }
 
+FRDGBufferSRVRef UE::FXRenderingUtils::RayTracing::GetRayTracingSceneViewRDG(const FSceneInterface& InScene, const FSceneView& InView)
+{
+	if (const FScene* Scene = InScene.GetRenderScene())
+	{
+		const FViewInfo& View = static_cast<const FViewInfo&>(InView);
+		return Scene->RayTracingScene.GetLayerView(ERayTracingSceneLayer::Base, View.GetRayTracingSceneViewHandle());
+	}
+
+	return nullptr;
+}
+
 FShaderBindingTableRHIRef UE::FXRenderingUtils::RayTracing::CreateShaderBindingTable(FRHICommandListBase& RHICmdList, const FSceneInterface* InScene, uint32 LocalBindingDataSize)
 {
 	if (const FScene* Scene = InScene->GetRenderScene())
diff --git a/Engine/Source/Runtime/Renderer/Private/GPUScene.h b/Engine/Source/Runtime/Renderer/Private/GPUScene.h
--- a/Engine/Source/Runtime/Renderer/Private/GPUScene.h
+++ b/Engine/Source/Runtime/Renderer/Private/GPUScene.h
@@ -298,7 +298,7 @@ public:
 	/**
 	 * Return the GPU scene resource
 	 */
-	FGPUSceneResourceParameters GetShaderParameters(FRDGBuilder& GraphBuilder) const;
+	RENDERER_API FGPUSceneResourceParameters GetShaderParameters(FRDGBuilder& GraphBuilder) const;
 
 	/**
 	 * Draw GPU-Scene debug info, such as bounding boxes. Call once per view at some point in the frame after GPU scene has been updated fully.
diff --git a/Engine/Source/Runtime/Renderer/Private/SceneRendering.h b/Engine/Source/Runtime/Renderer/Private/SceneRendering.h
--- a/Engine/Source/Runtime/Renderer/Private/SceneRendering.h
+++ b/Engine/Source/Runtime/Renderer/Private/SceneRendering.h
@@ -1656,6 +1656,11 @@ public:
 		RayTracingSceneViewHandle = InHandle;
 	}
 
+	virtual FRDGBufferRef GetInlineRayTracingBindingDataBuffer() const override
+	{
+		return InlineRayTracingBindingDataBuffer;
+	}
+
 	struct FRayTracingData
 	{
 		FRayTracingPipelineState* PipelineState = nullptr;
diff --git a/Engine/Source/Runtime/Renderer/Public/FXRenderingUtils.h b/Engine/Source/Runtime/Renderer/Public/FXRenderingUtils.h
--- a/Engine/Source/Runtime/Renderer/Public/FXRenderingUtils.h
+++ b/Engine/Source/Runtime/Renderer/Public/FXRenderingUtils.h
@@ -86,6 +86,7 @@ namespace UE::FXRenderingUtils
 		RENDERER_API bool HasRayTracingScene(const FSceneInterface& Scene);
 		RENDERER_API FRHIRayTracingScene* GetRayTracingScene(const FSceneInterface& Scene, const FSceneView& View);
 		RENDERER_API FRHIShaderResourceView* GetRayTracingSceneView(FRHICommandListBase& RHICmdList, const FSceneInterface& Scene, const FSceneView& View);
+		RENDERER_API FRDGBufferSRVRef GetRayTracingSceneViewRDG(const FSceneInterface& Scene, const FSceneView& View);
 
 		RENDERER_API FRHIShaderResourceView* GetInlineRayTracingBindingDataBuffer(const FSceneInterface& Scene);
 
@@ -107,7 +108,7 @@ namespace UE::FXRenderingUtils
 			return nullptr;
 		}
 
-		RENDERER_API FShaderBindingTableRHIRef CreateShaderBindingTable(FRHICommandListBase& RHICmdList, const FSceneInterface* Scene,  uint32 LocalBindingDataSize);		
+		RENDERER_API FShaderBindingTableRHIRef CreateShaderBindingTable(FRHICommandListBase& RHICmdList, const FSceneInterface* Scene,  uint32 LocalBindingDataSize);
 		RENDERER_API TConstArrayView<FRayTracingShaderBindingData> GetVisibleRayTracingShaderBindings(const FSceneView& View);
 		RENDERER_API const FRayTracingMeshCommandStorage* GetRayTracingMeshCommands(const FSceneInterface* Scene);
 	}
