Color: Add HLG support (ITU-R BT.2100-3) for image sequences, media i/o sources and texture assets.

* Implemented EncodeHLG/DecodeHLG (CPU) and HLGToLinear/LinearToHLG/HLGToLinearWithOOTF (GPU). Negative values are retained.
* Tentative support for the optional BT.2100 OOTF, controlled by r.Media.HLG.ApplyOOTF (off by default).
* Enabled HLG in Aja and Blackmagic media I/O (previously fell back to sRGB).

#Synced-CL 52172240

diff --git a/Engine/Plugins/Media/AjaMedia/Source/AjaMedia/Private/Player/AjaMediaPlayer.cpp b/Engine/Plugins/Media/AjaMedia/Source/AjaMedia/Private/Player/AjaMediaPlayer.cpp
--- a/Engine/Plugins/Media/AjaMedia/Source/AjaMedia/Private/Player/AjaMediaPlayer.cpp
+++ b/Engine/Plugins/Media/AjaMedia/Source/AjaMedia/Private/Player/AjaMediaPlayer.cpp
@@ -634,8 +634,7 @@ bool FAjaMediaPlayer::OnInputFrameReceived(const AJA::AJAInputFrameData& InInput
 				AjaColorSettings->SetEncodingOverride(UE::Color::EEncoding::ST2084);
 				break;
 			case EAjaHDRMetadataEOTF::HLG:
-				UE_MEDIA_IO_LOG_THROTTLE(LogAjaMedia, Warning, TimeBetweenLogs, TEXT("HLG EOTF is not supported at the moment. Defaulting to sRGB."));
-				AjaColorSettings->SetEncodingOverride(UE::Color::EEncoding::sRGB);
+				AjaColorSettings->SetEncodingOverride(UE::Color::EEncoding::HLG);
 				break;
 			default:
 				checkNoEntry();
diff --git a/Engine/Plugins/Media/BlackmagicMedia/Source/BlackmagicMedia/Private/Player/BlackmagicMediaPlayer.cpp b/Engine/Plugins/Media/BlackmagicMedia/Source/BlackmagicMedia/Private/Player/BlackmagicMediaPlayer.cpp
--- a/Engine/Plugins/Media/BlackmagicMedia/Source/BlackmagicMedia/Private/Player/BlackmagicMediaPlayer.cpp
+++ b/Engine/Plugins/Media/BlackmagicMedia/Source/BlackmagicMedia/Private/Player/BlackmagicMediaPlayer.cpp
@@ -319,8 +319,7 @@ namespace BlackmagicMediaPlayerHelpers
 							MediaPlayer->BlackmagicColorSettings->SetEncodingOverride(UE::Color::EEncoding::ST2084);
 							break;
 						case EBlackmagicHDRMetadataEOTF::HLG:
-							UE_MEDIA_IO_LOG_THROTTLE(LogBlackmagicMedia, Warning, TimeBetweenLogs, TEXT("HLG EOTF is not supported at the moment. Defaulting to sRGB."));
-							MediaPlayer->BlackmagicColorSettings->SetEncodingOverride(UE::Color::EEncoding::sRGB);
+							MediaPlayer->BlackmagicColorSettings->SetEncodingOverride(UE::Color::EEncoding::HLG);
 							break;
 						default:
 							checkNoEntry();
diff --git a/Engine/Plugins/Media/ImgMedia/Shaders/Private/ExrSwizzler.usf b/Engine/Plugins/Media/ImgMedia/Shaders/Private/ExrSwizzler.usf
--- a/Engine/Plugins/Media/ImgMedia/Shaders/Private/ExrSwizzler.usf
+++ b/Engine/Plugins/Media/ImgMedia/Shaders/Private/ExrSwizzler.usf
@@ -48,6 +48,8 @@ float4x4 ColorSpaceMatrix;
 #define UE_Color_EEncoding_sRGB			2u
 #define UE_Color_EEncoding_ST2084		3u
 #define UE_Color_EEncoding_SLog3		12u
+#define UE_Color_EEncoding_HLG			17u
+#define UE_Color_EEncoding_HLG_OOTF		128u // Sentinel, deliberately outside EEncoding range
 
 float3 ApplyColorTransform(float3 Color, float3x3 InColorSpaceMatrix, uint InEOTF)
 {
@@ -57,6 +59,8 @@ float3 ApplyColorTransform(float3 Color, float3x3 InColorSpaceMatrix, uint InEOT
 	case UE_Color_EEncoding_sRGB:	Color = sRGBToLinear(Color); break;
 	case UE_Color_EEncoding_ST2084: Color = ST2084ToLinear(Color); break;
 	case UE_Color_EEncoding_SLog3:	Color = SLog3ToLinear(Color); break;
+	case UE_Color_EEncoding_HLG:	Color = HLGToLinear(Color); break;
+	case UE_Color_EEncoding_HLG_OOTF:	Color = HLGToLinearWithOOTF(Color); break;
 	default: break;
 	}
 
diff --git a/Engine/Plugins/Media/ImgMedia/Source/ImgMedia/Private/Readers/ExrImgMediaReaderGpu.cpp b/Engine/Plugins/Media/ImgMedia/Source/ImgMedia/Private/Readers/ExrImgMediaReaderGpu.cpp
--- a/Engine/Plugins/Media/ImgMedia/Source/ImgMedia/Private/Readers/ExrImgMediaReaderGpu.cpp
+++ b/Engine/Plugins/Media/ImgMedia/Source/ImgMedia/Private/Readers/ExrImgMediaReaderGpu.cpp
@@ -12,6 +12,7 @@
 #include "ExrReaderGpu.h"
 #include "RHICommandList.h"
 #include "ExrSwizzlingShader.h"
+#include "MediaShaders.h"
 #include "SceneUtils.h"
 
 
@@ -510,7 +511,7 @@ void FExrImgMediaReaderGpu::CreateSampleConverterCallback(TSharedPtr<FExrMediaTe
 			if (ConverterParams.SourceColorSettings.IsValid())
 			{
 				Parameters.bApplyColorTransform = 1u;
-				Parameters.EOTF = static_cast<uint32>(ConverterParams.SourceColorSettings->GetEncodingOverride());
+				Parameters.EOTF = MediaShaders::GetMediaShaderEOTF_RenderThread(ConverterParams.SourceColorSettings->GetEncodingOverride());
 
 				const UE::Color::FColorSpace& DestinationCS = UE::Color::FColorSpace::GetWorking();
 				const UE::Color::FColorSpace& SourceCS = ConverterParams.SourceColorSettings->GetColorSpaceOverride(DestinationCS);
diff --git a/Engine/Plugins/Media/NDIMedia/Shaders/Private/NDIMediaShaders.usf b/Engine/Plugins/Media/NDIMedia/Shaders/Private/NDIMediaShaders.usf
--- a/Engine/Plugins/Media/NDIMedia/Shaders/Private/NDIMediaShaders.usf
+++ b/Engine/Plugins/Media/NDIMedia/Shaders/Private/NDIMediaShaders.usf
@@ -24,6 +24,8 @@
 #define UE_Color_EEncoding_CanonLog		14u
 #define UE_Color_EEncoding_ProTune		15u
 #define UE_Color_EEncoding_VLog			16u
+#define UE_Color_EEncoding_HLG			17u
+#define UE_Color_EEncoding_HLG_OOTF	128u // Sentinel, deliberately outside EEncoding range
 
 #define ToneMapMethod_None           0u
 #define ToneMapMethod_Hable          1u
@@ -37,6 +39,8 @@ float3 ApplyColorProcessingRGB(float3 Color, float4x4 CSMtx, uint EOTF, uint Ton
 	case UE_Color_EEncoding_sRGB:	Color = sRGBToLinear(Color); break;
 	case UE_Color_EEncoding_ST2084: Color = ST2084ToLinear(Color); break;
 	case UE_Color_EEncoding_SLog3:	Color = SLog3ToLinear(Color); break;
+	case UE_Color_EEncoding_HLG:	Color = HLGToLinear(Color); break;
+	case UE_Color_EEncoding_HLG_OOTF: Color = HLGToLinearWithOOTF(Color); break;
 	default: break;
 	}
 
diff --git a/Engine/Plugins/Media/NDIMedia/Source/NDIMediaRendering/Private/NDIMediaShaders.cpp b/Engine/Plugins/Media/NDIMedia/Source/NDIMediaRendering/Private/NDIMediaShaders.cpp
--- a/Engine/Plugins/Media/NDIMedia/Source/NDIMediaRendering/Private/NDIMediaShaders.cpp
+++ b/Engine/Plugins/Media/NDIMedia/Source/NDIMediaRendering/Private/NDIMediaShaders.cpp
@@ -33,7 +33,7 @@ void FNDIMediaShaderUYVAtoBGRAPS::SetParameters(FRHIBatchedShaderParameters& InB
 		UB.OutputHeight = InParameters.OutputSize.Y;
 		UB.ColorTransform = InParameters.ColorTransform;
 		UB.CSTransform = InParameters.CSTransform;
-		UB.EOTF = static_cast<uint32>(InParameters.Encoding);
+		UB.EOTF = MediaShaders::GetMediaShaderEOTF_RenderThread(InParameters.Encoding);
 		UB.ToneMapMethod = static_cast<uint32>(InParameters.ToneMapMethod);
 		UB.YUVTexture = InParameters.YUVTexture;
 		UB.AlphaTexture = InParameters.AlphaTexture;
diff --git a/Engine/Plugins/Media/TmvMedia/Shaders/Private/TmvColorProcessing.ush b/Engine/Plugins/Media/TmvMedia/Shaders/Private/TmvColorProcessing.ush
--- a/Engine/Plugins/Media/TmvMedia/Shaders/Private/TmvColorProcessing.ush
+++ b/Engine/Plugins/Media/TmvMedia/Shaders/Private/TmvColorProcessing.ush
@@ -7,6 +7,8 @@
 #define UE_Color_EEncoding_sRGB			2u
 #define UE_Color_EEncoding_ST2084		3u
 #define UE_Color_EEncoding_SLog3		12u
+#define UE_Color_EEncoding_HLG			17u
+#define UE_Color_EEncoding_HLG_OOTF	128u // Sentinel, deliberately outside EEncoding range
 
 // Simplified version of ApplyColorProcessingRGB, removed tonemapping.
 float3 ApplyColorTransformRGB(float3 InColor, float4x4 InColorSpaceMatrix, uint InEOTF)
@@ -17,6 +19,8 @@ float3 ApplyColorTransformRGB(float3 InColor, float4x4 InColorSpaceMatrix, uint
 	case UE_Color_EEncoding_sRGB:	InColor = sRGBToLinear(InColor); break;
 	case UE_Color_EEncoding_ST2084: InColor = ST2084ToLinear(InColor); break;
 	case UE_Color_EEncoding_SLog3:	InColor = SLog3ToLinear(InColor); break;
+	case UE_Color_EEncoding_HLG:      InColor = HLGToLinear(InColor); break;
+	case UE_Color_EEncoding_HLG_OOTF: InColor = HLGToLinearWithOOTF(InColor); break;
 	default: break;
 	}
 
diff --git a/Engine/Plugins/Media/TmvMedia/Shaders/Private/TmvTextureYuvConverter.usf b/Engine/Plugins/Media/TmvMedia/Shaders/Private/TmvTextureYuvConverter.usf
--- a/Engine/Plugins/Media/TmvMedia/Shaders/Private/TmvTextureYuvConverter.usf
+++ b/Engine/Plugins/Media/TmvMedia/Shaders/Private/TmvTextureYuvConverter.usf
@@ -67,7 +67,11 @@ float3 EncodeTransferFunction(float3 InColor, uint InEOTF)
 	case UE_Color_EEncoding_SLog3:
 		InColor = LinearToSLog3(InColor);
 		break;
-		
+
+	case UE_Color_EEncoding_HLG:
+		InColor = LinearToHLG(InColor);
+		break;
+
 	default:
 		break;
 	}
diff --git a/Engine/Plugins/Media/TmvMedia/Source/TmvMedia/Private/SampleConverter/TmvMediaTextureSampleConverter.cpp b/Engine/Plugins/Media/TmvMedia/Source/TmvMedia/Private/SampleConverter/TmvMediaTextureSampleConverter.cpp
--- a/Engine/Plugins/Media/TmvMedia/Source/TmvMedia/Private/SampleConverter/TmvMediaTextureSampleConverter.cpp
+++ b/Engine/Plugins/Media/TmvMedia/Source/TmvMedia/Private/SampleConverter/TmvMediaTextureSampleConverter.cpp
@@ -210,12 +210,12 @@ namespace UE::TmvMedia::Converter
 		}
 		
 		OutColorParams.bApplyColorTransform = 0u;
-		
+
 		// Note: we always decode to linear, if the destination RT is sRGB, it will encode to that automatically.
 		if (SourceEncoding != Color::EEncoding::Linear && SourceEncoding != Color::EEncoding::None)
 		{
 			OutColorParams.bApplyColorTransform = 1u;
-			OutColorParams.EOTF = static_cast<uint32>(SourceEncoding);
+			OutColorParams.EOTF = MediaShaders::GetMediaShaderEOTF_RenderThread(SourceEncoding);
 		}
 		else
 		{
diff --git a/Engine/Plugins/Media/TmvMedia/Source/TmvMedia/Private/Transcoder/TmvMediaFrameConverter.cpp b/Engine/Plugins/Media/TmvMedia/Source/TmvMedia/Private/Transcoder/TmvMediaFrameConverter.cpp
--- a/Engine/Plugins/Media/TmvMedia/Source/TmvMedia/Private/Transcoder/TmvMediaFrameConverter.cpp
+++ b/Engine/Plugins/Media/TmvMedia/Source/TmvMedia/Private/Transcoder/TmvMediaFrameConverter.cpp
@@ -143,6 +143,26 @@ namespace UE::TmvMedia::FrameConverter
 		}
 	};
 
+	/** Implements HLG transfer function */
+	struct FHLGEncoder
+	{
+		static float Encode(float InValue)
+		{
+			return Color::EncodeHLG(InValue);
+		}
+
+		static VectorRegister4Float EncodeVector(VectorRegister4Float InVec)
+		{
+			// HLG OETF is piecewise (sqrt + log), no simple SIMD form — fall back to scalar.
+			alignas(16) float Values[4];
+			VectorStoreAligned(InVec, Values);
+			Values[0] = Encode(Values[0]);
+			Values[1] = Encode(Values[1]);
+			Values[2] = Encode(Values[2]);
+			return VectorLoadAligned(Values);
+		}
+	};
+
 	/** Implements slog3 transfer function */
 	struct FSLog3Encoder
 	{
@@ -255,6 +275,10 @@ namespace UE::TmvMedia::FrameConverter
 			ConvertYUVInPlace<FSLog3Encoder>(InImage, Rgb2YuvMatrix);
 			break;
 
+		case Color::EEncoding::HLG:
+			ConvertYUVInPlace<FHLGEncoder>(InImage, Rgb2YuvMatrix);
+			break;
+
 		default:
 			// remain linear.
 			ConvertYUVInPlace<FLinearEncoder>(InImage, Rgb2YuvMatrix);
diff --git a/Engine/Shaders/Private/GammaCorrectionCommon.ush b/Engine/Shaders/Private/GammaCorrectionCommon.ush
--- a/Engine/Shaders/Private/GammaCorrectionCommon.ush
+++ b/Engine/Shaders/Private/GammaCorrectionCommon.ush
@@ -196,3 +196,79 @@ SLog3ToLinear(float3 Value)
 {
 	return select((Value >= 171.2102946929f / 1023.0f), (pow(10.0f, (Value * 1023.0f - 420.f) / 261.5f)) * 0.19f - 0.01f, (Value * 1023.0f - 95.0f) * 0.01125000f / (171.2102946929f - 95.0f));
 }
+
+float LinearToHLG(float Value)
+{
+	const float HLG_a = 0.17883277f;
+	const float HLG_b = 0.28466892f; // 1.0 - 4.0 * HLG_a
+	const float HLG_c = 0.559910729529562f; // 0.5 - HLG_a * log(4.0 * HLG_a)
+
+	const float E = abs(Value);
+	float Out = 0.0f;
+	
+	if (E <= 1.0f / 12.0f)
+	{
+		Out = sqrt(3.0f * E);
+	}
+	else
+	{
+		Out = HLG_a * log(12.0f * E - HLG_b) + HLG_c;
+	}
+
+	return Out * sign(Value);
+}
+
+float HLGToLinear(float Value)
+{
+	const float HLG_a = 0.17883277f;
+	const float HLG_b = 0.28466892f; // 1.0 - 4.0 * HLG_a
+	const float HLG_c = 0.559910729529562f; // 0.5 - HLG_a * log(4.0 * HLG_a)
+
+	const float Eprime = abs(Value);
+	float Out = 0.0f;
+	
+	if (Eprime <= 0.5f)
+	{
+		Out = Eprime * Eprime / 3.0f;
+	}
+	else
+	{
+		Out = exp((Eprime - HLG_c) / HLG_a) + HLG_b;
+		Out /= 12.0f;
+	}
+
+	return Out * sign(Value);
+}
+
+float3 LinearToHLG(float3 L)
+{
+	return float3(
+		LinearToHLG(L.x),
+		LinearToHLG(L.y),
+		LinearToHLG(L.z)
+	);
+}
+
+float3 HLGToLinear(float3 E)
+{
+	return float3(
+		HLGToLinear(E.x),
+		HLGToLinear(E.y),
+		HLGToLinear(E.z)
+	);
+}
+
+// Full BT.2100 HLG EOTF: inverse OETF followed by OOTF.
+// OOTF: FD = alpha * Ys^(gamma-1) * E, with alpha=PeakLuminance, gamma=SystemGamma.
+// Defaults to reference display: 1000 cd/m2, gamma=1.2.
+float3 HLGToLinearWithOOTF(float3 Color, float PeakLuminance = 1000.0f, float SystemGamma = 1.2f)
+{
+	// BT.2020 luminance weights (ITU-R BT.2100-3 Table 5).
+	const float3 Rec2020LuminanceWeights = float3(0.2627f, 0.6780f, 0.0593f);
+
+	Color = HLGToLinear(Color);
+	float Ys = dot(Color, Rec2020LuminanceWeights);
+
+	return PeakLuminance * pow(Ys, SystemGamma - 1.0f) * Color;
+}
+
diff --git a/Engine/Shaders/Private/MediaShaders.usf b/Engine/Shaders/Private/MediaShaders.usf
--- a/Engine/Shaders/Private/MediaShaders.usf
+++ b/Engine/Shaders/Private/MediaShaders.usf
@@ -69,6 +69,8 @@ void BMPConvertPS(
 #define UE_Color_EEncoding_CanonLog		14u
 #define UE_Color_EEncoding_ProTune		15u
 #define UE_Color_EEncoding_VLog			16u
+#define UE_Color_EEncoding_HLG			17u
+#define UE_Color_EEncoding_HLG_OOTF		128u // Sentinel, deliberately outside EEncoding range
 
 #define ToneMapMethod_None           0u
 #define ToneMapMethod_Hable          1u
@@ -82,6 +84,8 @@ float3 ApplyColorProcessingRGB(float3 Color, float4x4 CSMtx, uint EOTF, uint Ton
 	case UE_Color_EEncoding_sRGB:	Color = sRGBToLinear(Color); break;
 	case UE_Color_EEncoding_ST2084: Color = ST2084ToLinear(Color); break;
 	case UE_Color_EEncoding_SLog3:	Color = SLog3ToLinear(Color); break;
+	case UE_Color_EEncoding_HLG:	Color = HLGToLinear(Color); break;
+	case UE_Color_EEncoding_HLG_OOTF:	Color = HLGToLinearWithOOTF(Color); break;
 	default: break;
 	}
 
diff --git a/Engine/Source/Runtime/Core/Private/ColorManagement/Tests/ColorManagementTests.cpp b/Engine/Source/Runtime/Core/Private/ColorManagement/Tests/ColorManagementTests.cpp
--- a/Engine/Source/Runtime/Core/Private/ColorManagement/Tests/ColorManagementTests.cpp
+++ b/Engine/Source/Runtime/Core/Private/ColorManagement/Tests/ColorManagementTests.cpp
@@ -278,6 +278,83 @@ static void TestLuminance()
 	CHECK_MESSAGE(TEXT("Rec2020 luminance factors equality test"), LuminanceFactors.Equals(FLinearColor(0.262700212011f, 0.677998071519f, 0.0593017164699f), UE_SMALL_NUMBER));
 }
 
+TEST_CASE_NAMED(FHLG2100TransferFunctionsTest, "System::ColorManagement::HLG2100TransferFunctions", "[EditorContext][EngineFilter]")
+{
+	using namespace UE::Color;
+
+	SECTION("OETF piecewise boundary")
+	{
+		// Verify the piecewise boundary maps correctly: E = 1/12 <-> E' = 0.5
+		CHECK_MESSAGE(TEXT("HLG OETF at boundary (1/12) should equal 0.5"), FMath::IsNearlyEqual(EncodeHLG(1.0f / 12.0f), 0.5f, UE_SMALL_NUMBER));
+		CHECK_MESSAGE(TEXT("HLG inverse OETF at boundary (0.5) should equal 1/12"), FMath::IsNearlyEqual(DecodeHLG(0.5f), 1.0f / 12.0f, UE_SMALL_NUMBER));
+	}
+
+	SECTION("OETF negative values")
+	{
+		// HLG OETF should handle negative values via CopySign
+		float Positive = EncodeHLG(0.5f);
+		float Negative = EncodeHLG(-0.5f);
+		CHECK_MESSAGE(TEXT("HLG OETF negative value should mirror positive"), FMath::IsNearlyEqual(Negative, -Positive, UE_SMALL_NUMBER));
+	}
+
+	SECTION("BT.2100 round-trip at default luminance")
+	{
+		// Encode then decode should produce identity for display-referred colors.
+		// Note: black and single-channel colors are excluded because the OOTF luminance weighting
+		// (Ys^(gamma-1)) is degenerate at zero luminance, making round-trip undefined.
+		const FLinearColor TestColors[] = {
+			FLinearColor(100.0f, 200.0f, 50.0f),
+			FLinearColor(500.0f, 500.0f, 500.0f),
+			FLinearColor(1000.0f, 1000.0f, 1000.0f),
+			FLinearColor(10.0f, 50.0f, 30.0f),
+		};
+
+		for (const FLinearColor& Src : TestColors)
+		{
+			FLinearColor Encoded = EncodeHLG_2100(Src);
+			FLinearColor Decoded = DecodeHLG_2100(Encoded);
+
+			bool bRClose = FMath::IsNearlyEqual(Decoded.R, Src.R, FMath::Abs(Src.R) * UE_KINDA_SMALL_NUMBER);
+			bool bGClose = FMath::IsNearlyEqual(Decoded.G, Src.G, FMath::Abs(Src.G) * UE_KINDA_SMALL_NUMBER);
+			bool bBClose = FMath::IsNearlyEqual(Decoded.B, Src.B, FMath::Abs(Src.B) * UE_KINDA_SMALL_NUMBER);
+			CHECK_MESSAGE(TEXT("BT.2100 round-trip should preserve color values"), bRClose && bGClose && bBClose);
+		}
+	}
+
+	SECTION("BT.2100 round-trip at various peak luminances")
+	{
+		const FLinearColor Src(200.0f, 400.0f, 100.0f);
+		const float PeakLuminances[] = { 400.0f, 1000.0f, 2000.0f, 4000.0f };
+
+		for (float Peak : PeakLuminances)
+		{
+			FLinearColor Encoded = EncodeHLG_2100(Src, Peak);
+			FLinearColor Decoded = DecodeHLG_2100(Encoded, Peak);
+
+			bool bRClose = FMath::IsNearlyEqual(Decoded.R, Src.R, FMath::Abs(Src.R) * UE_KINDA_SMALL_NUMBER);
+			bool bGClose = FMath::IsNearlyEqual(Decoded.G, Src.G, FMath::Abs(Src.G) * UE_KINDA_SMALL_NUMBER);
+			bool bBClose = FMath::IsNearlyEqual(Decoded.B, Src.B, FMath::Abs(Src.B) * UE_KINDA_SMALL_NUMBER);
+			CHECK_MESSAGE(TEXT("BT.2100 round-trip at varying peak luminance should preserve color values"), bRClose && bGClose && bBClose);
+		}
+	}
+
+	SECTION("BT.2100 round-trip with black luminance")
+	{
+		const FLinearColor Src(300.0f, 600.0f, 150.0f);
+		const float PeakLuminance = 1000.0f;
+		const float BlackLuminance = 0.01f;
+
+		FLinearColor Encoded = EncodeHLG_2100(Src, PeakLuminance, BlackLuminance);
+		FLinearColor Decoded = DecodeHLG_2100(Encoded, PeakLuminance, BlackLuminance);
+
+		bool bRClose = FMath::IsNearlyEqual(Decoded.R, Src.R, FMath::Abs(Src.R) * UE_KINDA_SMALL_NUMBER);
+		bool bGClose = FMath::IsNearlyEqual(Decoded.G, Src.G, FMath::Abs(Src.G) * UE_KINDA_SMALL_NUMBER);
+		bool bBClose = FMath::IsNearlyEqual(Decoded.B, Src.B, FMath::Abs(Src.B) * UE_KINDA_SMALL_NUMBER);
+		CHECK_MESSAGE(TEXT("BT.2100 round-trip with black luminance should preserve color values"), bRClose && bGClose && bBClose);
+	}
+
+}
+
 TEST_CASE_NAMED(FColorSpaceTest, "System::ColorManagement::ColorSpace", "[EditorContext][EngineFilter]")
 {
 	using namespace UE::Color;
diff --git a/Engine/Source/Runtime/Core/Private/ColorManagement/TransferFunctions.cpp b/Engine/Source/Runtime/Core/Private/ColorManagement/TransferFunctions.cpp
--- a/Engine/Source/Runtime/Core/Private/ColorManagement/TransferFunctions.cpp
+++ b/Engine/Source/Runtime/Core/Private/ColorManagement/TransferFunctions.cpp
@@ -1,6 +1,7 @@
 // Copyright Epic Games, Inc. All Rights Reserved.
 
 #include "ColorManagement/TransferFunctions.h"
+#include "ColorManagement/ColorSpace.h"
 #include "Logging/LogMacros.h"
 
 DEFINE_LOG_CATEGORY_STATIC(LogTransferFunctions, Log, All);
@@ -28,6 +29,7 @@ inline TFunction<float(float)> GetTransferFunction(EEncoding SourceEncoding, boo
 	case EEncoding::CanonLog:    return bIsEncode ? &EncodeCanonLog : &DecodeCanonLog;
 	case EEncoding::ProTune:     return bIsEncode ? &EncodeGoProProTune : &DecodeGoProProTune;
 	case EEncoding::VLog:        return bIsEncode ? &EncodePanasonicVLog : &DecodePanasonicVLog;
+	case EEncoding::HLG:         return bIsEncode ? &EncodeHLG : &DecodeHLG;
 	default:
 		check(false);
 		break;
@@ -38,6 +40,127 @@ inline TFunction<float(float)> GetTransferFunction(EEncoding SourceEncoding, boo
 	return nullptr;
 }
 
+namespace
+{
+
+float HLG_SystemGamma(const float PeakLuminance)
+{
+	return 1.2f + 0.42f * FMath::LogX(10.0f, PeakLuminance / 1000.0f);
+}
+
+float HLG_SystemGamma_Extended(const float PeakLuminance)
+{
+	// See ITU-R BT.2100-3 Note 5f: adjusted formula for displays with a nominal peak luminance other than 1000 cd/m^2.
+	constexpr float Kappa = 1.111f;
+	const float Exponent = FMath::Log2(PeakLuminance / 1000.0f);
+
+	return 1.2f * FMath::Pow(Kappa, Exponent);
+}
+
+FLinearColor HLG_OOTF(const FLinearColor& Value, const float PeakLuminance, const float Gamma)
+{
+	const float Alpha = PeakLuminance;
+	const float Ys = FColorSpace::GetRec2020().GetLuminance(Value);
+
+	return Alpha * FMath::Pow(Ys, Gamma - 1.0f) * Value;
+}
+
+FLinearColor HLG_InverseOOTF(const FLinearColor& Value, const float PeakLuminance, const float Gamma)
+{
+	const float Alpha = PeakLuminance;
+	const float Yd = FColorSpace::GetRec2020().GetLuminance(Value);
+
+	if (Yd <= 0.0f)
+	{
+		return FLinearColor::Black;
+	}
+
+	return FMath::Pow(Yd / Alpha, (1.0f - Gamma) / Gamma) * Value / Alpha;
+}
+
+} // anonymous namespace
+
+FLinearColor EncodeHLG_2100(const FLinearColor& LinearDisplayReferred, const float PeakLuminance, const float BlackLuminance)
+{
+	if (!ensureMsgf(PeakLuminance > 0.0f, TEXT("Peak luminance should be positive, aborting.")))
+	{
+		return LinearDisplayReferred;
+	}
+
+	if (!ensureMsgf(BlackLuminance >= 0.0f && BlackLuminance < PeakLuminance, TEXT("Black luminance is out of range, aborting.")))
+	{
+		return LinearDisplayReferred;
+	}
+
+	const float Gamma =
+		(PeakLuminance < 400.0f || PeakLuminance > 2000.0f)
+		? HLG_SystemGamma_Extended(PeakLuminance)
+		: HLG_SystemGamma(PeakLuminance);
+
+	float Beta = FMath::Sqrt(3.0f * FMath::Pow(BlackLuminance / PeakLuminance, 1.0f / Gamma));
+
+	// Beta must be in [0, 1) for valid black lift. Values >= 1 arise from non-physical BlackLuminance
+	// and would cause division by zero or signal inversion in the affine mapping below.
+	if (!ensureMsgf(Beta < 1.0f, TEXT("HLG Beta (%.4f) >= 1 from BlackLuminance=%.2f, PeakLuminance=%.2f. Clamping."
+		" Reduce BlackLuminance so that Lb/Lw < (1/3)^gamma."), Beta, BlackLuminance, PeakLuminance))
+	{
+		Beta = 1.0f - UE_SMALL_NUMBER;
+	}
+
+	// The inverse OOTF maps relative display to scene
+	FLinearColor Output = HLG_InverseOOTF(LinearDisplayReferred, PeakLuminance, Gamma);
+
+	// Encode into HLG signal using OETF
+	Output.R = EncodeHLG(Output.R);
+	Output.G = EncodeHLG(Output.G);
+	Output.B = EncodeHLG(Output.B);
+
+	// Invert user black lift
+	return (Output - FLinearColor::White * Beta) / (1.0f - Beta);
+}
+
+FLinearColor DecodeHLG_2100(const FLinearColor& EncodedSceneReferred, const float PeakLuminance, const float BlackLuminance)
+{
+	if (!ensureMsgf(PeakLuminance > 0.0f, TEXT("Peak luminance should be positive, aborting.")))
+	{
+		return EncodedSceneReferred;
+	}
+
+	if (!ensureMsgf(BlackLuminance >= 0.0f && BlackLuminance < PeakLuminance, TEXT("Black luminance is out of range, aborting.")))
+	{
+		return EncodedSceneReferred;
+	}
+
+	const float Gamma =
+		(PeakLuminance < 400.0f || PeakLuminance > 2000.0f)
+		? HLG_SystemGamma_Extended(PeakLuminance)
+		: HLG_SystemGamma(PeakLuminance);
+
+	float Beta = FMath::Sqrt(3.0f * FMath::Pow(BlackLuminance / PeakLuminance, 1.0f / Gamma));
+
+	// Beta must be in [0, 1) for valid black lift. Values >= 1 arise from non-physical BlackLuminance
+	// and would cause signal inversion in the affine mapping below.
+	if (!ensureMsgf(Beta < 1.0f, TEXT("HLG Beta (%.4f) >= 1 from BlackLuminance=%.2f, PeakLuminance=%.2f. Clamping."
+		" Reduce BlackLuminance so that Lb/Lw < (1/3)^gamma."), Beta, BlackLuminance, PeakLuminance))
+	{
+		Beta = 1.0f - UE_SMALL_NUMBER;
+	}
+
+	// Apply user black lift, clamping to zero per BT.2100-3 EOTF: max(0, (1-beta)*E' + beta)
+	FLinearColor Output = (1.0f - Beta) * EncodedSceneReferred + FLinearColor::White * Beta;
+	Output.R = FMath::Max(0.0f, Output.R);
+	Output.G = FMath::Max(0.0f, Output.G);
+	Output.B = FMath::Max(0.0f, Output.B);
+
+	// Decode HLG signal using inverse OETF
+	Output.R = DecodeHLG(Output.R);
+	Output.G = DecodeHLG(Output.G);
+	Output.B = DecodeHLG(Output.B);
+
+	// The OOTF maps relative scene to relative display
+	return HLG_OOTF(Output, PeakLuminance, Gamma);
+}
+
 TFunction<float(float)> GetEncodeFunction(EEncoding SourceEncoding)
 {
 	return GetTransferFunction(SourceEncoding, true);
diff --git a/Engine/Source/Runtime/Core/Public/ColorManagement/ColorManagementDefines.h b/Engine/Source/Runtime/Core/Public/ColorManagement/ColorManagementDefines.h
--- a/Engine/Source/Runtime/Core/Public/ColorManagement/ColorManagementDefines.h
+++ b/Engine/Source/Runtime/Core/Public/ColorManagement/ColorManagementDefines.h
@@ -35,6 +35,7 @@ enum class EEncoding : uint8
 	CanonLog,
 	ProTune,
 	VLog,
+	HLG,
 	Max,
 };
 
diff --git a/Engine/Source/Runtime/Core/Public/ColorManagement/TransferFunctions.h b/Engine/Source/Runtime/Core/Public/ColorManagement/TransferFunctions.h
--- a/Engine/Source/Runtime/Core/Public/ColorManagement/TransferFunctions.h
+++ b/Engine/Source/Runtime/Core/Public/ColorManagement/TransferFunctions.h
@@ -514,6 +514,89 @@ inline float DecodePanasonicVLog(float Value)
 	}
 }
 
+/** HLG (Hybrid Log-Gamma) constants per ITU-R BT.2100-3. */
+namespace HLG
+{
+	constexpr float a = 0.17883277f;
+	constexpr float b = 1.0f - 4.0f * a;
+	constexpr float c = 0.559910729529562f; // 0.5 - a * ln(4a)
+}
+
+/**
+ * HLG reference OETF (Opto-Electronic Transfer Function) per ITU-R BT.2100-3.
+ * Maps a scene-referred linear value E in [0,1] to a non-linear signal E' in [0,1].
+ * Negative values are supported via sign mirroring.
+ *
+ * @param Value Scene-referred linear value.
+ * @return Non-linear HLG signal.
+ */
+inline float EncodeHLG(float Value)
+{
+	const float E = FMath::Abs(Value);
+	float Out = 0.0f;
+
+	if (E <= 1.0f / 12.0f)
+	{
+		Out = FMath::Sqrt(3.0f * E);
+	}
+	else
+	{
+		Out = HLG::a * FMath::Loge(12.0f * E - HLG::b) + HLG::c;
+	}
+
+	return FMath::CopySign(Out, Value);
+}
+
+/**
+ * HLG reference inverse OETF per ITU-R BT.2100-3.
+ * Maps a non-linear HLG signal E' in [0,1] back to a scene-referred linear value E in [0,1].
+ * Negative values are supported via sign mirroring.
+ *
+ * @param Value Non-linear HLG signal.
+ * @return Scene-referred linear value.
+ */
+inline float DecodeHLG(float Value)
+{
+	const float Eprime = FMath::Abs(Value);
+	float Out = 0.0f;
+
+	if (Eprime <= 0.5f)
+	{
+		Out = Eprime * Eprime / 3.0f;
+	}
+	else
+	{
+		Out = FMath::Exp((Eprime - HLG::c) / HLG::a) + HLG::b;
+		Out /= 12.0f;
+	}
+
+	return FMath::CopySign(Out, Value);
+}
+
+/**
+ * Full BT.2100-3 HLG encoding: inverse OOTF followed by OETF, with optional black lift.
+ * Converts display-referred linear light (cd/m2) to a non-linear HLG signal.
+ *
+ * @param LinearDisplayReferred Display-referred linear color in cd/m2.
+ * @param PeakLuminance Display peak luminance in cd/m2 (default 1000).
+ * @param BlackLuminance Display black luminance in cd/m2 (default 0).
+ * @return Encoded HLG signal.
+ */
+CORE_API FLinearColor EncodeHLG_2100(const FLinearColor& LinearDisplayReferred, const float PeakLuminance = 1000.0f, const float BlackLuminance = 0.0f);
+
+/**
+ * Full BT.2100-3 HLG decoding (EOTF): black lift, inverse OETF, then OOTF.
+ * Converts a non-linear HLG signal to display-referred linear light (cd/m2).
+ * The shader equivalent (HLGToLinearWithOOTF in GammaCorrectionCommon.ush) is a
+ * reference-display specialization of this function with PeakLuminance=1000 and no black lift.
+ *
+ * @param EncodedSceneReferred Encoded HLG signal.
+ * @param PeakLuminance Display peak luminance in cd/m2 (default 1000).
+ * @param BlackLuminance Display black luminance in cd/m2 (default 0).
+ * @return Display-referred linear color in cd/m2.
+ */
+CORE_API FLinearColor DecodeHLG_2100(const FLinearColor& EncodedSceneReferred, const float PeakLuminance = 1000.0f, const float BlackLuminance = 0.0f);
+
 /** Get the encode function that matches the encoding type. */
 CORE_API TFunction<float(float)> GetEncodeFunction(EEncoding Encoding);
 
diff --git a/Engine/Source/Runtime/Engine/Classes/Engine/TextureDefines.h b/Engine/Source/Runtime/Engine/Classes/Engine/TextureDefines.h
--- a/Engine/Source/Runtime/Engine/Classes/Engine/TextureDefines.h
+++ b/Engine/Source/Runtime/Engine/Classes/Engine/TextureDefines.h
@@ -431,6 +431,7 @@ enum class ETextureSourceEncoding : uint8
 	TSE_CanonLog	= 14 UMETA(DisplayName = "CanonLog", ToolTip = "Canon Log source encoding to be linearized (before optional sRGB encoding is applied)."),
 	TSE_ProTune		= 15 UMETA(DisplayName = "ProTune", ToolTip = "GoPro ProTune source encoding to be linearized (before optional sRGB encoding is applied)."),
 	TSE_VLog		= 16 UMETA(DisplayName = "V-Log", ToolTip = "Panasonic V-Log source encoding to be linearized (before optional sRGB encoding is applied)."),
+	TSE_HLG			= 17 UMETA(DisplayName = "HLG", ToolTip = "HLG (OETF) source encoding to be linearized (before optional sRGB encoding is applied)."),
 	TSE_MAX,
 };
 
diff --git a/Engine/Source/Runtime/MediaAssets/Public/BaseMediaSourceColorSettings.h b/Engine/Source/Runtime/MediaAssets/Public/BaseMediaSourceColorSettings.h
--- a/Engine/Source/Runtime/MediaAssets/Public/BaseMediaSourceColorSettings.h
+++ b/Engine/Source/Runtime/MediaAssets/Public/BaseMediaSourceColorSettings.h
@@ -18,6 +18,7 @@ enum class EMediaSourceEncoding : uint8
 	MSE_sRGB	= 2 UMETA(DisplayName = "sRGB", ToolTip = "sRGB source encoding to be linearized"),
 	MSE_ST2084	= 3 UMETA(DisplayName = "ST 2084/PQ", ToolTip = "SMPTE ST 2084/PQ source encoding to be linearized"),
 	MSE_SLog3	= 12 UMETA(DisplayName = "SLog3", ToolTip = "Sony SLog3 source encoding to be linearized"),
+	MSE_HLG		= 17 UMETA(DisplayName = "HLG", ToolTip = "HLG (OETF) source encoding to be linearized."),
 
 	MSE_MAX,
 };
diff --git a/Engine/Source/Runtime/OpenColorIOWrapper/Private/OpenColorIOWrapper.cpp b/Engine/Source/Runtime/OpenColorIOWrapper/Private/OpenColorIOWrapper.cpp
--- a/Engine/Source/Runtime/OpenColorIOWrapper/Private/OpenColorIOWrapper.cpp
+++ b/Engine/Source/Runtime/OpenColorIOWrapper/Private/OpenColorIOWrapper.cpp
@@ -918,6 +918,25 @@ FOpenColorIOWrapperProcessor FOpenColorIOWrapperProcessor::CreateTransformToWork
 			ParentTransform->setEncoding("log");
 		}
 		break;
+		case EEncoding::HLG:
+		{
+			/**
+			 * OCIO's built-in CURVE - HLG-OETF expects input scaled to [0, 3] (or equivalently computes
+			 * standard_OETF(E/3)). Post-multiplying by 1/3 via a MatrixTransform after the inverse OETF
+			 * gives an exact match to the ITU-R BT.2100 standard.
+			 */
+			BuiltinTransformRcPtr ChildTransform = BuiltinTransform::Create();
+			ChildTransform->setStyle("CURVE - HLG-OETF-INVERSE");
+			TransformToWCS->appendTransform(ChildTransform);
+
+			MatrixTransformRcPtr MatrixTransform = MatrixTransform::Create();
+			FMatrix44d ScaleMat = FMatrix44d::Identity * (1.0 / 3.0);
+			ScaleMat.M[3][3] = 1.0;
+			MatrixTransform->setMatrix(&ScaleMat.M[0][0]);
+			TransformToWCS->appendTransform(MatrixTransform);
+			ParentTransform->setEncoding("hdr-video");
+		}
+		break;
 
 		default:
 			checkNoEntry();
diff --git a/Engine/Source/Runtime/OpenColorIOWrapper/Private/Tests/OpenColorIOWrapperTests.cpp b/Engine/Source/Runtime/OpenColorIOWrapper/Private/Tests/OpenColorIOWrapperTests.cpp
--- a/Engine/Source/Runtime/OpenColorIOWrapper/Private/Tests/OpenColorIOWrapperTests.cpp
+++ b/Engine/Source/Runtime/OpenColorIOWrapper/Private/Tests/OpenColorIOWrapperTests.cpp
@@ -88,4 +88,40 @@ bool FOpenColorIOTransferFunctionsTest::RunTest(const FString& Parameters)
 	return bSuccess;
 }
 
+IMPLEMENT_SIMPLE_AUTOMATION_TEST(FOpenColorIOHLGTest, "System.OpenColorIO.HLGDecodeMatchesEngine", EAutomationTestFlags::EditorContext | EAutomationTestFlags::EngineFilter)
+bool FOpenColorIOHLGTest::RunTest(const FString& Parameters)
+{
+	using namespace UE::Color;
+
+	// Test multiple values across the HLG range, including the piecewise boundary at 0.5
+	// and values in both the quadratic (<=0.5) and exponential (>0.5) regions.
+	const float TestValues[] = { 0.0f, 0.1f, 0.25f, 0.5f, 0.75f, 0.9f, 1.0f };
+
+	FOpenColorIOWrapperSourceColorSettings HLGSettings;
+	HLGSettings.EncodingOverride = EEncoding::HLG;
+
+	FOpenColorIOWrapperProcessor Processor = FOpenColorIOWrapperProcessor::CreateTransformToWorkingColorSpace(HLGSettings);
+
+	bool bSuccess = true;
+
+	for (float V : TestValues)
+	{
+		FLinearColor TestColor(V, V, V, 1.0f);
+		FLinearColor Expected = UE::Color::Decode(EEncoding::HLG, TestColor);
+
+		FLinearColor Actual = TestColor;
+		Processor.TransformColor(Actual);
+
+		const float Tolerance = FMath::Max(UE_KINDA_SMALL_NUMBER, UE_KINDA_SMALL_NUMBER * 0.5f * (Actual.R + Expected.R));
+
+		if (!Actual.Equals(Expected, Tolerance))
+		{
+			AddError(FString::Printf(TEXT("HLG decode mismatch at value %.4f: expected %s, got %s."), V, *Expected.ToString(), *Actual.ToString()), 1);
+			bSuccess = false;
+		}
+	}
+
+	return bSuccess;
+}
+
 #endif
\ No newline at end of file
diff --git a/Engine/Source/Runtime/RenderCore/Private/MediaShaders.cpp b/Engine/Source/Runtime/RenderCore/Private/MediaShaders.cpp
--- a/Engine/Source/Runtime/RenderCore/Private/MediaShaders.cpp
+++ b/Engine/Source/Runtime/RenderCore/Private/MediaShaders.cpp
@@ -8,6 +8,29 @@
 #include "ColorManagement/ColorManagementDefines.h"
 #include "ColorManagement/ColorSpace.h"
 
+static TAutoConsoleVariable<bool> CVarMediaHLGApplyOOTF(
+	TEXT("r.Media.HLG.ApplyOOTF"),
+	false,
+	TEXT("When enabled, applies the HLG OOTF (scene-to-display mapping) during media texture decoding.\n")
+	TEXT("This produces display-referred output matching PQ brightness levels.\n")
+	TEXT("0: Disabled (inverse OETF only, scene-referred output) (default)\n")
+	TEXT("1: Enabled (full BT.2100 EOTF with OOTF at 1000 cd/m2)"),
+	ECVF_RenderThreadSafe);
+
+// Shader-side sentinel for HLG with OOTF applied. Not a real EEncoding enum value —
+// deliberately set far above EEncoding::Max to avoid collisions with future enum additions.
+static constexpr uint32 UE_Color_EEncoding_HLG_OOTF = 128u;
+
+uint32 MediaShaders::GetMediaShaderEOTF_RenderThread(UE::Color::EEncoding Encoding)
+{
+	uint32 EOTF = static_cast<uint32>(Encoding);
+	if (Encoding == UE::Color::EEncoding::HLG && CVarMediaHLGApplyOOTF.GetValueOnRenderThread())
+	{
+		EOTF = UE_Color_EEncoding_HLG_OOTF;
+	}
+	return EOTF;
+}
+
 FMediaVertexDeclaration::FMediaVertexDeclaration() = default;
 FMediaVertexDeclaration::~FMediaVertexDeclaration() = default;
 
@@ -243,7 +266,7 @@ void FNV12ConvertPS::SetParameters(FRHIBatchedShaderParameters& BatchedParameter
 		UB.OutputWidth = OutputDimensions.X;
 		UB.SamplerB = TStaticSamplerState<SF_Bilinear>::GetRHI();
 		UB.SamplerP = TStaticSamplerState<SF_Point>::GetRHI();
-		UB.EOTF = static_cast<uint32>(Encoding);
+		UB.EOTF = MediaShaders::GetMediaShaderEOTF_RenderThread(Encoding);
 		UB.ToneMapMethod = static_cast<uint32>(ToneMapMethod);
 		UB.SRV_Y = SRV_Y;
 		UB.SRV_UV = SRV_UV;
@@ -284,7 +307,7 @@ void FNV12ConvertAsBytesPS::SetParameters(FRHIBatchedShaderParameters& BatchedPa
 		UB.OutputWidth = OutputDimensions.X;
 		UB.SamplerB = TStaticSamplerState<SF_Bilinear>::GetRHI();
 		UB.SamplerP = TStaticSamplerState<SF_Point>::GetRHI();
-		UB.EOTF = static_cast<uint32>(Encoding);
+		UB.EOTF = MediaShaders::GetMediaShaderEOTF_RenderThread(Encoding);
 		UB.ToneMapMethod = static_cast<uint32>(ToneMapMethod);
 		UB.Texture = NV12Texture;
 		UB.UVScale = FVector2f((float)OutputDimensions.X / (float)NV12Texture->GetSizeX(), (float)OutputDimensions.Y / (float)NV12Texture->GetSizeY());
@@ -357,7 +380,7 @@ void FP010ConvertPS::SetParameters(FRHIBatchedShaderParameters& BatchedParameter
 		UB.CSTransform = CSTransform;
 		UB.SamplerB = TStaticSamplerState<SF_Bilinear>::GetRHI();
 		UB.SamplerP = TStaticSamplerState<SF_Point>::GetRHI();
-		UB.EOTF = static_cast<uint32>(Encoding);
+		UB.EOTF = MediaShaders::GetMediaShaderEOTF_RenderThread(Encoding);
 		UB.ToneMapMethod = static_cast<uint32>(ToneMapMethod);
 		UB.SRV_Y = SRV_Y;
 		UB.SRV_UV = SRV_UV;
@@ -396,7 +419,7 @@ void FP010ConvertAsUINT16sPS::SetParameters(FRHIBatchedShaderParameters& Batched
 		UB.OutputWidth = OutputDimensions.X;
 		UB.SamplerB = TStaticSamplerState<SF_Bilinear>::GetRHI();
 		UB.SamplerP = TStaticSamplerState<SF_Point>::GetRHI();
-		UB.EOTF = static_cast<uint32>(Encoding);
+		UB.EOTF = MediaShaders::GetMediaShaderEOTF_RenderThread(Encoding);
 		UB.ToneMapMethod = static_cast<uint32>(ToneMapMethod);
 		UB.Texture = NV12Texture;
 		UB.UVScale = FVector2f((float)OutputDimensions.X / (float)TexDim.X, (float)OutputDimensions.Y / (float)TexDim.Y);
@@ -435,7 +458,7 @@ void FP010_2101010ConvertPS::SetParameters(FRHIBatchedShaderParameters& BatchedP
 		UB.ColorTransform = ColorTransform;
 		UB.CSTransform = CSTransform;
 		UB.SamplerP = TStaticSamplerState<SF_Point>::GetRHI();
-		UB.EOTF = static_cast<uint32>(Encoding);
+		UB.EOTF = MediaShaders::GetMediaShaderEOTF_RenderThread(Encoding);
 		UB.ToneMapMethod = static_cast<uint32>(ToneMapMethod);
 		UB.SRV_Y = SRV_Y;
 		UB.SRV_U = SRV_U;
@@ -463,7 +486,7 @@ SHADER_PARAMETER(FVector2f, UVScale)
 SHADER_PARAMETER_TEXTURE(Texture2D, Texture)
 SHADER_PARAMETER_SAMPLER(SamplerState, Sampler)
 SHADER_PARAMETER(FMatrix44f, CSTransform)
-SHADER_PARAMETER(uint32, EOTF)				// 0 = linear, 1=sRGB, 2=ST2084
+SHADER_PARAMETER(uint32, EOTF)				// UE::Color::EEncoding value (see ColorManagementDefines.h). 128=HLG with OOTF.
 SHADER_PARAMETER(uint32, ToneMapMethod)
 END_GLOBAL_SHADER_PARAMETER_STRUCT()
 
@@ -478,7 +501,7 @@ void FRGBConvertPS::SetParameters(FRHIBatchedShaderParameters& BatchedParameters
 		UB.Sampler = TStaticSamplerState<SF_Bilinear>::GetRHI();
 		UB.Texture = RGBTexture;
 		UB.UVScale = FVector2f((float)OutputDimensions.X / (float)RGBTexture->GetSizeX(), (float)OutputDimensions.Y / (float)RGBTexture->GetSizeY());
-		UB.EOTF = static_cast<uint32>(Encoding);
+		UB.EOTF = MediaShaders::GetMediaShaderEOTF_RenderThread(Encoding);
 		UB.ToneMapMethod = static_cast<uint32>(ToneMapMethod);
 		UB.CSTransform = CSTransform;
 	}
@@ -496,7 +519,7 @@ SHADER_PARAMETER(FVector2f, UVScale)
 SHADER_PARAMETER_TEXTURE(Texture2D, Texture)
 SHADER_PARAMETER_SAMPLER(SamplerState, Sampler)
 SHADER_PARAMETER(FMatrix44f, CSTransform)
-SHADER_PARAMETER(uint32, EOTF)				// 0 = linear, 1=sRGB, 2=ST2084
+SHADER_PARAMETER(uint32, EOTF)				// UE::Color::EEncoding value (see ColorManagementDefines.h). 128=HLG with OOTF.
 SHADER_PARAMETER(uint32, ToneMapMethod)
 END_GLOBAL_SHADER_PARAMETER_STRUCT()
 
@@ -511,7 +534,7 @@ void FYCoCgConvertPS::SetParameters(FRHIBatchedShaderParameters& BatchedParamete
 		UB.Sampler = TStaticSamplerState<SF_Bilinear>::GetRHI();
 		UB.Texture = RGBTexture;
 		UB.UVScale = FVector2f((float)OutputDimensions.X / (float)RGBTexture->GetSizeX(), (float)OutputDimensions.Y / (float)RGBTexture->GetSizeY());
-		UB.EOTF = static_cast<uint32>(Encoding);
+		UB.EOTF = MediaShaders::GetMediaShaderEOTF_RenderThread(Encoding);
 		UB.ToneMapMethod = static_cast<uint32>(ToneMapMethod);
 		UB.CSTransform = CSTransform;
 	}
@@ -658,7 +681,7 @@ void FYUVv216ConvertPS::SetParameters(FRHIBatchedShaderParameters& BatchedParame
 	{
 		UB.ColorTransform = ColorTransform;
 		UB.CSTransform = CSTransform;
-		UB.EOTF = static_cast<uint32>(Encoding);
+		UB.EOTF = MediaShaders::GetMediaShaderEOTF_RenderThread(Encoding);
 		UB.ToneMapMethod = static_cast<uint32>(ToneMapMethod);
 		UB.IsCbY0CrY1 = bIsCbY0CrY1;
 		UB.IsARGBFmt = bIsARGBFmt;
@@ -699,7 +722,7 @@ void FYUVv210ConvertPS::SetParameters(FRHIBatchedShaderParameters& BatchedParame
 	{
 		UB.ColorTransform = ColorTransform;
 		UB.CSTransform = CSTransform;
-		UB.EOTF = static_cast<uint32>(Encoding);
+		UB.EOTF = MediaShaders::GetMediaShaderEOTF_RenderThread(Encoding);
 		UB.ToneMapMethod = static_cast<uint32>(ToneMapMethod);
 		UB.IsCbY0CrY1 = bIsCbY0CrY1;
 		UB.OutputDimX = (float)OutputDimensions.X;
@@ -720,7 +743,7 @@ SHADER_PARAMETER(FMatrix44f, ColorTransform)
 SHADER_PARAMETER(FMatrix44f, CSTransform)
 SHADER_PARAMETER_SRV(Texture2D, SRV_Y)
 SHADER_PARAMETER_SAMPLER(SamplerState, BilinearClampedSamplerUV)
-SHADER_PARAMETER(uint32, EOTF)				// 0 = linear, 1=sRGB, 2=ST2084
+SHADER_PARAMETER(uint32, EOTF)				// UE::Color::EEncoding value (see ColorManagementDefines.h). 128=HLG with OOTF.
 SHADER_PARAMETER(uint32, ToneMapMethod)
 SHADER_PARAMETER(uint32, bIsARGBFmt)
 END_GLOBAL_SHADER_PARAMETER_STRUCT()
@@ -737,7 +760,7 @@ void FYUVY416ConvertPS::SetParameters(FRHIBatchedShaderParameters& BatchedParame
 		UB.BilinearClampedSamplerUV = TStaticSamplerState<SF_Bilinear, AM_Clamp, AM_Clamp, AM_Clamp>::GetRHI();
 		UB.ColorTransform = ColorTransform;
 		UB.CSTransform = CSTransform;
-		UB.EOTF = static_cast<uint32>(Encoding);
+		UB.EOTF = MediaShaders::GetMediaShaderEOTF_RenderThread(Encoding);
 		UB.ToneMapMethod = static_cast<uint32>(ToneMapMethod);
 		UB.bIsARGBFmt = bIsARGBFmt;
 	}
@@ -754,7 +777,7 @@ SHADER_PARAMETER(FMatrix44f, CSTransform)
 SHADER_PARAMETER(uint32, OutputWidth)
 SHADER_PARAMETER(uint32, OutputHeight)
 SHADER_PARAMETER_SRV(Texture2D<uint4>, Texture)
-SHADER_PARAMETER(uint32, EOTF)				// 0 = linear, 1=sRGB, 2=ST2084
+SHADER_PARAMETER(uint32, EOTF)				// UE::Color::EEncoding value (see ColorManagementDefines.h). 128=HLG with OOTF.
 SHADER_PARAMETER(uint32, ToneMapMethod)
 END_GLOBAL_SHADER_PARAMETER_STRUCT()
 
@@ -770,7 +793,7 @@ void FARGB16BigConvertPS::SetParameters(FRHIBatchedShaderParameters& BatchedPara
 		UB.OutputWidth = OutputDimensions.X;
 		UB.OutputHeight = OutputDimensions.Y;
 		UB.CSTransform = CSTransform;
-		UB.EOTF = static_cast<uint32>(Encoding);
+		UB.EOTF = MediaShaders::GetMediaShaderEOTF_RenderThread(Encoding);
 		UB.ToneMapMethod = static_cast<uint32>(ToneMapMethod);
 	}
 
@@ -990,7 +1013,7 @@ void FVYUConvertPS::SetParameters(FRHIBatchedShaderParameters& BatchedParameters
 	{
 		UB.ColorTransform = ColorTransform;
 		UB.CSTransform = CSTransform;
-		UB.EOTF = static_cast<uint32>(Encoding);
+		UB.EOTF = MediaShaders::GetMediaShaderEOTF_RenderThread(Encoding);
 		UB.ToneMapMethod = static_cast<uint32>(ToneMapMethod);
 		UB.UVRotation = UVRotation;
 		UB.UVRotOffsetAndScale = FLinearColor(UVRotOffset.R, UVRotOffset.G, (float)OutputDimensions.X / (float)VYUTexture->GetSizeX(), (float)OutputDimensions.Y / (float)VYUTexture->GetSizeY());
diff --git a/Engine/Source/Runtime/RenderCore/Public/MediaShaders.h b/Engine/Source/Runtime/RenderCore/Public/MediaShaders.h
--- a/Engine/Source/Runtime/RenderCore/Public/MediaShaders.h
+++ b/Engine/Source/Runtime/RenderCore/Public/MediaShaders.h
@@ -84,6 +84,13 @@ namespace MediaShaders
 
 	/** Combine color transform matrix with yuv offset in a single matrix */
 	RENDERCORE_API FMatrix CombineColorTransformAndOffset(const FMatrix& InMatrix, const FVector& InYUVOffset);
+
+	/**
+	 * Returns the shader-side EOTF value for a given encoding.
+	 * Accounts for runtime CVars (e.g. r.Media.HLG.ApplyOOTF) that remap encodings to
+	 * shader-only variants. Must be called on the render thread.
+	 */
+	RENDERCORE_API uint32 GetMediaShaderEOTF_RenderThread(UE::Color::EEncoding Encoding);
 }
 
 
