Removing the RootConstantsParam parameter from FDispatchShaderBundleCS since it hasn't been needed since 2024.
All remaining arguments should be optional now due to a new way this shader is used.

#Synced-CL 52273107

diff --git a/Engine/Source/Runtime/RenderCore/Public/ShaderBundles.h b/Engine/Source/Runtime/RenderCore/Public/ShaderBundles.h
--- a/Engine/Source/Runtime/RenderCore/Public/ShaderBundles.h
+++ b/Engine/Source/Runtime/RenderCore/Public/ShaderBundles.h
@@ -20,17 +20,13 @@ public:
 	FDispatchShaderBundleCS(const ShaderMetaType::CompiledShaderInitializerType& Initializer)
 		: FGlobalShader(Initializer)
 	{
-		// Platforms with support for root constants will not have a bind point for this parameter
-		RootConstantsParam.Bind(Initializer.ParameterMap, TEXT("PassData"), SPF_Optional);
-
-		RecordArgBufferParam.Bind(Initializer.ParameterMap, TEXT("RecordArgBuffer"), SPF_Mandatory);
+		RecordArgBufferParam.Bind(Initializer.ParameterMap, TEXT("RecordArgBuffer"), SPF_Optional);
 		RecordDataBufferParam.Bind(Initializer.ParameterMap, TEXT("RecordDataBuffer"), SPF_Optional);
 		RWExecutionBufferParam.Bind(Initializer.ParameterMap, TEXT("RWExecutionBuffer"), SPF_Optional);
 	}
 
 	static const uint32 ThreadGroupSizeX = 64;
 
-	LAYOUT_FIELD(FShaderParameter, RootConstantsParam);
 	LAYOUT_FIELD(FShaderResourceParameter, RecordArgBufferParam);
 	LAYOUT_FIELD(FShaderResourceParameter, RecordDataBufferParam);
 	LAYOUT_FIELD(FShaderResourceParameter, RWExecutionBufferParam);
