Add support for reserved buffers in DumpGPU.

These buffers use the committed size for dumping instead of the reserved size. The reserved size is added as separate metadata. Also changed size handling to use unsigned integers to avoid overflows for buffers >= 2 GB (which is the case for rerseved GPUScene resources).

#rnx

#Synced-CL 52126719

diff --git a/Engine/Source/Runtime/RHI/Private/RHIStrings.cpp b/Engine/Source/Runtime/RHI/Private/RHIStrings.cpp
--- a/Engine/Source/Runtime/RHI/Private/RHIStrings.cpp
+++ b/Engine/Source/Runtime/RHI/Private/RHIStrings.cpp
@@ -668,6 +668,8 @@ const TCHAR* GetBufferUsageFlagString(EBufferUsageFlags BufferUsage)
 		return TEXT("IndexBuffer");
 	case EBufferUsageFlags::StructuredBuffer:
 		return TEXT("StructuredBuffer");
+	case EBufferUsageFlags::ReservedResource:
+		return TEXT("ReservedResource");
 	}
 	return TEXT("");
 }
diff --git a/Engine/Source/Runtime/RenderCore/Private/DumpGPU.cpp b/Engine/Source/Runtime/RenderCore/Private/DumpGPU.cpp
--- a/Engine/Source/Runtime/RenderCore/Private/DumpGPU.cpp
+++ b/Engine/Source/Runtime/RenderCore/Private/DumpGPU.cpp
@@ -775,7 +775,7 @@ public:
 		return JsonObject;
 	}
 
-	TSharedPtr<FJsonObject> ToJson(const FString& UniqueResourceName, const TCHAR* Name, const FRDGBufferDesc& Desc, int32 ResourceByteSize)
+	TSharedPtr<FJsonObject> ToJson(const FString& UniqueResourceName, const TCHAR* Name, const FRDGBufferDesc& Desc, uint32 ResourceByteSize, uint32 ReservedSize)
 	{
 		TSharedPtr<FJsonObject> JsonObject = MakeShareable(new FJsonObject);
 		JsonObject->SetStringField(TEXT("Name"), Name);
@@ -786,6 +786,11 @@ public:
 		JsonObject->SetNumberField(TEXT("NumElements"), Desc.NumElements);
 		JsonObject->SetStringField(TEXT("Metadata"), *PtrToString(Desc.Metadata));
 
+		if (EnumHasAnyFlags(Desc.Usage, EBufferUsageFlags::ReservedResource))
+		{
+			JsonObject->SetNumberField(TEXT("ReservedSize"), ReservedSize);
+		}
+
 		{
 			TArray<TSharedPtr<FJsonValue>> UsageNames;
 			for (uint64 BitId = 0; BitId < uint64(8 * sizeof(EBufferUsageFlags)); BitId++)
@@ -1027,7 +1032,7 @@ public:
 
 	struct FStagingBufferPoolEntry : public FStagingPoolEntry
 	{
-		int32 ByteSize = 0;
+		uint32 ByteSize = 0;
 		FStagingBufferRHIRef StagingBuffer;
 	};
 
@@ -1130,7 +1135,7 @@ public:
 
 	FStagingBufferPoolEntry& ReuseStagingBuffer(
 		FRHICommandListImmediate& RHICmdList,
-		int32 ByteSize)
+		uint32 ByteSize)
 	{
 		check(IsInRenderingThread());
 		check(bStream);
@@ -1620,7 +1625,19 @@ public:
 		}
 
 		const FRDGBufferDesc& Desc = Buffer->Desc;
-		const int32 ByteSize = Desc.GetSize();
+		uint32 ByteSize = Desc.GetSize();
+		uint32 ReservedSize = 0;
+
+		// For reserved resources, only dump the committed portion. Currently reserved buffers are only used for no-copy
+		// grow/shrink operations, so the committed region is contiguous and starts at offset 0. We don't need to add
+		// handling for true sparse resources.
+		if (EnumHasAnyFlags(Desc.Usage, EBufferUsageFlags::ReservedResource))
+		{
+			ReservedSize = ByteSize;
+			// GetCommittedSize returns a 64-bit value, but this can never exceed the reserved size returned by Desc.GetSize(),
+			// which is currently a 32-bit value. Therefore it's safe to truncate here.
+			ByteSize = static_cast<uint32>(Buffer->GetCommittedSize());
+		}
 
 		bool bDumpResourceInfos = false;
 		bool bDumpResourceBinary = bIsOutputResource;
@@ -1647,7 +1664,7 @@ public:
 		// Dump the information of the buffer to json file.
 		if (bDumpResourceInfos)
 		{
-			TSharedPtr<FJsonObject> JsonObject = ToJson(UniqueResourceName, Buffer->Name, Desc, ByteSize);
+			TSharedPtr<FJsonObject> JsonObject = ToJson(UniqueResourceName, Buffer->Name, Desc, ByteSize, ReservedSize);
 			DumpJsonToFile(JsonObject, FString(FRDGResourceDumpContext::kBaseDir) / TEXT("ResourceDescs.json"), FILEWRITE_Append);
 
 			if (Desc.Metadata && DumpBufferMode == 2)
@@ -1659,7 +1676,13 @@ public:
 		// Dump the resource's binary to a .bin file.
 		if (bDumpResourceBinary && DumpBufferMode == 2)
 		{
-			const int32 StagingResourceByteSize = bStream ? ByteSize : FMath::Min(ByteSize, GDumpMaxStagingSize.GetValueOnRenderThread() * 1024 * 1024);
+			if (ByteSize == 0)
+			{
+				// Reserved resource with nothing committed yet, nothing to dump.
+				return;
+			}
+
+			const uint32 StagingResourceByteSize = bStream ? ByteSize : FMath::Min(ByteSize, GDumpMaxStagingSize.GetValueOnRenderThread() * 1024U * 1024U);
 
 			FString DumpFilePath = kResourcesDir / FString::Printf(TEXT("%s.v%s.bin"), *UniqueResourceName, *PtrToString(bIsOutputResource ? Pass : nullptr));
 
@@ -1728,9 +1751,9 @@ public:
 						Ar = TUniquePtr<FArchive>(IFileManager::Get().CreateFileWriter(*FullPath, /* WriteFlags = */ 0));
 					}
 
-					for (int32 Offset = 0; Offset < ByteSize; Offset += StagingResourceByteSize)
+					for (uint32 Offset = 0; Offset < ByteSize; Offset += StagingResourceByteSize)
 					{
-						int32 CopyByteSize = FMath::Min(StagingResourceByteSize, ByteSize - Offset);
+						uint32 CopyByteSize = FMath::Min(StagingResourceByteSize, ByteSize - Offset);
 
 						// Transfer memory GPU -> CPU
 						{
diff --git a/Engine/Source/Runtime/RenderCore/Public/RenderGraphResources.h b/Engine/Source/Runtime/RenderCore/Public/RenderGraphResources.h
--- a/Engine/Source/Runtime/RenderCore/Public/RenderGraphResources.h
+++ b/Engine/Source/Runtime/RenderCore/Public/RenderGraphResources.h
@@ -1351,6 +1351,11 @@ public:
 		return Desc.BytesPerElement;
 	}
 
+	inline uint64 GetCommittedSize() const
+	{
+		return PooledBuffer ? PooledBuffer->GetCommittedSize() : 0;
+	}
+
 	bool IsCulled() const
 	{
 		return ReferenceCount == 0 && PendingCommitSize == 0;
