[shaders] add cook-time hash collision detection; safety measure due to recent change to 64-bit hashes. the likelihood of collisions is exceedingly low, but this gives us assurance that any possible shader mismatch issues are not caused by hash collisions.

#Synced-CL 52312814

diff --git a/Engine/Source/Runtime/RenderCore/Private/ShaderLibrary/EditorShaderCodeArchive.cpp b/Engine/Source/Runtime/RenderCore/Private/ShaderLibrary/EditorShaderCodeArchive.cpp
--- a/Engine/Source/Runtime/RenderCore/Private/ShaderLibrary/EditorShaderCodeArchive.cpp
+++ b/Engine/Source/Runtime/RenderCore/Private/ShaderLibrary/EditorShaderCodeArchive.cpp
@@ -15,6 +15,7 @@
 #include "Interfaces/ITargetPlatformManagerModule.h"
 #include "Internationalization/Regex.h"
 #include "Math/UnitConversion.h"
+#include "Misc/Compression.h"
 #include "RHIStrings.h"
 #include "PipelineCacheUtilities.h"
 #include "Serialization/CompactBinarySerialization.h"
@@ -29,6 +30,47 @@
 
 using namespace UE::ShaderLibrary::Private;
 
+// Compares only the shader code portion of two (potentially compressed) shader buffers,
+// excluding optional data attachments which may legitimately differ between shader types
+// that share the same bytecode. Returns true if the shader code portions are equal.
+static bool ShaderCodeBytesAreEqual(
+	FMemoryView BufferA, uint32 UncompressedSizeA,
+	FMemoryView BufferB, uint32 UncompressedSizeB)
+{
+	FName CompressionFormat = GetShaderCompressionFormat();
+
+	auto DecompressIfNeeded = [CompressionFormat](FMemoryView CompressedView, uint32 UncompressedSize, TArray<uint8>& OutDecompressed) -> TConstArrayView<uint8>
+	{
+		bool bIsCompressed = (CompressionFormat != NAME_None) && (CompressedView.GetSize() != UncompressedSize);
+		if (bIsCompressed)
+		{
+			OutDecompressed.SetNumUninitialized(UncompressedSize);
+			bool bDecompressSuccess = FCompression::UncompressMemory(CompressionFormat,
+				OutDecompressed.GetData(), UncompressedSize,
+				CompressedView.GetData(), CompressedView.GetSize());
+			if (!bDecompressSuccess)
+			{
+				UE_LOGF(LogShaderLibrary, Fatal, "Failed to decompress shader code buffer for collision check, aborting");
+			}
+			return TConstArrayView<uint8>(OutDecompressed.GetData(), UncompressedSize);
+		}
+		return TConstArrayView<uint8>(static_cast<const uint8*>(CompressedView.GetData()), static_cast<int32>(CompressedView.GetSize()));
+	};
+
+	TArray<uint8> DecompressedA, DecompressedB;
+	TConstArrayView<uint8> ViewA = DecompressIfNeeded(BufferA, UncompressedSizeA, DecompressedA);
+	TConstArrayView<uint8> ViewB = DecompressIfNeeded(BufferB, UncompressedSizeB, DecompressedB);
+
+	FShaderCodeReader ReaderA(ViewA);
+	int32 ShaderCodeSizeA = ReaderA.GetShaderCodeSize();
+	FShaderCodeReader ReaderB(ViewB);
+	int32 ShaderCodeSizeB = ReaderB.GetShaderCodeSize();
+
+	FMemoryView CodeA(ViewA.GetData(), FMath::Min(ShaderCodeSizeA, ViewA.Num()));
+	FMemoryView CodeB(ViewB.GetData(), FMath::Min(ShaderCodeSizeB, ViewB.Num()));
+	return CodeA.EqualBytes(CodeB);
+}
+
 // Returns the shader format interface for the specified format and platform name.
 // This is used to determine whether native shader libraries are supported in FEditorShaderCodeArchive::SupportsShaderArchives().
 static const IShaderFormat* GetShaderFormatFromName(FName InFormatAndPlatformName)
@@ -114,9 +156,58 @@ int32 FEditorShaderCodeArchive::AddShaderCode(const FShaderMapResourceCode* Code
 				ShaderCode.Add(CodeBuffer);
 				check(ShaderCode.Num() == SerializedShaders.ShaderEntries.Num());
 			}
+			else
+			{
+				// Verify the existing shader with the same hash has identical shader code.
+				// The full code buffer includes optional data (attachments) after ShaderCodeSize
+				// that may legitimately differ between shader types sharing the same bytecode.
+				// Fast path: skip if full buffers match. Slow path: decompress and compare only
+				// the shader code portion. Guard against empty existing code (after CopyToArchiveAndClear).
+				FMemoryView ExistingView = ShaderCode[ShaderIndex].GetView();
+				if (ExistingView.GetSize() > 0 && !ExistingView.EqualBytes(Code->ShaderCodeResources[i].GetCodeBuffer().GetView()))
+				{
+					const FShaderCodeEntry& ExistingEntry = SerializedShaders.ShaderEntries[ShaderIndex];
+					FSharedBuffer IncomingCodeBuffer = Code->ShaderCodeResources[i].GetCodeBuffer();
+					if (!ShaderCodeBytesAreEqual(
+						ExistingView, ExistingEntry.GetUncompressedSize(),
+						IncomingCodeBuffer.GetView(), Code->ShaderCodeResources[i].GetUncompressedSize()))
+					{
+						UE_LOGF(LogShaderLibrary, Fatal,
+							"Shader hash collision: two different shaders produced hash %ls",
+							*Code->ShaderHashes[i].ToString());
+					}
+				}
+			}
 			SerializedShaders.ShaderIndices[ShaderMapEntry.ShaderIndicesOffset + i] = ShaderIndex;
 		}
 	}
+	else
+	{
+		// Verify the existing shader map with the same ResourceHash has identical shader content.
+		// ResourceHash is a simple sum of shader hashes, so different shader sets could theoretically
+		// produce the same sum, though this is exceedingly unlikely (this is just a safety measure).
+		// The element-wise comparison below is valid because ShaderIndices are stored in the same
+		// order as Code->ShaderHashes (see the insertion path above). Do not sort either side.
+		const FShaderMapEntry& ExistingEntry = SerializedShaders.ShaderMapEntries[ShaderMapIndex];
+		if (ExistingEntry.NumShaders != static_cast<uint32>(NumShaders))
+		{
+			UE_LOGF(LogShaderLibrary, Fatal,
+				"Shader map hash collision: ResourceHash %ls maps to existing shader map with %u shaders, but incoming has %d shaders",
+				*Code->ResourceHash.ToString(), ExistingEntry.NumShaders, NumShaders);
+		}
+		else
+		{
+			for (uint32 ShaderIndex = 0; ShaderIndex < ExistingEntry.NumShaders; ++ShaderIndex)
+			{
+				if (SerializedShaders.ShaderHashes[SerializedShaders.ShaderIndices[ExistingEntry.ShaderIndicesOffset + ShaderIndex]] != Code->ShaderHashes[ShaderIndex])
+				{
+					UE_LOGF(LogShaderLibrary, Fatal,
+						"Shader map hash collision: ResourceHash %ls maps to existing shader map with matching shader count (%u), but shader hashes differ at index %u",
+						*Code->ResourceHash.ToString(), ExistingEntry.NumShaders, ShaderIndex);
+				}
+			}
+		}
+	}
 
 	for (int32 i = 0; i < NumShaders; ++i)
 	{
@@ -332,6 +423,26 @@ bool FEditorShaderCodeArchive::AppendFromArchive(
 
 		if (!bShaderIsNew)
 		{
+			// Verify the existing shader with the same hash has identical bytecode.
+			// Guard against empty target (possible after CopyToArchiveAndClear) and missing source.
+			FMemoryView ExistingView = TargetShaderCodes[TargetShaderIndex].GetView();
+			if (ExistingView.GetSize() > 0 && SourceShaderEntry.Offset != INDEX_NONE
+				&& SourceShaderEntry.Offset + SourceShaderEntry.Size <= static_cast<uint64>(SourceFlatShaderCode.Num()))
+			{
+				FMemoryView IncomingView(SourceFlatShaderCode.GetData() + SourceShaderEntry.Offset, SourceShaderEntry.Size);
+				if (!ExistingView.EqualBytes(IncomingView))
+				{
+					const FShaderCodeEntry& TargetEntry = TargetArchive.ShaderEntries[TargetShaderIndex];
+					if (!ShaderCodeBytesAreEqual(
+						ExistingView, TargetEntry.GetUncompressedSize(),
+						IncomingView, SourceShaderEntry.GetUncompressedSize()))
+					{
+						UE_LOGF(LogShaderLibrary, Fatal,
+							"Shader hash collision: two different shaders produced hash %ls",
+							*SourceShaderHash.ToString());
+					}
+				}
+			}
 			continue;
 		}
 		check(TargetShaderIndex == TargetArchive.ShaderEntries.Num() - 1 && TargetShaderCodes.Num() == TargetArchive.ShaderEntries.Num() - 1);
@@ -378,7 +489,28 @@ bool FEditorShaderCodeArchive::AppendFromArchive(
 		MarkShaderMapDirty(TargetShaderMapIndex);
 		if (!bIsNewShaderMap)
 		{
-			// ShaderMap has already been loaded in this process, e.g. from another CookWorker loading the same Material
+			// Verify the existing shader map with the same hash has identical shader content.
+			const FShaderMapEntry& TargetEntry = TargetArchive.ShaderMapEntries[TargetShaderMapIndex];
+			const FShaderMapEntry& SourceEntry = SourceArchive.ShaderMapEntries[SourceShaderMapIndex];
+			if (TargetEntry.NumShaders != SourceEntry.NumShaders)
+			{
+				UE_LOGF(LogShaderLibrary, Fatal,
+					"Shader map hash collision: ResourceHash %ls maps to existing shader map with %u shaders, but incoming has %u shaders",
+					*SourceShaderMapHash.ToString(), TargetEntry.NumShaders, SourceEntry.NumShaders);
+			}
+			else
+			{
+				for (uint32 ShaderIndex = 0; ShaderIndex < TargetEntry.NumShaders; ++ShaderIndex)
+				{
+					if (TargetArchive.ShaderHashes[TargetArchive.ShaderIndices[TargetEntry.ShaderIndicesOffset + ShaderIndex]]
+						!= SourceArchive.ShaderHashes[SourceArchive.ShaderIndices[SourceEntry.ShaderIndicesOffset + ShaderIndex]])
+					{
+						UE_LOGF(LogShaderLibrary, Fatal,
+							"Shader map hash collision: ResourceHash %ls maps to existing shader map with matching shader count (%u), but shader hashes differ at index %u",
+							*SourceShaderMapHash.ToString(), TargetEntry.NumShaders, ShaderIndex);
+					}
+				}
+			}
 			continue;
 		}
 
@@ -520,6 +652,25 @@ int32 FEditorShaderCodeArchive::AddShaderCode(int32 OtherShaderMapIndex, const F
 				ShaderCode.Add(OtherShaderCodeEntry);
 				check(ShaderCode.Num() == SerializedShaders.ShaderEntries.Num());
 			}
+			else
+			{
+				// Verify the existing shader with the same hash has identical shader code.
+				// Guard against empty existing code (possible after CopyToArchiveAndClear).
+				FMemoryView ExistingView = ShaderCode[ShaderIndex].GetView();
+				if (ExistingView.GetSize() > 0 && !ExistingView.EqualBytes(OtherArchive.ShaderCode[OtherShaderIndex].GetView()))
+				{
+					const FShaderCodeEntry& ExistingEntry = SerializedShaders.ShaderEntries[ShaderIndex];
+					const FShaderCodeEntry& OtherEntry = OtherShaders.ShaderEntries[OtherShaderIndex];
+					if (!ShaderCodeBytesAreEqual(
+						ExistingView, ExistingEntry.GetUncompressedSize(),
+						OtherArchive.ShaderCode[OtherShaderIndex].GetView(), OtherEntry.GetUncompressedSize()))
+					{
+						UE_LOGF(LogShaderLibrary, Fatal,
+							"Shader hash collision: two different shaders produced hash %ls",
+							*OtherShaders.ShaderHashes[OtherShaderIndex].ToString());
+					}
+				}
+			}
 			SerializedShaders.ShaderIndices[ShaderMapEntry.ShaderIndicesOffset + i] = ShaderIndex;
 		}
 	}
