Fixing PVS warnings
#Synced-CL 52357860

diff --git a/Engine/Source/Developer/DerivedDataCache/Public/DerivedDataCacheUsageStats.h b/Engine/Source/Developer/DerivedDataCache/Public/DerivedDataCacheUsageStats.h
--- a/Engine/Source/Developer/DerivedDataCache/Public/DerivedDataCacheUsageStats.h
+++ b/Engine/Source/Developer/DerivedDataCache/Public/DerivedDataCacheUsageStats.h
@@ -192,7 +192,7 @@ public:
 			}
 		}
 		else
-		{
+		{ //-V523
 			for (const auto& KVP : UsageStats)
 			{
 				COOK_STAT(UsageStatsMap.Add(FString::Printf(TEXT("%s: %s.%s"), *GraphPath, *GetCacheName(), *KVP.Key), KVP.Value));
diff --git a/Engine/Source/Runtime/RenderCore/Private/RenderGraphPrivate.cpp b/Engine/Source/Runtime/RenderCore/Private/RenderGraphPrivate.cpp
--- a/Engine/Source/Runtime/RenderCore/Private/RenderGraphPrivate.cpp
+++ b/Engine/Source/Runtime/RenderCore/Private/RenderGraphPrivate.cpp
@@ -674,7 +674,7 @@ FRDGScopeState::FState::FState(bool bInImmediate, bool bInParallelExecute)
 		case ERDGScopeMode::TopLevelOnly:
 		case ERDGScopeMode::AllEvents:
 			// Override to a higher level in some cases
-			if (bRDGChannelEnabled != 0)
+			if (bRDGChannelEnabled != 0) //-V547
 			{
 				LocalScopeMode = ERDGScopeMode::AllEventsAndPassNames;
 			}
