[shaders] fix tsan error - niagara is polling bReleased/bFinalized bitfield members of job struct from the gamethread, while worker thread is submitting job/calculating input hash and setting the bInputHashSet member. fixed by making bInputHashSet its own boolean flag (the bitfield was overflowing anyways - 9 1-bit members of a uint8, so this doesn't change memory footprint of the struct)

#Synced-CL 52312814

diff --git a/Engine/Source/Runtime/RenderCore/Public/ShaderCompilerJobTypes.h b/Engine/Source/Runtime/RenderCore/Public/ShaderCompilerJobTypes.h
--- a/Engine/Source/Runtime/RenderCore/Public/ShaderCompilerJobTypes.h
+++ b/Engine/Source/Runtime/RenderCore/Public/ShaderCompilerJobTypes.h
@@ -188,8 +188,6 @@ public:
 		After a job is bFinalized it will be bReleased when ReleaseJob() is invoked, which means that the shader compile thread
 		is no longer processing the job; which is useful for non standard job handling (Niagara as an example). */
 	uint8 bReleased : 1;
-	/** Whether we hashed the inputs */
-	uint8 bInputHashSet : 1;
 	/** Whether or not we are a default material. */
 	uint8 bIsDefaultMaterial : 1;
 	/** Whether or not we are a global shader. */
@@ -198,6 +196,11 @@ public:
 	uint8 bBypassCache : 1;
 	/** Whether this job has been transferred after submitted to a remote worker. If this is true, don't register this as an assigned job again. */
 	uint8 bTransferred : 1;
+
+
+	/** Whether inputs have been hash and the below input hash has been set (note: intentionally separate from above bitfield to avoid data race due
+		to code accessing other bitfield members from game thread while submission task is running in parallel on a worker thread) */
+	uint8 bInputHashSet;
 	/** Hash of all the job inputs */
 	FShaderCompilerInputHash InputHash;
 
@@ -291,11 +294,11 @@ protected:
 		bSucceeded(false),
 		bErrorsAreLikelyToBeCode(false),
 		bReleased(false),
-		bInputHashSet(false),
 		bIsDefaultMaterial(false),
 		bIsGlobalShader(false),
 		bBypassCache(false),
-		bTransferred(false)
+		bTransferred(false),
+		bInputHashSet(false)
 	{
 		check(InPriority != EShaderCompileJobPriority::None);
 
