Metal / FIoStoreShaderCodeArchive / FIoStoreShaderLibraryProcessor - Handle memory mapped FIoStoreShaderCodeArchive shutdown & closeShader & per-library scoped shader group chunk IDs
- Extract Metal-specific shader lifecycle management into FMetalIoStoreShaderCodeArchive (new class), tracking created shaders and releasing Metal objects (MTLFunction/MTLLibrary) independently of FRHIShader destruction
- Assign unique chunk IDs per shader group scoped to their shader library, ensuring mmap regions are never opened in pak files that don't belong to the owning FRHIShaderLibrary, which is required to safely release all mmap regions before a library is closed
- Add ReleaseMetalObjects() to Metal shader types so Metal GPU resources can be freed before FRHIShader ref counts reach zero
- Change the shader library close callback from OnPakFileUnmounted to OnPakFileUnmounting so that OnCloseShaderCode is called before the pak file is fully unmounted, allowing all mmap regions to be released first.

#tests Local project + Preflight

#Synced-CL 52396837

diff --git a/Engine/Source/Developer/IoStoreUtilities/Private/IoStoreShaderLibraryProcessor.cpp b/Engine/Source/Developer/IoStoreUtilities/Private/IoStoreShaderLibraryProcessor.cpp
--- a/Engine/Source/Developer/IoStoreUtilities/Private/IoStoreShaderLibraryProcessor.cpp
+++ b/Engine/Source/Developer/IoStoreUtilities/Private/IoStoreShaderLibraryProcessor.cpp
@@ -16,6 +16,7 @@
 #include "Serialization/MemoryReader.h"
 #include "ShaderCodeArchive.h"
 #include "ShaderLibrary/EditorShaderCodeArchive.h"
+#include "ShaderLibrary/ShaderCodeArchiveInternal.h"
 
 #define IOSTORE_CPU_SCOPE(NAME) TRACE_CPUPROFILER_EVENT_SCOPE(IoStore##NAME);
 #define IOSTORE_CPU_SCOPE_DATA(NAME, DATA) TRACE_CPUPROFILER_EVENT_SCOPE(IoStore##NAME);
@@ -563,8 +564,12 @@ bool FIoStoreShaderLibraryProcessor::ConvertToIoStoreShaderLibrary(FContainerTar
 	// Is shader group compression for this container target file enabled?
 	const bool bCompressShaderGroup = ShouldCompressShaderGroups(LibraryNameInfo.ChunkName);
 
+	// Compute the library's own ChunkId once so it can be captured by the parallel lambda below.
+	// This is the same value that will be assigned to TargetFile.ChunkId after the ParallelFor.
+	const FIoChunkId LibraryChunkId = ShaderCodeArchive::GetShaderCodeArchiveChunkId(LibraryNameInfo.ChunkName, FName(LibraryNameInfo.GetFormatAndPlatformName()));
+
 	ParallelFor(CodeIoChunks.Num(),
-		[this, &IoStoreLibraryHeader, &SerializedShaders, bCompressShaderGroup](int32 GroupIndex)
+		[this, &IoStoreLibraryHeader, &SerializedShaders, bCompressShaderGroup, &LibraryChunkId](int32 GroupIndex)
 		{
 			FIoStoreShaderGroupEntry& Group = IoStoreLibraryHeader.ShaderGroupEntries[GroupIndex];
 
@@ -584,7 +589,13 @@ bool FIoStoreShaderLibraryProcessor::ConvertToIoStoreShaderLibrary(FContainerTar
 			// Write out IoChunk for current shader group
 			{
 				FShaderGroupIoChunk& OutCodeIoChunk = this->CodeIoChunks[GroupIndex];
-				OutCodeIoChunk.ChunkId = IoStoreLibraryHeader.ShaderGroupIoHashes[GroupIndex];
+				// Fuse the content-based group hash with the library's ChunkId so that the same
+				// shader group present in two different libraries produces distinct ChunkIds,
+				// preventing collisions in the IoStore container. Both the code chunk and the
+				// serialized header entry must use the same scoped id so the runtime can resolve them.
+				const FIoChunkId ScopedChunkId = ShaderCodeArchive::GetLibraryScopedShaderGroupChunkId(IoStoreLibraryHeader.ShaderGroupIoHashes[GroupIndex], LibraryChunkId);
+				OutCodeIoChunk.ChunkId = ScopedChunkId;
+				IoStoreLibraryHeader.ShaderGroupIoHashes[GroupIndex] = ScopedChunkId;
 				// This value is the load order factor for the group (the smaller, the more likely), that IoStore will use to sort the chunks.
 				// We calculate it as the smallest shader index in the group, which shouldn't be bad. Arguably a better way would be to get the lowest-numbered shadermap that references _any_ shader in the group, 
 				// but it's much, much slower to calculate and the resulting order would be likely the same/similar most of the time
diff --git a/Engine/Source/Runtime/Apple/MetalRHI/Private/MetalDynamicRHIModule.cpp b/Engine/Source/Runtime/Apple/MetalRHI/Private/MetalDynamicRHIModule.cpp
--- a/Engine/Source/Runtime/Apple/MetalRHI/Private/MetalDynamicRHIModule.cpp
+++ b/Engine/Source/Runtime/Apple/MetalRHI/Private/MetalDynamicRHIModule.cpp
@@ -10,6 +10,7 @@
 #include "DynamicRHI.h"
 #include "RHIValidation.h"
 #include "Modules/ModuleManager.h"
+#include "Shaders/MetalIoStoreShaderCodeArchive.h"
 
 
 //------------------------------------------------------------------------------
@@ -32,6 +33,11 @@ bool FMetalDynamicRHIModule::IsSupported()
 
 FDynamicRHI* FMetalDynamicRHIModule::CreateRHI(ERHIFeatureLevel::Type RequestedFeatureLevel)
 {
+	FIoStoreShaderCodeArchive::RegisterIoStoreShaderCodeArchiveFactory([](EShaderPlatform Platform, const FString& Name, FIoDispatcher& IoDispatcher)
+	{
+		return static_cast<FIoStoreShaderCodeArchive*>(new FMetalIoStoreShaderCodeArchive(Platform, Name, IoDispatcher));
+	});
+
 	LLM(MetalLLM::Initialise());
 	FDynamicRHI* FinalRHI =  new FMetalDynamicRHI(RequestedFeatureLevel);
 	
diff --git a/Engine/Source/Runtime/Apple/MetalRHI/Private/MetalPipeline.cpp b/Engine/Source/Runtime/Apple/MetalRHI/Private/MetalPipeline.cpp
--- a/Engine/Source/Runtime/Apple/MetalRHI/Private/MetalPipeline.cpp
+++ b/Engine/Source/Runtime/Apple/MetalRHI/Private/MetalPipeline.cpp
@@ -1120,6 +1120,9 @@ static FMetalShaderPipelinePtr CreateMTLRenderPipeline(FMetalDevice& Device, boo
 		}
 		
 #if METAL_DEBUG_OPTIONS
+		// Reset the pipeline descriptor to release MTLFunction/MTLLibrary references it retains.
+		Pipeline->RenderDesc.reset();
+
 		Pipeline->VertexSource = VertexShader->GetSourceCode();
         Pipeline->VertexSource->retain();
         
@@ -1221,8 +1224,11 @@ static FMetalShaderPipelinePtr CreateMTLRenderPipeline(FMetalDevice& Device, boo
 			delete Pipeline;
 			return nullptr;
 		}
-        
+
 #if METAL_DEBUG_OPTIONS
+		// Reset the pipeline descriptor to release MTLFunction/MTLLibrary references it retains.
+		Pipeline->MeshRenderDesc.reset();
+
         Pipeline->MeshSource = MeshShader ? MeshShader->GetSourceCode() : GeometryShader->GetSourceCode();
         Pipeline->ObjectSource = AmplificationShader ? AmplificationShader->GetSourceCode() : nil;
         Pipeline->FragmentSource = PixelShader ? PixelShader->GetSourceCode() : nil;
@@ -1299,8 +1305,11 @@ static FMetalShaderPipelinePtr CreateMTLRenderPipeline(FMetalDevice& Device, boo
 			delete Pipeline;
             return nil;
         }
-        
+
 #if METAL_DEBUG_OPTIONS
+		// Reset the pipeline descriptor to release MTLFunction/MTLLibrary references it retains.
+		Pipeline->MeshRenderDesc.reset();
+
         Pipeline->MeshSource = MeshShader->GetSourceCode();
         Pipeline->ObjectSource = AmplificationShader ? AmplificationShader->GetSourceCode() : nil;
         Pipeline->FragmentSource = PixelShader ? PixelShader->GetSourceCode() : nil;
diff --git a/Engine/Source/Runtime/Apple/MetalRHI/Private/Shaders/MetalCompiledShaderCache.h b/Engine/Source/Runtime/Apple/MetalRHI/Private/Shaders/MetalCompiledShaderCache.h
--- a/Engine/Source/Runtime/Apple/MetalRHI/Private/Shaders/MetalCompiledShaderCache.h
+++ b/Engine/Source/Runtime/Apple/MetalRHI/Private/Shaders/MetalCompiledShaderCache.h
@@ -7,50 +7,35 @@
 #pragma once
 
 #include "MetalCompiledShaderKey.h"
-#include "Misc/ScopeRWLock.h"
+#include "Containers/StripedMap.h"
 #include "MetalRHIPrivate.h"
 
+struct FCachedCompiledShader
+{
+	MTLFunctionPtr Function;
+	MTLLibraryPtr Library;
+};
+
 struct FMetalCompiledShaderCache
 {
 public:
-	FMetalCompiledShaderCache()
-	{
-		// VOID
-	}
-
-	~FMetalCompiledShaderCache()
+	FCachedCompiledShader Find(FMetalCompiledShaderKey const& Key)
 	{
-		// VOID
+		return Cache.FindRef(Key);
 	}
 
-	MTLFunctionPtr FindRef(FMetalCompiledShaderKey const& Key)
-	{
-		FRWScopeLock ScopedLock(Lock, SLT_ReadOnly);
-        MTLFunctionPtr Func = Cache.FindRef(Key);
-		return Func;
-	}
-
-	MTLLibraryPtr FindLibrary(MTLFunctionPtr Function)
+	void Add(FMetalCompiledShaderKey Key, MTLLibraryPtr Lib, MTLFunctionPtr Function)
 	{
-		FRWScopeLock ScopedLock(Lock, SLT_ReadOnly);
-        MTLLibraryPtr Lib = LibCache.FindRef(Function->functionType());
-		return Lib;
+		Cache.FindOrProduce(Key, [&]() { return FCachedCompiledShader{ Function, Lib }; });
 	}
 
-	void Add(FMetalCompiledShaderKey Key, MTLLibraryPtr Lib, MTLFunctionPtr Function)
+	void Remove(FMetalCompiledShaderKey const& Key)
 	{
-		FRWScopeLock ScopedLock(Lock, SLT_Write);
-		if (Cache.FindRef(Key).get() == nullptr)
-		{
-			Cache.Add(Key, Function);
-			LibCache.Add(Function->functionType(), Lib);
-		}
+		Cache.Remove(Key);
 	}
 
 private:
-	FRWLock Lock;
-	TMap<FMetalCompiledShaderKey, MTLFunctionPtr> Cache;
-	TMap<MTL::FunctionType, MTLLibraryPtr> LibCache;
+	TStripedMap<32, FMetalCompiledShaderKey, FCachedCompiledShader> Cache;
 };
 
 extern FMetalCompiledShaderCache& GetMetalCompiledShaderCache();
diff --git a/Engine/Source/Runtime/Apple/MetalRHI/Private/Shaders/MetalCompiledShaderKey.h b/Engine/Source/Runtime/Apple/MetalRHI/Private/Shaders/MetalCompiledShaderKey.h
--- a/Engine/Source/Runtime/Apple/MetalRHI/Private/Shaders/MetalCompiledShaderKey.h
+++ b/Engine/Source/Runtime/Apple/MetalRHI/Private/Shaders/MetalCompiledShaderKey.h
@@ -11,9 +11,10 @@
 
 struct FMetalCompiledShaderKey
 {
-	FMetalCompiledShaderKey(uint32 CodeSize, uint32 CodeCRC, bool bIsExtendedBytecode = false)
+	FMetalCompiledShaderKey(uint32 CodeSize, uint32 CodeCRC, bool bIsExtendedBytecode = false, uint64 InLibrarySeed = 0)
 	{
-		Hash = HashCombine(HashCombine(GetTypeHash(CodeSize), GetTypeHash(CodeCRC)), bIsExtendedBytecode);
+		// Prevent multiple FRHIShaderLibrary from sharing MTLLibrary & MTLFunction by adding the LibrarySeed inside the hash
+		Hash = HashCombine(HashCombine(HashCombine(GetTypeHash(CodeSize), GetTypeHash(CodeCRC)), bIsExtendedBytecode), GetTypeHash(InLibrarySeed));
 	}
 
 	friend bool operator ==(const FMetalCompiledShaderKey& A, const FMetalCompiledShaderKey& B)
diff --git a/Engine/Source/Runtime/Apple/MetalRHI/Private/Shaders/MetalIoStoreShaderCodeArchive.cpp b/Engine/Source/Runtime/Apple/MetalRHI/Private/Shaders/MetalIoStoreShaderCodeArchive.cpp
new file mode 100644
--- /dev/null
+++ b/Engine/Source/Runtime/Apple/MetalRHI/Private/Shaders/MetalIoStoreShaderCodeArchive.cpp
@@ -0,0 +1,110 @@
+// Copyright Epic Games, Inc. All Rights Reserved.
+
+#include "Shaders/MetalIoStoreShaderCodeArchive.h"
+#include "MetalRHIPrivate.h"
+#include "MetalShaderTypes.h"
+#if METAL_RHI_RAYTRACING
+#include "Shaders/Types/MetalRayShader.h"
+#endif
+
+FMetalIoStoreShaderCodeArchive::FMetalIoStoreShaderCodeArchive(EShaderPlatform InPlatform, const FString& InLibraryName, FIoDispatcher& InIoDispatcher)
+	: FIoStoreShaderCodeArchive(InPlatform, InLibraryName, InIoDispatcher)
+{
+}
+
+void FMetalIoStoreShaderCodeArchive::ReleaseMetalObjectsForShader(FRHIShader* Shader)
+{
+	switch (Shader->GetFrequency())
+	{
+	case SF_Vertex:        static_cast<FMetalVertexShader*>(Shader)->ReleaseMetalObjects(); break;
+	case SF_Pixel:         static_cast<FMetalPixelShader*>(Shader)->ReleaseMetalObjects(); break;
+	case SF_Compute:       static_cast<FMetalComputeShader*>(Shader)->ReleaseMetalObjects(); break;
+#if PLATFORM_SUPPORTS_GEOMETRY_SHADERS
+	case SF_Geometry:      static_cast<FMetalGeometryShader*>(Shader)->ReleaseMetalObjects(); break;
+#endif
+#if PLATFORM_SUPPORTS_MESH_SHADERS
+	case SF_Mesh:          static_cast<FMetalMeshShader*>(Shader)->ReleaseMetalObjects(); break;
+	case SF_Amplification: static_cast<FMetalAmplificationShader*>(Shader)->ReleaseMetalObjects(); break;
+#endif
+#if METAL_RHI_RAYTRACING
+	case SF_RayGen:
+	case SF_RayMiss:
+	case SF_RayHitGroup:
+	case SF_RayCallable:   static_cast<FMetalRayShader*>(Shader)->ReleaseMetalObjects(); break;
+#endif
+	default: break;
+	}
+}
+
+TRefCountPtr<FRHIShader> FMetalIoStoreShaderCodeArchive::CreateShader(int32 Index, bool bRequired)
+{
+	TRefCountPtr<FRHIShader> Shader = FIoStoreShaderCodeArchive::CreateShader(Index, bRequired);
+	if (Shader)
+	{
+		CreatedShaders.Enqueue(Shader);
+	}
+	return Shader;
+}
+
+void FMetalIoStoreShaderCodeArchive::WaitForPendingDispatchDataDestructors()
+{
+	if (PendingDispatchDataDestructorCount.load(std::memory_order_acquire) > 0)
+	{
+		constexpr double TimeoutSeconds = 5.0;
+		const double StartTime = FPlatformTime::Seconds();
+		const UE::FMonotonicTimePoint Deadline = UE::FMonotonicTimePoint::Now() + UE::FMonotonicTimeSpan::FromSeconds(TimeoutSeconds);
+
+		while (PendingDispatchDataDestructorCount.load(std::memory_order_acquire) > 0)
+		{
+			UE::FEventCountToken Token = DispatchDataDestructorEvent.PrepareWait();
+			if (PendingDispatchDataDestructorCount.load(std::memory_order_acquire) > 0)
+			{
+				const bool bNotified = DispatchDataDestructorEvent.WaitUntil(Token, Deadline);
+				if (!bNotified)
+				{
+					UE_LOG(LogMetal, Warning, TEXT("Shader library '%s' teardown: timed out after %.2f ms waiting for %d Metal shader buffer(s) still held by GCD to be released. Note: Xcode GPU Frame Capture & Metal Validation may hold an additional references to the MTLLibrary, preventing it from being released."),
+						*GetName(), (FPlatformTime::Seconds() - StartTime) * 1000.0, PendingDispatchDataDestructorCount.load());
+					break;
+				}
+			}
+		}
+
+		UE_LOG(LogMetal, Display, TEXT("Shader library '%s' teardown: waited %.2f ms for Metal shader buffers held by GCD to be released."),
+			*GetName(), (FPlatformTime::Seconds() - StartTime) * 1000.0);
+	}
+}
+
+void FMetalIoStoreShaderCodeArchive::OnShaderGroupDataOwnerCreated()
+{
+	PendingDispatchDataDestructorCount.fetch_add(1, std::memory_order_acq_rel);
+}
+
+void FMetalIoStoreShaderCodeArchive::OnShaderGroupDataOwnerReleased()
+{
+	PendingDispatchDataDestructorCount.fetch_sub(1, std::memory_order_acq_rel);
+	DispatchDataDestructorEvent.Notify();
+}
+
+void FMetalIoStoreShaderCodeArchive::ReleaseMetalObjectsAndWait()
+{
+	{
+		TRefCountPtr<FRHIShader> Shader;
+		while (CreatedShaders.Dequeue(Shader))
+		{
+			ReleaseMetalObjectsForShader(Shader.GetReference());
+		}
+	}
+	WaitForPendingDispatchDataDestructors();
+}
+
+void FMetalIoStoreShaderCodeArchive::OnCloseShaderCode()
+{
+	FIoStoreShaderCodeArchive::OnCloseShaderCode();
+	ReleaseMetalObjectsAndWait();
+}
+
+void FMetalIoStoreShaderCodeArchive::Teardown()
+{
+	ReleaseMetalObjectsAndWait();
+	FIoStoreShaderCodeArchive::Teardown();
+}
diff --git a/Engine/Source/Runtime/Apple/MetalRHI/Private/Shaders/MetalIoStoreShaderCodeArchive.h b/Engine/Source/Runtime/Apple/MetalRHI/Private/Shaders/MetalIoStoreShaderCodeArchive.h
new file mode 100644
--- /dev/null
+++ b/Engine/Source/Runtime/Apple/MetalRHI/Private/Shaders/MetalIoStoreShaderCodeArchive.h
@@ -0,0 +1,45 @@
+// Copyright Epic Games, Inc. All Rights Reserved.
+
+#pragma once
+
+#include "ShaderCodeArchive.h"
+#include "Async/EventCount.h"
+#include <atomic>
+
+/**
+ * Metal-specific extension of FIoStoreShaderCodeArchive.
+ * Tracks created shaders and releases Metal objects (MTLFunction/MTLLibrary) on Teardown,
+ * independently of when FRHIShader objects are destroyed.
+ */
+class FMetalIoStoreShaderCodeArchive : public FIoStoreShaderCodeArchive
+{
+public:
+	FMetalIoStoreShaderCodeArchive(EShaderPlatform InPlatform, const FString& InLibraryName, FIoDispatcher& InIoDispatcher);
+
+	virtual void Teardown() override;
+	virtual void OnCloseShaderCode() override;
+	virtual TRefCountPtr<FRHIShader> CreateShader(int32 Index, bool bRequired = true) override;
+
+protected:
+	virtual void OnShaderGroupDataOwnerCreated() override;
+	virtual void OnShaderGroupDataOwnerReleased() override;
+
+private:
+	/** Helper function to call ReleaseMetalObjects for all shader types. */
+	static void ReleaseMetalObjectsForShader(FRHIShader* Shader);
+
+	/** Releases all Metal objects and waits for pending dispatch_data_t destructors. */
+	void ReleaseMetalObjectsAndWait();
+
+	/** Waits for all outstanding dispatch_data_t destructors for this archive to complete. */
+	void WaitForPendingDispatchDataDestructors();
+
+	/** Counts dispatch_data_t objects with a custom destructor that are still alive for this archive. */
+	std::atomic<int32> PendingDispatchDataDestructorCount{0};
+
+	/** Notified each time a dispatch_data_t destructor completes for this archive. */
+	UE::FEventCount DispatchDataDestructorEvent;
+
+	/** All shaders created from this archive. Lock-free MPSC queue since CreateShader can be called from multiple threads. */
+	TQueue<TRefCountPtr<FRHIShader>, EQueueMode::Mpsc> CreatedShaders;
+};
diff --git a/Engine/Source/Runtime/Apple/MetalRHI/Private/Shaders/Types/MetalComputeShader.cpp b/Engine/Source/Runtime/Apple/MetalRHI/Private/Shaders/Types/MetalComputeShader.cpp
--- a/Engine/Source/Runtime/Apple/MetalRHI/Private/Shaders/Types/MetalComputeShader.cpp
+++ b/Engine/Source/Runtime/Apple/MetalRHI/Private/Shaders/Types/MetalComputeShader.cpp
@@ -44,6 +44,12 @@ FMetalComputeShader::~FMetalComputeShader()
     }
 }
 
+void FMetalComputeShader::ReleaseMetalObjects()
+{
+	Pipeline.Reset();
+	TMetalBaseShader::ReleaseMetalObjects();
+}
+
 FMetalShaderPipelinePtr FMetalComputeShader::GetPipeline()
 {
 	FScopeLock Lock(&PipelineCS);
diff --git a/Engine/Source/Runtime/Apple/MetalRHI/Private/Shaders/Types/MetalComputeShader.h b/Engine/Source/Runtime/Apple/MetalRHI/Private/Shaders/Types/MetalComputeShader.h
--- a/Engine/Source/Runtime/Apple/MetalRHI/Private/Shaders/Types/MetalComputeShader.h
+++ b/Engine/Source/Runtime/Apple/MetalRHI/Private/Shaders/Types/MetalComputeShader.h
@@ -21,6 +21,9 @@ public:
 	FMetalComputeShader(FMetalDevice& Device, const FRHICreateShaderDesc& CreateShaderDesc, MTLLibraryPtr InLibrary);
 	virtual ~FMetalComputeShader();
 
+	// Release the Metal Pipeline.
+	virtual void ReleaseMetalObjects() override;
+
 	FMetalShaderPipelinePtr GetPipeline();
     MTLFunctionPtr GetFunction();
 
diff --git a/Engine/Source/Runtime/Apple/MetalRHI/Private/Shaders/Types/MetalRayShader.cpp b/Engine/Source/Runtime/Apple/MetalRHI/Private/Shaders/Types/MetalRayShader.cpp
--- a/Engine/Source/Runtime/Apple/MetalRHI/Private/Shaders/Types/MetalRayShader.cpp
+++ b/Engine/Source/Runtime/Apple/MetalRHI/Private/Shaders/Types/MetalRayShader.cpp
@@ -158,6 +158,15 @@ FMetalRayShader::~FMetalRayShader()
 	}
 }
 
+void FMetalRayShader::ReleaseMetalObjects()
+{	
+	if (RayGen)       { RayGen->ReleaseMetalObjects(); }
+	if (Miss)         { Miss->ReleaseMetalObjects(); }
+	if (Callable)     { Callable->ReleaseMetalObjects(); }
+	if (ClosestHit)   { ClosestHit->ReleaseMetalObjects(); }
+	if (AnyHitAndIntersection) { AnyHitAndIntersection->ReleaseMetalObjects(); }
+}
+
 IRRootSignature *FMetalRayShader::CreateRootSignature(const TArray<IRRootParameter1> &InRootParams, bool bIsLocal)
 {
 	TArray<IRRootParameter1> RootParams;
diff --git a/Engine/Source/Runtime/Apple/MetalRHI/Private/Shaders/Types/MetalRayShader.h b/Engine/Source/Runtime/Apple/MetalRHI/Private/Shaders/Types/MetalRayShader.h
--- a/Engine/Source/Runtime/Apple/MetalRHI/Private/Shaders/Types/MetalRayShader.h
+++ b/Engine/Source/Runtime/Apple/MetalRHI/Private/Shaders/Types/MetalRayShader.h
@@ -56,6 +56,9 @@ public:
 	FMetalRayShader(FMetalDevice& Device, const FRHICreateShaderDesc& CreateShaderDesc, EShaderFrequency InFrequency, MTLLibraryPtr InLibrary = nullptr);
 	virtual ~FMetalRayShader();
 
+	// Releases MTLFunction and MTLLibrary references on all sub-shaders.
+	void ReleaseMetalObjects();
+
 	inline IRRootSignature *GetGlobalRootSignature() const { return GlobalRootSignature; }
 	inline IRRootSignature *GetLocalRootSignature() const { return LocalRootSignature; }
 
diff --git a/Engine/Source/Runtime/Apple/MetalRHI/Private/Shaders/Types/Templates/MetalBaseShader.h b/Engine/Source/Runtime/Apple/MetalRHI/Private/Shaders/Types/Templates/MetalBaseShader.h
--- a/Engine/Source/Runtime/Apple/MetalRHI/Private/Shaders/Types/Templates/MetalBaseShader.h
+++ b/Engine/Source/Runtime/Apple/MetalRHI/Private/Shaders/Types/Templates/MetalBaseShader.h
@@ -63,6 +63,9 @@ struct FMetalShaderData
     /** CRC & Len for name disambiguation */
     uint32 SourceLen = 0;
     uint32 SourceCRC = 0;
+
+    /** Unique identifier for the owning FRHIShaderLibrary, used to scope the compiled shader cache per library. 0 if no library. */
+    uint64 LibrarySeed = 0;
     
     // this is the compiler shader
     MTLFunctionPtr Function;
@@ -98,6 +101,16 @@ public:
 	void Init(const FRHICreateShaderDesc& CreateShaderDesc, FMetalCodeHeader& Header, MTLLibraryPtr InLibrary, bool bBinaryCompiled = false, bool bUseExtendedCodeSection = false);
 	void Destroy();
 
+	/** Releases MTLFunction and MTLLibrary references. Called during shader library teardown to free Metal objects
+	 *  while the FRHIShader may still be alive (held by materials, PSO cache, etc.). */
+	virtual void ReleaseMetalObjects()
+	{
+		FMetalCompiledShaderKey Key(SourceLen, SourceCRC, bIsExtendedBytecode, LibrarySeed);
+		GetMetalCompiledShaderCache().Remove(Key);
+		Library.reset();
+		Function.reset();
+	}
+
 	virtual FString GetFunctionName()
 	{
 		return FString::Printf(TEXT("Main_%0.8x_%0.8x"), SourceLen, SourceCRC);
@@ -223,14 +236,17 @@ void TMetalBaseShader<BaseResourceType, ShaderType>::Init(const FRHICreateShader
 
 	Library = InLibrary;
 
+	LibrarySeed = reinterpret_cast<uint64>(CreateShaderDesc.Library);
+
 	bool bNeedsCompiling = false;
-	
-	FMetalCompiledShaderKey Key(Header.SourceLen, Header.SourceCRC, bIsExtendedBytecode);
 
-	Function = GetMetalCompiledShaderCache().FindRef(Key);
+	FMetalCompiledShaderKey Key(Header.SourceLen, Header.SourceCRC, bIsExtendedBytecode, LibrarySeed);
+
+	FCachedCompiledShader CachedShader = GetMetalCompiledShaderCache().Find(Key);
+	Function = CachedShader.Function;
 	if (!Library && Function)
 	{
-		Library = GetMetalCompiledShaderCache().FindLibrary(Function);
+		Library = CachedShader.Library;
 	}
 	else
 	{
@@ -251,13 +267,28 @@ void TMetalBaseShader<BaseResourceType, ShaderType>::Init(const FRHICreateShader
 				// Archived shaders should never get in here.
 				check(!(Header.CompileFlags & (1 << CFLAG_Archive)) || BufferSize > 0);
 				dispatch_data_t GCDBuffer;
-				if (CreateShaderDesc.Owner.IsValid())
+				if (CreateShaderDesc.CodeOwner.IsValid())
 				{
 					// Hold a reference until the MTLLibrary is done using it
-					__block auto OwnerCopy = CreateShaderDesc.Owner;
+					__block auto OwnerCopy = CreateShaderDesc.CodeOwner;
 					GCDBuffer = dispatch_data_create(SourceCode, BufferSize, nil, ^{
 						OwnerCopy.Reset();
 					});
+
+					// Append seed data after the metallib content to prevent Metal's
+					// newLibrary cache from returning a cached library. Metal's cache key appears
+					// to be based on the full dispatch_data bytes, so appending unique data
+					// ensures each library is treated as distinct.
+					if (LibrarySeed != 0)
+					{
+						uint64* SeedData = new uint64;
+						*SeedData = LibrarySeed;
+						dispatch_data_t SeedBuffer = dispatch_data_create(SeedData, sizeof(uint64), nil, ^{ delete SeedData; });
+						dispatch_data_t CombinedBuffer = dispatch_data_create_concat(GCDBuffer, SeedBuffer);
+						dispatch_release(GCDBuffer);
+						dispatch_release(SeedBuffer);
+						GCDBuffer = CombinedBuffer;
+					}
 				}
 				else
 				{
@@ -434,8 +465,8 @@ MTLFunctionPtr TMetalBaseShader<BaseResourceType, ShaderType>::GetCompiledFuncti
     if (!Func || bNeedToRecreateFunction)
 	{
 		// Find the existing compiled shader in the cache.
-		FMetalCompiledShaderKey Key(SourceLen, SourceCRC, bIsExtendedBytecode);
-		Func = Function = GetMetalCompiledShaderCache().FindRef(Key);
+		FMetalCompiledShaderKey Key(SourceLen, SourceCRC, bIsExtendedBytecode, LibrarySeed);
+		Func = Function = GetMetalCompiledShaderCache().Find(Key).Function;
 
         if (bNeedToRecreateFunction)
         {
diff --git a/Engine/Source/Runtime/RHI/Public/RHIResources.h b/Engine/Source/Runtime/RHI/Public/RHIResources.h
--- a/Engine/Source/Runtime/RHI/Public/RHIResources.h
+++ b/Engine/Source/Runtime/RHI/Public/RHIResources.h
@@ -790,20 +790,24 @@ struct FRHICreateShaderDesc
 {
 	TConstArrayView<uint8> Code;
 	// Optional Owner of the memory backing the View
-	TSharedPtr<void> Owner;
-	
-	FRHICreateShaderDesc(TConstArrayView<uint8> InCode)
+	TSharedPtr<void> CodeOwner;
+	// Optional pointer to the shader library this shader is being created from
+	class FRHIShaderLibrary* Library;
+
+	FRHICreateShaderDesc(TConstArrayView<uint8> InCode, class FRHIShaderLibrary* InLibrary = nullptr)
 	: Code(InCode)
+	, Library(InLibrary)
 	{
 	}
-	
-	template<typename DeleterType>
-	FRHICreateShaderDesc(TConstArrayView<uint8> InCode, DeleterType&& InDeleter)
+
+	template<typename DeleterType, typename = std::enable_if_t<!std::is_pointer_v<std::decay_t<DeleterType>>>>
+	FRHICreateShaderDesc(TConstArrayView<uint8> InCode, DeleterType&& InDeleter, class FRHIShaderLibrary* InLibrary = nullptr)
 	 : Code(InCode)
 	// Giving the view data here serve two purposes:
 	// - Hints on the intent that this ownership correspond to the data behind the View
 	// - Make the Owner valid to differentiate between data that stay valid as long as we keep a reference to it versus temporary data
-	 , Owner(MakeShareable((void*)InCode.GetData(), Forward<DeleterType>(InDeleter)))
+	 , CodeOwner(MakeShareable((void*)InCode.GetData(), Forward<DeleterType>(InDeleter)))
+	 , Library(InLibrary)
 	{
 	}
 };
diff --git a/Engine/Source/Runtime/RHI/Public/RHIShaderLibrary.h b/Engine/Source/Runtime/RHI/Public/RHIShaderLibrary.h
--- a/Engine/Source/Runtime/RHI/Public/RHIShaderLibrary.h
+++ b/Engine/Source/Runtime/RHI/Public/RHIShaderLibrary.h
@@ -55,7 +55,10 @@ public:
 	virtual int32 GetLibraryId() { return LibraryId; }
 	/*CreateShader can return a null shader when bRequired == false. Usefull to debug dynamic shader preloading or when shaders haven't finished loading.*/
 	virtual TRefCountPtr<FRHIShader> CreateShader(int32 ShaderIndex, bool bRequired = true) { return nullptr; }
-	virtual void Teardown() {};
+	virtual void Teardown() {}
+
+	/** Called prior to Teardown(), giving subclasses a chance to perform additional cleanup. */
+	virtual void OnCloseShaderCode() {}
 
 protected:
 	EShaderPlatform Platform;
diff --git a/Engine/Source/Runtime/RenderCore/Internal/ShaderLibrary/ShaderCodeArchiveInternal.h b/Engine/Source/Runtime/RenderCore/Internal/ShaderLibrary/ShaderCodeArchiveInternal.h
--- a/Engine/Source/Runtime/RenderCore/Internal/ShaderLibrary/ShaderCodeArchiveInternal.h
+++ b/Engine/Source/Runtime/RenderCore/Internal/ShaderLibrary/ShaderCodeArchiveInternal.h
@@ -57,4 +57,13 @@ namespace ShaderCodeArchive
 		ChunkId.Set(Data, 12);
 		return ChunkId;
 	}
+	
+	inline FIoChunkId GetLibraryScopedShaderGroupChunkId(const FIoChunkId& ContentChunkId, const FIoChunkId& LibraryChunkId)
+	{
+		uint8 Input[24];
+		FMemory::Memcpy(Input,      ContentChunkId.GetData(), 12);
+		FMemory::Memcpy(Input + 12, LibraryChunkId.GetData(), 12);
+		const uint64 Hash = FShaderHash::HashBuffer(Input, sizeof(Input)).Hash;
+		return CreateIoChunkId(Hash, 0, EIoChunkType::ShaderCode);
+	}
 } // namespace ShaderCodeArchive
diff --git a/Engine/Source/Runtime/RenderCore/Private/ShaderLibrary/ShaderCodeArchive.cpp b/Engine/Source/Runtime/RenderCore/Private/ShaderLibrary/ShaderCodeArchive.cpp
--- a/Engine/Source/Runtime/RenderCore/Private/ShaderLibrary/ShaderCodeArchive.cpp
+++ b/Engine/Source/Runtime/RenderCore/Private/ShaderLibrary/ShaderCodeArchive.cpp
@@ -1745,7 +1745,7 @@ TRefCountPtr<FRHIShader> FShaderCodeArchive::CreateShader(int32 Index, bool bReq
 	// detect the breach of contract early
 	ensureAlwaysMsgf(IsInRenderingThread() || GRHISupportsMultithreadedShaderCreation, TEXT("More than one thread is creating shaders, but GRHISupportsMultithreadedShaderCreation is false."));
 
-	FRHICreateShaderDesc CreateShaderDesc = FRHICreateShaderDesc(MakeArrayView(ShaderCode, ShaderEntry.GetUncompressedSize()));
+	FRHICreateShaderDesc CreateShaderDesc = FRHICreateShaderDesc(MakeArrayView(ShaderCode, ShaderEntry.GetUncompressedSize()), this);
 
 	switch (ShaderEntry.GetFrequency())
 	{
@@ -2509,7 +2509,15 @@ FIoStoreShaderCodeArchive* FIoStoreShaderCodeArchive::Create(EShaderPlatform InP
 		Ar << Version;
 		if (Version == CurrentVersion)
 		{
-			FIoStoreShaderCodeArchive* Library = new FIoStoreShaderCodeArchive(InPlatform, InLibraryName, InIoDispatcher);
+			FIoStoreShaderCodeArchive* Library;
+			if (IoStoreShaderCodeArchiveFactory)
+			{
+				Library = IoStoreShaderCodeArchiveFactory(InPlatform, InLibraryName, InIoDispatcher);
+			}
+			else
+			{
+				Library = new FIoStoreShaderCodeArchive(InPlatform, InLibraryName, InIoDispatcher);
+			}
 			Ar << Library->Header;
 			{
 				const uint32 HashSize = FMath::Min<uint32>(0x10000, 1u << FMath::CeilLogTwo(Library->Header.ShaderMapHashes.Num()));
@@ -2563,6 +2571,13 @@ FIoStoreShaderCodeArchive* FIoStoreShaderCodeArchive::Create(EShaderPlatform InP
 	return nullptr;
 }
 
+FIoStoreShaderCodeArchive::FCreateIoStoreShaderCodeArchiveDelegate FIoStoreShaderCodeArchive::IoStoreShaderCodeArchiveFactory;
+
+void FIoStoreShaderCodeArchive::RegisterIoStoreShaderCodeArchiveFactory(FCreateIoStoreShaderCodeArchiveDelegate InFactory)
+{
+	IoStoreShaderCodeArchiveFactory = MoveTemp(InFactory);
+}
+
 FIoStoreShaderCodeArchive::FIoStoreShaderCodeArchive(EShaderPlatform InPlatform, const FString& InLibraryName, FIoDispatcher& InIoDispatcher)
 	: FRHIShaderLibrary(InPlatform, InLibraryName)
 	, IoDispatcher(InIoDispatcher)
@@ -2575,6 +2590,17 @@ FIoStoreShaderCodeArchive::~FIoStoreShaderCodeArchive()
 	Teardown();
 }
 
+void FIoStoreShaderCodeArchive::OnCloseShaderCode()
+{
+	// Release shader group data now since the backing IoStore memory is about to be unmounted.
+	// The preload entries themselves are kept alive and will be cleaned up in Teardown().
+	FWriteScopeLock Lock(PreloadedShaderGroupsLock);
+	for (TMap<int32, FShaderGroupPreloadEntry*>::TIterator Iter(PreloadedShaderGroups); Iter; ++Iter)
+	{
+		Iter.Value()->ShaderGroupData.Reset();
+	}
+}
+
 void FIoStoreShaderCodeArchive::Teardown()
 {
 	DebugVisualizer.SaveShaderUsageBitmap(GetName(), GetPlatform());
@@ -3152,15 +3178,23 @@ TRefCountPtr<FRHIShader> FIoStoreShaderCodeArchive::CreateShader(int32 ShaderInd
 	{
 		ShaderCode += ShaderEntry.UncompressedOffsetInGroup;
 	}
-	FRHICreateShaderDesc CreateShaderDesc = IsGroupCompressed
-		? FRHICreateShaderDesc(MakeArrayView(ShaderCode, Header.GetShaderUncompressedSize(ShaderIndex)))
-		: FRHICreateShaderDesc(
+	auto MakeCreateShaderDesc = [this, IsGroupCompressed, ShaderCode, ShaderIndex, PreloadEntryPtr]() -> FRHICreateShaderDesc
+	{
+		if (IsGroupCompressed)
+		{
+			return FRHICreateShaderDesc(MakeArrayView(ShaderCode, Header.GetShaderUncompressedSize(ShaderIndex)), this);
+		}
+		OnShaderGroupDataOwnerCreated();
+		return FRHICreateShaderDesc(
 			MakeArrayView(ShaderCode, Header.GetShaderUncompressedSize(ShaderIndex)),
-			[ShaderGroupDataSharedPtr = PreloadEntryPtr->ShaderGroupData](void*)
+			[this, ShaderGroupDataSharedPtr = PreloadEntryPtr->ShaderGroupData](void*)
 			{
 				// ShaderGroupDataSharedPtr will be released when this lambda is destroyed
-			}
-		);
+				OnShaderGroupDataOwnerReleased();
+			},
+			this);
+	};
+	FRHICreateShaderDesc CreateShaderDesc = MakeCreateShaderDesc();
 	
 	switch (ShaderEntry.Frequency)
 	{
diff --git a/Engine/Source/Runtime/RenderCore/Private/ShaderLibrary/ShaderCodeLibrary.cpp b/Engine/Source/Runtime/RenderCore/Private/ShaderLibrary/ShaderCodeLibrary.cpp
--- a/Engine/Source/Runtime/RenderCore/Private/ShaderLibrary/ShaderCodeLibrary.cpp
+++ b/Engine/Source/Runtime/RenderCore/Private/ShaderLibrary/ShaderCodeLibrary.cpp
@@ -82,7 +82,7 @@ namespace UE::ShaderLibrary::Private
 	using ESaveToDiskSortOrder = ShaderCodeArchive::ESaveToDiskSortOrder;
 
 	static FDelegateHandle OnPakFileMountedDelegateHandle;
-	static FDelegateHandle OnPakFileUnmountedDelegateHandle;
+	static FDelegateHandle OnPakFileUnmountingDelegateHandle;
 	static FDelegateHandle OnPluginMountedDelegateHandle;
 	static FDelegateHandle OnPluginUnmountedDelegateHandle;
 
@@ -242,7 +242,7 @@ namespace UE::ShaderLibrary::Private
 			}
 		}
 
-		void OnPakFileUnmounted(const FMountedPakFileInfo& MountInfo, const FString& Directory)
+		void OnPakFileUnmounting(const FMountedPakFileInfo& MountInfo, const FString& Directory)
 		{
 			FShaderLibraryChunkInfo ChunkInfo = GetLibraryChunkInfo(MountInfo);
 			if (GetPresentChunks().Contains(ChunkInfo.ChunkId))
@@ -1300,10 +1300,10 @@ public:
 		});
 	}
 
-	void OnPakFileUnmounted(const UE::ShaderLibrary::Private::FMountedPakFileInfo& MountInfo)
+	void OnPakFileUnmounting(const UE::ShaderLibrary::Private::FMountedPakFileInfo& MountInfo)
 	{
 		PakFileStateChangeInternal(FShaderPipelineCache::Closed, [&MountInfo](UE::ShaderLibrary::Private::FNamedShaderLibrary* InNamedShaderLibrary) {
-			InNamedShaderLibrary->OnPakFileUnmounted(MountInfo, InNamedShaderLibrary->BaseDirectory);
+			InNamedShaderLibrary->OnPakFileUnmounting(MountInfo, InNamedShaderLibrary->BaseDirectory);
 		});
 	}
 
@@ -2203,11 +2203,11 @@ static void FShaderLibraryPakFileMountedCallback(const IPakFile& PakFile)
 	}
 }
 
-static void FShaderLibraryPakFileUnmountedCallback(const IPakFile& PakFile)
+static void FShaderLibraryPakFileUnmountingCallback(const IPakFile& PakFile)
 {
 	using namespace UE::ShaderLibrary::Private;
 
-	UE_LOG(LogShaderLibrary, Verbose, TEXT("ShaderCodeLibraryPakFileUnmountedCallback: PakFile '%s' (chunk index %d, root '%s') unmounted"), *PakFile.PakGetPakFilename(), PakFile.PakGetPakchunkIndex(), *PakFile.PakGetMountPoint());
+	UE_LOG(LogShaderLibrary, Verbose, TEXT("ShaderCodeLibraryPakFileUnmountingCallback: PakFile '%s' (chunk index %d, root '%s') unmounting"), *PakFile.PakGetPakFilename(), PakFile.PakGetPakchunkIndex(), *PakFile.PakGetMountPoint());
 
 	FMountedPakFileInfo PakFileInfo(PakFile);
 	{
@@ -2218,7 +2218,7 @@ static void FShaderLibraryPakFileUnmountedCallback(const IPakFile& PakFile)
 	// If shaderlibrary has not yet been initialized, report the chunk as pending
 	if (FShaderLibrariesCollection::Impl)
 	{
-		FShaderLibrariesCollection::Impl->OnPakFileUnmounted(PakFileInfo);
+		FShaderLibrariesCollection::Impl->OnPakFileUnmounting(PakFileInfo);
 	}
 	else
 	{
@@ -2230,7 +2230,7 @@ void FShaderCodeLibrary::PreInit()
 {
 	// add a callback for opening later chunks
 	UE::ShaderLibrary::Private::OnPakFileMountedDelegateHandle = FCoreDelegates::GetOnPakFileMounted2().AddStatic(&FShaderLibraryPakFileMountedCallback);
-	UE::ShaderLibrary::Private::OnPakFileUnmountedDelegateHandle = FCoreDelegates::GetOnPakFileUnmounted().AddStatic(&FShaderLibraryPakFileUnmountedCallback);
+	UE::ShaderLibrary::Private::OnPakFileUnmountingDelegateHandle = FCoreDelegates::GetOnPakFileUnmounting().AddStatic(&FShaderLibraryPakFileUnmountingCallback);
 }
 
 void FShaderCodeLibrary::InitForRuntime(EShaderPlatform ShaderPlatform)
@@ -2304,10 +2304,10 @@ void FShaderCodeLibrary::Shutdown()
 		FCoreDelegates::GetOnPakFileMounted2().Remove(UE::ShaderLibrary::Private::OnPakFileMountedDelegateHandle);
 		UE::ShaderLibrary::Private::OnPakFileMountedDelegateHandle.Reset();
 	}
-	if (UE::ShaderLibrary::Private::OnPakFileUnmountedDelegateHandle.IsValid())
+	if (UE::ShaderLibrary::Private::OnPakFileUnmountingDelegateHandle.IsValid())
 	{
-		FCoreDelegates::GetOnPakFileUnmounted().Remove(UE::ShaderLibrary::Private::OnPakFileUnmountedDelegateHandle);
-		UE::ShaderLibrary::Private::OnPakFileUnmountedDelegateHandle.Reset();
+		FCoreDelegates::GetOnPakFileUnmounting().Remove(UE::ShaderLibrary::Private::OnPakFileUnmountingDelegateHandle);
+		UE::ShaderLibrary::Private::OnPakFileUnmountingDelegateHandle.Reset();
 	}
 	if (UE::ShaderLibrary::Private::OnPluginMountedDelegateHandle.IsValid())
 	{
@@ -3107,6 +3107,8 @@ void UE::ShaderLibrary::Private::FNamedShaderLibrary::CloseShaderCode(const FStr
 			{
 				LogShaderLibraryInstanceState(LibraryInstance, *Library, TEXT("Released"));
 
+				// Give the FRHIShaderLibrary an opportunity to do cleanup before we unmount the pak file at the end of the frame.
+				LibraryInstance->Library->OnCloseShaderCode();
 				// Defer the actual deletion of the shader library instance to the end of a frame, since some RHI resources might still hold on to their shaders in the current frame.
 				FDeferredShaderLibraryInstanceDeleter::Get().EnqueueRelease(MoveTemp(Components[Index]));
 
diff --git a/Engine/Source/Runtime/RenderCore/Public/ShaderCodeArchive.h b/Engine/Source/Runtime/RenderCore/Public/ShaderCodeArchive.h
--- a/Engine/Source/Runtime/RenderCore/Public/ShaderCodeArchive.h
+++ b/Engine/Source/Runtime/RenderCore/Public/ShaderCodeArchive.h
@@ -419,12 +419,13 @@ public:
 	RENDERCORE_API bool FindOrAddShaderMapNonEditor(const FShaderHash& Hash, int32& OutIndex);
 #endif
 	RENDERCORE_API void RemoveLastAddedShader();
+	RENDERCORE_API void Finalize();
+#endif // !USE_MMAPPED_SHADERARCHIVE
+
 #if WITH_EDITORONLY_DATA
 	/** Sorts Shaders and ShaderMaps by their Hash values. Used to make cooked data deterministic. */
 	RENDERCORE_API void SortByHashes(TArray<FSharedBuffer>& InOutShaderCodeForShaderIndex);
 #endif
-	RENDERCORE_API void Finalize();
-#endif // !USE_MMAPPED_SHADERARCHIVE
 
 	RENDERCORE_API void DecompressShader(int32 Index, const TArray<FSharedBuffer>& ShaderCode, TArray<uint8>& OutDecompressedShader) const;
 	RENDERCORE_API static void DecompressShaderEntryAndAppend(const FShaderCodeEntry& Entry, const uint8* CompressedShaderCode, TArray<uint8>& OutDecompressedShader);
@@ -839,7 +840,6 @@ public:
 class FIoStoreShaderCodeArchive : public FRHIShaderLibrary
 {
 public:
-
 	UE_DEPRECATED(5.8, "GetShaderCodeArchiveChunkId is internal engine functionality and so is no longer publically exposed")
 	RENDERCORE_API static FIoChunkId GetShaderCodeArchiveChunkId(const FString& LibraryName, FName FormatName);
 	UE_DEPRECATED(5.8, "GetShaderCodeChunkId is internal engine functionality and so is no longer publically exposed")
@@ -847,9 +847,12 @@ public:
 	/** This function creates the archive header, including splitting shaders into groups. */
 	RENDERCORE_API static void CreateIoStoreShaderCodeArchiveHeader(const FName& Format, const FSerializedShaderArchive& SerializedShaders, FIoStoreShaderCodeArchiveHeader& OutHeader);
 	RENDERCORE_API static void SaveIoStoreShaderCodeArchive(const FIoStoreShaderCodeArchiveHeader& Header, FArchive& OutLibraryAr);
-	static FIoStoreShaderCodeArchive* Create(EShaderPlatform InPlatform, const FString& InLibraryName, FIoDispatcher& InIoDispatcher);
+	RENDERCORE_API static FIoStoreShaderCodeArchive* Create(EShaderPlatform InPlatform, const FString& InLibraryName, FIoDispatcher& InIoDispatcher);
+
+	using FCreateIoStoreShaderCodeArchiveDelegate = TFunction<FIoStoreShaderCodeArchive*(EShaderPlatform, const FString&, FIoDispatcher&)>;
+	RENDERCORE_API static void RegisterIoStoreShaderCodeArchiveFactory(FCreateIoStoreShaderCodeArchiveDelegate InFactory);
 
-	virtual ~FIoStoreShaderCodeArchive();
+	RENDERCORE_API virtual ~FIoStoreShaderCodeArchive();
 
 	virtual bool IsNativeLibrary() const override { return false; }
 
@@ -860,7 +863,7 @@ public:
 			PreloadedShaderGroups.GetAllocatedSize();
 	}
 
-	virtual uint32 GetShaderSizeBytes(int32 ShaderIndex) const override;
+	RENDERCORE_API virtual uint32 GetShaderSizeBytes(int32 ShaderIndex) const override;
 
 	virtual int32 GetNumShaders() const override { return Header.ShaderEntries.Num(); }
 	virtual int32 GetNumShaderMaps() const override { return Header.ShaderMapEntries.Num(); }
@@ -881,8 +884,8 @@ public:
 		}
 	}
 
-	virtual int32 FindShaderMapIndex(const FShaderHash& Hash) override;
-	virtual int32 FindShaderIndex(const FShaderHash& Hash) override;
+	RENDERCORE_API virtual int32 FindShaderMapIndex(const FShaderHash& Hash) override;
+	RENDERCORE_API virtual int32 FindShaderIndex(const FShaderHash& Hash) override;
 	virtual FShaderHash GetShaderHash(int32 ShaderMapIndex, int32 ShaderIndex) override
 	{
 		return Header.ShaderHashes[GetShaderIndex(ShaderMapIndex, ShaderIndex)];
@@ -893,10 +896,10 @@ public:
 		return Header.ShaderMapHashes[ShaderMapIndex];
 	}
 
-	virtual bool IsPreloading(int32 ShaderIndex, FGraphEventArray& OutCompletionEvents) override;
-	virtual bool PreloadShader(int32 ShaderIndex, FGraphEventArray& OutCompletionEvents) override;
-	virtual void AddRefPreloadedShaderGroup(int32 ShaderGroupIndex) override;
-	virtual void ReleasePreloadedShaderGroup(int32 ShaderGroupIndex) override;
+	RENDERCORE_API virtual bool IsPreloading(int32 ShaderIndex, FGraphEventArray& OutCompletionEvents) override;
+	RENDERCORE_API virtual bool PreloadShader(int32 ShaderIndex, FGraphEventArray& OutCompletionEvents) override;
+	RENDERCORE_API virtual void AddRefPreloadedShaderGroup(int32 ShaderGroupIndex) override;
+	RENDERCORE_API virtual void ReleasePreloadedShaderGroup(int32 ShaderGroupIndex) override;
 
 	/** Returns the index of shader group that a given shader belongs to. */
 	virtual int32 GetGroupIndexForShader(int32 ShaderIndex) const override
@@ -904,28 +907,27 @@ public:
 		return Header.ShaderEntries[ShaderIndex].ShaderGroupIndex;
 	}
 	
-	virtual bool ResolveShaderMap(int32 ShaderMapIndex, FCoreDelegates::FIoChunkIdResolvedFunc IoChunkIdResolvedFunc) const override;
-	virtual bool PreloadShaderMap(int32 ShaderMapIndex, FGraphEventArray& OutCompletionEvents) override;
-	virtual bool PreloadShaderMap(int32 ShaderMapIndex, FCoreDelegates::FAttachShaderReadRequestFunc AttachShaderReadRequestFunc) override;
-	virtual void ReleasePreloadedShader(int32 ShaderIndex) override;
+	RENDERCORE_API virtual bool ResolveShaderMap(int32 ShaderMapIndex, FCoreDelegates::FIoChunkIdResolvedFunc IoChunkIdResolvedFunc) const override;
+	RENDERCORE_API virtual bool PreloadShaderMap(int32 ShaderMapIndex, FGraphEventArray& OutCompletionEvents) override;
+	RENDERCORE_API virtual bool PreloadShaderMap(int32 ShaderMapIndex, FCoreDelegates::FAttachShaderReadRequestFunc AttachShaderReadRequestFunc) override;
+	RENDERCORE_API virtual void ReleasePreloadedShader(int32 ShaderIndex) override;
 
 	/** Returns true if all IoChunks associated with the specified shadermap have been resolved and are available for use without any further delay.
 	    Otherwise, at least one IoChunk is either unknown or has not been installed yet (when coming from an on-demand container). */
-	virtual bool IsShaderMapResolved(int32 ShaderMapIndex) const override;
+	RENDERCORE_API virtual bool IsShaderMapResolved(int32 ShaderMapIndex) const override;
 
-	virtual TRefCountPtr<FRHIShader> CreateShader(int32 Index, bool bRequired = true) override;
-	virtual void Teardown() override;
+	RENDERCORE_API virtual TRefCountPtr<FRHIShader> CreateShader(int32 Index, bool bRequired = true) override;
+	RENDERCORE_API virtual void Teardown() override;
+	RENDERCORE_API virtual void OnCloseShaderCode() override;
 
-private:
-	static constexpr uint32 CurrentVersion = 1;
-	
+protected:
 	struct FShaderGroupData
 	{
 		FIoRequest IoRequest;
 #if USE_MMAPPED_SHADERARCHIVE
 		TIoStatusOr<FIoMappedRegion> MappedRegionStatus;
 #endif
-		
+
 		~FShaderGroupData()
 		{
 			IoRequest.Cancel();
@@ -962,11 +964,37 @@ private:
 
 		~FShaderGroupPreloadEntry() = default;
 	};
+	RENDERCORE_API FIoStoreShaderCodeArchive(EShaderPlatform InPlatform, const FString& InLibraryName, FIoDispatcher& InIoDispatcher);
+
+	/** Called when a shader group data owner with a custom deleter is about to be created. */
+	virtual void OnShaderGroupDataOwnerCreated() {}
 
-	FIoStoreShaderCodeArchive(EShaderPlatform InPlatform, const FString& InLibraryName, FIoDispatcher& InIoDispatcher);
+	/** Called when a shader group data owner's deleter has been invoked (the data is no longer in use). */
+	virtual void OnShaderGroupDataOwnerReleased() {}
 
 	FIoDispatcher& IoDispatcher;
 
+	/** Archive header with all the metadata */
+	FIoStoreShaderCodeArchiveHeader Header;
+
+	/** Hash tables for faster searching for shader and shadermap hashes. */
+	FHashTable ShaderMapHashTable;
+	FHashTable ShaderHashTable;
+
+	/** Mapping between the group index and preloaded groups. Should be only modified when lock is taken. */
+	TMap<int32, FShaderGroupPreloadEntry*> PreloadedShaderGroups;
+	/** Lock guarding access to the book-keeping info above.*/
+	FRWLock PreloadedShaderGroupsLock;
+
+	/** debug visualizer - in Shipping compiles out to an empty struct with no-op functions */
+	FShaderUsageVisualizer DebugVisualizer;
+
+	static constexpr uint32 CurrentVersion = 1;
+
+private:
+
+	static FCreateIoStoreShaderCodeArchiveDelegate IoStoreShaderCodeArchiveFactory;
+
 	/** Preloads a given shader group. */
 	bool PreloadShaderGroup(int32 ShaderGroupIndex, FGraphEventArray& OutCompletionEvents, 
 #if UE_SCA_DEBUG_PRELOADING
@@ -1012,21 +1040,6 @@ private:
 
 	/** Returns true if the group contains only RTX shaders. We can avoid preloading it when running with RTX off. */
 	bool GroupOnlyContainsRaytracingShaders(int32 ShaderGroupIndex) const;
-
-	/** Archive header with all the metadata */
-	FIoStoreShaderCodeArchiveHeader Header;
-
-	/** Hash tables for faster searching for shader and shadermap hashes. */
-	FHashTable ShaderMapHashTable;
-	FHashTable ShaderHashTable;
-
-	/** Mapping between the group index and preloaded groups. Should be only modified when lock is taken. */
-	TMap<int32, FShaderGroupPreloadEntry*> PreloadedShaderGroups;
-	/** Lock guarding access to the book-keeping info above.*/
-	FRWLock PreloadedShaderGroupsLock;
-
-	/** debug visualizer - in Shipping compiles out to an empty struct with no-op functions */
-	FShaderUsageVisualizer DebugVisualizer;
 };
 
 
