Delete oculus openXR shader

#Synced-CL 52511927
#Synced-CL 52511927

diff --git a/Engine/Source/Runtime/RenderCore/Public/OculusShaders.h b/Engine/Source/Runtime/RenderCore/Public/OculusShaders.h
--- a/Engine/Source/Runtime/RenderCore/Public/OculusShaders.h
+++ b/Engine/Source/Runtime/RenderCore/Public/OculusShaders.h
@@ -127,38 +127,4 @@ private:
 	LAYOUT_FIELD(FShaderResourceParameter, InTexture);
 	LAYOUT_FIELD(FShaderResourceParameter, InTextureSampler);
 	LAYOUT_FIELD(FShaderParameter, InFaceIndexParameter);
-};
-
-class FOculusOpenXRCubemapPS : public FGlobalShader
-{
-	DECLARE_EXPORTED_SHADER_TYPE(FOculusOpenXRCubemapPS, Global, RENDERCORE_API);
-public:
-
-	static bool ShouldCompilePermutation(const FGlobalShaderPermutationParameters& Parameters) { return true; }
-
-	FOculusOpenXRCubemapPS(const ShaderMetaType::CompiledShaderInitializerType& Initializer) :
-		FGlobalShader(Initializer)
-	{
-		InTexture.Bind(Initializer.ParameterMap, TEXT("InTextureCube"), SPF_Mandatory);
-		InTextureSampler.Bind(Initializer.ParameterMap, TEXT("InTextureSampler"));
-		InFaceToSampleParameter.Bind(Initializer.ParameterMap, TEXT("FaceToSample"));
-	}
-	FOculusOpenXRCubemapPS() {}
-
-	void SetParameters(FRHIBatchedShaderParameters& BatchedParameters, const FTexture* Texture, const FMatrix44f&  FaceToSampleMatrix)
-	{
-		SetTextureParameter(BatchedParameters, InTexture, InTextureSampler, Texture);
-		SetShaderValue(BatchedParameters, InFaceToSampleParameter, FaceToSampleMatrix);
-	}
-
-	void SetParameters(FRHIBatchedShaderParameters& BatchedParameters, FRHISamplerState* SamplerStateRHI, FRHITexture* TextureRHI, const FMatrix44f&  FaceToSampleMatrix)
-	{
-		SetTextureParameter(BatchedParameters, InTexture, InTextureSampler, SamplerStateRHI, TextureRHI);
-		SetShaderValue(BatchedParameters, InFaceToSampleParameter, FaceToSampleMatrix);
-	}
-
-private:
-	LAYOUT_FIELD(FShaderResourceParameter, InTexture);
-	LAYOUT_FIELD(FShaderResourceParameter, InTextureSampler);
-	LAYOUT_FIELD(FShaderParameter, InFaceToSampleParameter);
-};
+};
\ No newline at end of file
