[shaders] remove explicit finalized/released flags from job struct, replace with extended/refactored logic in the thread-safe job status object. this fixes some TSAN issues with subsystems which are polling these flags, as well as cleaning up some redundant logic - finalized flag was always set immediately before the job "complete" state was set, so there's no need for both. separated the completion method into a separate thread-safe enum - this allows tracking distributed vs local vs cache hit independent of the completion state of the job (so we can just set "distributed" or "duplicate" flags once, and have them carry through the remaining lifetime of the job, without needing the combinatorial states.

#Synced-CL 52473344

diff --git a/Engine/Plugins/Compositing/OpenColorIO/Source/OpenColorIO/Private/OpenColorIOShaderCompilationManager.cpp b/Engine/Plugins/Compositing/OpenColorIO/Source/OpenColorIO/Private/OpenColorIOShaderCompilationManager.cpp
--- a/Engine/Plugins/Compositing/OpenColorIO/Source/OpenColorIO/Private/OpenColorIOShaderCompilationManager.cpp
+++ b/Engine/Plugins/Compositing/OpenColorIO/Source/OpenColorIO/Private/OpenColorIOShaderCompilationManager.cpp
@@ -98,7 +98,7 @@ void FOpenColorIOShaderCompilationManager::ProcessAsyncResults()
 	{
 		auto CurrentJob = JobQueue[JobIndex]->GetSingleShaderJob();
 
-		if (!CurrentJob->bReleased)
+		if (CurrentJob->GetStatus() != EShaderCompileJobStatus::Released)
 		{
 			continue;
 		}
diff --git a/Engine/Plugins/Experimental/NNERuntimeIREE/Source/NNERuntimeIREEShader/Private/NNERuntimeIREEShaderCompilationManager.cpp b/Engine/Plugins/Experimental/NNERuntimeIREE/Source/NNERuntimeIREEShader/Private/NNERuntimeIREEShaderCompilationManager.cpp
--- a/Engine/Plugins/Experimental/NNERuntimeIREE/Source/NNERuntimeIREEShader/Private/NNERuntimeIREEShaderCompilationManager.cpp
+++ b/Engine/Plugins/Experimental/NNERuntimeIREE/Source/NNERuntimeIREEShader/Private/NNERuntimeIREEShaderCompilationManager.cpp
@@ -68,7 +68,7 @@ void FNNERuntimeIREEShaderCompilationManager::ProcessAsyncResults()
 	{
 		auto CurrentJob = JobQueue[JobIndex]->GetSingleShaderJob();
 
-		if (!CurrentJob->bReleased)
+		if (CurrentJob->GetStatus() != EShaderCompileJobStatus::Released)
 		{
 			continue;
 		}
diff --git a/Engine/Plugins/FX/Niagara/Source/NiagaraEditor/Private/NiagaraCompiler.cpp b/Engine/Plugins/FX/Niagara/Source/NiagaraEditor/Private/NiagaraCompiler.cpp
--- a/Engine/Plugins/FX/Niagara/Source/NiagaraEditor/Private/NiagaraCompiler.cpp
+++ b/Engine/Plugins/FX/Niagara/Source/NiagaraEditor/Private/NiagaraCompiler.cpp
@@ -2205,13 +2205,13 @@ TOptional<FNiagaraCompileResults> FHlslNiagaraCompiler::GetCompileResult(int32 J
 
 	TArray<int32> ShaderMapIDs;
 	ShaderMapIDs.Add(JobID);
-	if (bWait && !CompilationJob->ShaderCompileJob->bReleased)
+	if (bWait && (CompilationJob->ShaderCompileJob->GetStatus() != EShaderCompileJobStatus::Released))
 	{
 		GShaderCompilingManager->FinishCompilation(NULL, ShaderMapIDs);
-		check(CompilationJob->ShaderCompileJob->bReleased);
+		check(CompilationJob->ShaderCompileJob->GetStatus() == EShaderCompileJobStatus::Released);
 	}
 
-	if (!CompilationJob->ShaderCompileJob->bReleased)
+	if (CompilationJob->ShaderCompileJob->GetStatus() != EShaderCompileJobStatus::Released)
 	{
 		return TOptional<FNiagaraCompileResults>();
 	}
@@ -2504,7 +2504,7 @@ bool FNiagaraShaderMapCompiler::ProcessCompileResults(bool bWait)
 	{
 		auto IsCompileJobIncomplete = [](const FShaderCommonCompileJobPtr& CompileJob) -> bool
 		{
-			return !CompileJob.IsValid() || !CompileJob->bReleased || !CompileJob->bFinalized;
+			return !CompileJob.IsValid() || (CompileJob->GetStatus() != EShaderCompileJobStatus::Released);
 		};
 
 		auto IsCompileJobError = [](const FShaderCommonCompileJobPtr& CompileJob) -> bool
@@ -2512,7 +2512,7 @@ bool FNiagaraShaderMapCompiler::ProcessCompileResults(bool bWait)
 			return !CompileJob->bSucceeded;
 		};
 
-		// make sure that all of the shader compile jobs have been released and finalized
+		// make sure that all of the shader compile jobs have been released 
 		const bool bReadyToProcess = !CompileIt->ShaderCompileJobs.ContainsByPredicate(IsCompileJobIncomplete);
 
 		if (!bReadyToProcess)
diff --git a/Engine/Plugins/FX/Niagara/Source/NiagaraShader/Private/NiagaraShaderCompilationManager.cpp b/Engine/Plugins/FX/Niagara/Source/NiagaraShader/Private/NiagaraShaderCompilationManager.cpp
--- a/Engine/Plugins/FX/Niagara/Source/NiagaraShader/Private/NiagaraShaderCompilationManager.cpp
+++ b/Engine/Plugins/FX/Niagara/Source/NiagaraShader/Private/NiagaraShaderCompilationManager.cpp
@@ -66,7 +66,7 @@ void FNiagaraShaderCompilationManager::ProcessAsyncResults()
 	{
 		auto CurrentJob = JobQueue[JobIndex]->GetSingleShaderJob();
 
-		if (!CurrentJob->bReleased)
+		if (CurrentJob->GetStatus() != EShaderCompileJobStatus::Released)
 		{
 			continue;
 		}
diff --git a/Engine/Plugins/Runtime/ComputeFramework/Source/ComputeFramework/Private/ComputeKernelShaderCompilationManager.cpp b/Engine/Plugins/Runtime/ComputeFramework/Source/ComputeFramework/Private/ComputeKernelShaderCompilationManager.cpp
--- a/Engine/Plugins/Runtime/ComputeFramework/Source/ComputeFramework/Private/ComputeKernelShaderCompilationManager.cpp
+++ b/Engine/Plugins/Runtime/ComputeFramework/Source/ComputeFramework/Private/ComputeKernelShaderCompilationManager.cpp
@@ -68,7 +68,7 @@ void FComputeKernelShaderCompilationManager::ProcessAsyncResults()
 	{
 		auto CurrentJob = JobQueue[JobIndex]->GetSingleShaderJob();
 
-		if (!CurrentJob->bReleased)
+		if (CurrentJob->GetStatus() != EShaderCompileJobStatus::Released)
 		{
 			continue;
 		}
diff --git a/Engine/Source/Runtime/Engine/Private/Materials/MaterialShader.cpp b/Engine/Source/Runtime/Engine/Private/Materials/MaterialShader.cpp
--- a/Engine/Source/Runtime/Engine/Private/Materials/MaterialShader.cpp
+++ b/Engine/Source/Runtime/Engine/Private/Materials/MaterialShader.cpp
@@ -3029,7 +3029,7 @@ void FMaterialShaderMap::AppendCompileStateDebugInfo(FStringBuilderBase& OutDebu
 		{
 			OutDebugInfo << " (input hash " << LexToString(Status->GetInputHash()) << ")";
 		}
-		OutDebugInfo << "; current status is " << LexToString(Status->GetStatus()) << (Status->GetIsDuplicate() ? " (duplicate)" : "") << "\n";
+		OutDebugInfo << "; current status is: " << LexToString(Status->GetStatus()) << ", completion method is: " << LexToString(Status->GetCompletionMethod()) << "\n";
 	}
 #endif
 }
diff --git a/Engine/Source/Runtime/Engine/Private/ShaderCompiler/ShaderCompiler.cpp b/Engine/Source/Runtime/Engine/Private/ShaderCompiler/ShaderCompiler.cpp
--- a/Engine/Source/Runtime/Engine/Private/ShaderCompiler/ShaderCompiler.cpp
+++ b/Engine/Source/Runtime/Engine/Private/ShaderCompiler/ShaderCompiler.cpp
@@ -1512,7 +1512,7 @@ void FShaderCompilingManager::ReleaseJob(FShaderCommonCompileJobPtr& Job)
 void FShaderCompilingManager::ReleaseJob(FShaderCommonCompileJob* Job)
 {
 	Job->PendingShaderMap.SafeRelease();
-	Job->bReleased = true;
+	Job->UpdateStatus(EShaderCompileJobStatus::Released);
 	AllJobs.RemoveJob(Job);
 }
 
@@ -2661,7 +2661,6 @@ bool FShaderCompilingManager::HandlePotentialRetry(TMap<int32, FShaderMapFinaliz
 					SingleJob.Output = FShaderCompilerOutput();
 					SingleJob.PreprocessOutput = FShaderPreprocessOutput();
 					SingleJob.JobStatusPtr->Reset();
-					SingleJob.bFinalized = false;
 				}
 			);
 			
@@ -2669,7 +2668,6 @@ bool FShaderCompilingManager::HandlePotentialRetry(TMap<int32, FShaderMapFinaliz
 			CurrentJob->RequestOwner.Reset();
 			CurrentJob->JobStatusPtr->Reset();
 			CurrentJob->bInputHashSet = false;
-			CurrentJob->bFinalized = false;
 			// Need to force reissued jobs to skip the cache queries otherwise jobs with warnings will just be cache hits and not actually recompile
 			// (since debug info being enabled for a job intentionally does not affect the cached key/cached results)
 			CurrentJob->bBypassCache = true;
diff --git a/Engine/Source/Runtime/Engine/Private/ShaderCompiler/ShaderCompilerDistributed.cpp b/Engine/Source/Runtime/Engine/Private/ShaderCompiler/ShaderCompilerDistributed.cpp
--- a/Engine/Source/Runtime/Engine/Private/ShaderCompiler/ShaderCompilerDistributed.cpp
+++ b/Engine/Source/Runtime/Engine/Private/ShaderCompiler/ShaderCompilerDistributed.cpp
@@ -600,7 +600,7 @@ int32 FShaderCompileDistributedThreadRunnable_Interface::CompilingLoop()
 						[](FShaderCompileJob& SingleJob)
 						{
 							SingleJob.Output = FShaderCompilerOutput();
-							SingleJob.bFinalized = false;
+							SingleJob.UpdateStatus(EShaderCompileJobStatus::Pending, EShaderCompileJobCompletionMethod::LocalExecution);
 						});
 
 					FShaderCompileUtilities::ExecuteShaderCompileJob(*Job);
@@ -612,7 +612,7 @@ int32 FShaderCompileDistributedThreadRunnable_Interface::CompilingLoop()
 				FScopeLock Lock(&Manager->CompileQueueSection);
 				for (const auto& Job : Task->ShaderJobs)
 				{
-					Manager->ProcessFinishedJob(Job, EShaderCompileJobStatus::CompleteDistributedExecution);
+					Manager->ProcessFinishedJob(Job, EShaderCompileJobStatus::Completed);
 				}
 			}
 		}
diff --git a/Engine/Source/Runtime/Engine/Private/ShaderCompiler/ShaderCompilerJobCache.cpp b/Engine/Source/Runtime/Engine/Private/ShaderCompiler/ShaderCompilerJobCache.cpp
--- a/Engine/Source/Runtime/Engine/Private/ShaderCompiler/ShaderCompilerJobCache.cpp
+++ b/Engine/Source/Runtime/Engine/Private/ShaderCompiler/ShaderCompilerJobCache.cpp
@@ -633,7 +633,7 @@ public:
 	void OnSubmitJobComplete();
 
 	/** This is an entry point for all jobs that have finished the compilation (whether real or cached). Can be called from multiple threads. */
-	void ProcessFinishedJob(FShaderCommonCompileJob* FinishedJob, EShaderCompileJobStatus Status);
+	void ProcessFinishedJob(FShaderCommonCompileJob* FinishedJob, EShaderCompileJobStatus Status, EShaderCompileJobCompletionMethod Method = EShaderCompileJobCompletionMethod::Unset);
 
 	/** Adds the job to cache. */
 	void AddToCacheAndProcessPending(FShaderCommonCompileJob* FinishedJob);
@@ -1245,7 +1245,8 @@ int32 FShaderJobCache::RemoveAllPendingJobsWithId(uint32 InId)
 							Unlink(*DuplicateJob);
 							RemoveDuplicateJob(DuplicateJob);
 
-							DuplicateJob->JobStatusPtr->SetIsDuplicate(false);
+							// Unset the execution method on the previously "duplicate" job; this will now be set by whatever executor picks up the job.
+							DuplicateJob->JobStatusPtr->SetCompletionMethod(EShaderCompileJobCompletionMethod::Unset);
 
 							// Add it as pending at the appropriate priority
 							GShaderCompilerStats->RegisterNewPendingJob(*DuplicateJob);
@@ -1400,14 +1401,17 @@ void FShaderJobCache::SubmitJob(FShaderCommonCompileJob* Job)
 			Job->SerializeOutput(LoadContext);
 
 			// finish the job instantly
-			ProcessFinishedJob(Job, EShaderCompileJobStatus::CompleteFoundInCache);
+			ProcessFinishedJob(Job, EShaderCompileJobStatus::Completed, EShaderCompileJobCompletionMethod::CacheHit);
 		}
 		else
 		{
 			FShaderJobData& JobData = GetShaderJobData(Job->JobCacheRef);
-			Job->UpdateStatus(EShaderCompileJobStatus::Queued);
-			// see if another job with the same input hash is being worked on
-			if (JobData.JobInFlight)
+			// if there's an in flight job for this cache ref already another job with the same input hash is already being worked on
+			bool bDuplicateJob = JobData.JobInFlight.IsValid(); 			
+			// update the completion method for a duplicate job immediately (so the status reflects that it's a duplicate while it's in queued state)
+			// keep it unset if it's a primary job that will be executed; the execution method will be set when it's picked up by a compilation thread
+			Job->UpdateStatus(EShaderCompileJobStatus::Queued, bDuplicateJob ? EShaderCompileJobCompletionMethod::Duplicate : EShaderCompileJobCompletionMethod::Unset);
+			if (bDuplicateJob)
 			{
 				UE_LOG(LogShaderCompilers, UE_SHADERCACHE_LOG_LEVEL, TEXT("There is an outstanding job with the ihash %s, not submitting another one (adding to wait list)."), *LexToString(InputHash));
 
@@ -1424,8 +1428,6 @@ void FShaderJobCache::SubmitJob(FShaderCommonCompileJob* Job)
 					*WaitListHead = Job;
 				}
 				++TotalCacheDuplicates;
-
-				Job->JobStatusPtr->SetIsDuplicate();
 				AddDuplicateJob(Job);
 				JobLock.WriteUnlock();
 				bNewJob = false;
@@ -1549,9 +1551,9 @@ void FShaderJobCache::SubmitJobs(const TArray<FShaderCommonCompileJobPtr>& InJob
 	}
 }
 
-void FShaderJobCache::ProcessFinishedJob(FShaderCommonCompileJob* FinishedJob, EShaderCompileJobStatus Status)
+void FShaderJobCache::ProcessFinishedJob(FShaderCommonCompileJob* FinishedJob, EShaderCompileJobStatus Status, EShaderCompileJobCompletionMethod CompletionMethod)
 {
-	FinishedJob->UpdateStatus(Status);
+	FinishedJob->UpdateStatus(Status, CompletionMethod);
 
 	FShaderDebugDataContext Ctx;
 	FinishedJob->OnComplete(Ctx);
@@ -1728,7 +1730,7 @@ int32 FShaderJobCache::GetPendingJobs(EShaderCompilerWorkerType InWorkerType, ES
 		UnlinkJobWithPriority(Job);
 
 		Job.CurrentWorker = InWorkerType;
-		Job.UpdateStatus(InWorkerType == EShaderCompilerWorkerType::Distributed ? EShaderCompileJobStatus::PendingDistributedExecution : EShaderCompileJobStatus::PendingLocalExecution);
+		Job.UpdateStatus(EShaderCompileJobStatus::Pending, (InWorkerType == EShaderCompilerWorkerType::Distributed) ? EShaderCompileJobCompletionMethod::DistributedExecution : EShaderCompileJobCompletionMethod::LocalExecution);
 		OutJobs.Add(&Job);
 
 		// get a random number of jobs to skip (if we can). We're skipping after taking the first job so we can ensure that we always take the latest job into the batch
@@ -1948,7 +1950,14 @@ FShaderCacheLoadContext FShaderJobCache::FindOrAdd(const FShaderCompilerInputHas
 							{
 								LoadContext.Reuse();
 								FinishedJob->SerializeOutput(LoadContext);
-								ProcessFinishedJob(FinishedJob, EShaderCompileJobStatus::CompleteFoundInDDC);
+								EShaderCompileJobCompletionMethod Method = FinishedJob->JobStatusPtr->GetCompletionMethod();
+								// We only update execution method to "DdcHit" for the first job in the chain (of which all subsequent entries are dupes) for tracking purposes
+								// Duplicate jobs will have previously had their execution method set to "Duplicate" when queued, so this maintains that
+								if (Method == EShaderCompileJobCompletionMethod::Unset)
+								{
+									Method = EShaderCompileJobCompletionMethod::DdcHit;
+								}
+								ProcessFinishedJob(FinishedJob, EShaderCompileJobStatus::Completed, Method);
 							}
 
 							if (FinishedJobs.Num() > 1)
diff --git a/Engine/Source/Runtime/Engine/Private/ShaderCompiler/ShaderCompilerThreadRunnable.cpp b/Engine/Source/Runtime/Engine/Private/ShaderCompiler/ShaderCompilerThreadRunnable.cpp
--- a/Engine/Source/Runtime/Engine/Private/ShaderCompiler/ShaderCompilerThreadRunnable.cpp
+++ b/Engine/Source/Runtime/Engine/Private/ShaderCompiler/ShaderCompilerThreadRunnable.cpp
@@ -618,7 +618,7 @@ void FShaderCompileThreadRunnable::PushCompletedJobsToManager()
 			{
 				auto& Job = CurrentWorkerInfo.QueuedJobs[JobIndex];
 
-				Manager->ProcessFinishedJob(Job.GetReference(), EShaderCompileJobStatus::CompleteLocalExecution);
+				Manager->ProcessFinishedJob(Job.GetReference(), EShaderCompileJobStatus::Completed);
 			}
 
 			const double ElapsedTime = FPlatformTime::Seconds() - CurrentWorkerInfo.StartTime;
diff --git a/Engine/Source/Runtime/Engine/Private/ShaderCompiler/ShaderCompilerUtil.cpp b/Engine/Source/Runtime/Engine/Private/ShaderCompiler/ShaderCompilerUtil.cpp
--- a/Engine/Source/Runtime/Engine/Private/ShaderCompiler/ShaderCompilerUtil.cpp
+++ b/Engine/Source/Runtime/Engine/Private/ShaderCompiler/ShaderCompilerUtil.cpp
@@ -21,7 +21,7 @@ void FShaderCompileUtilities::ExecuteShaderCompileJob(FShaderCommonCompileJob& J
 {
 	TRACE_CPUPROFILER_EVENT_SCOPE(FShaderCompileUtilities::ExecuteShaderCompileJob);
 
-	check(!Job.bFinalized);
+	check(Job.GetStatus() != EShaderCompileJobStatus::Completed);
 
 	static ITargetPlatformManagerModule& TPM = GetTargetPlatformManagerRef();
 	auto* SingleJob = Job.GetSingleShaderJob();
@@ -56,7 +56,7 @@ void FShaderCompileUtilities::ExecuteShaderCompileJob(FShaderCommonCompileJob& J
 		CompileShaderPipeline(ShaderFormats, PipelineJob);
 	}
 
-	Job.bFinalized = true;
+	Job.UpdateStatus(EShaderCompileJobStatus::Completed, EShaderCompileJobCompletionMethod::LocalExecution);
 }
 
 FArchive* FShaderCompileUtilities::CreateFileHelper(const FString& Filename)
diff --git a/Engine/Source/Runtime/RenderCore/Private/ShaderCompileWorkerUtil.cpp b/Engine/Source/Runtime/RenderCore/Private/ShaderCompileWorkerUtil.cpp
--- a/Engine/Source/Runtime/RenderCore/Private/ShaderCompileWorkerUtil.cpp
+++ b/Engine/Source/Runtime/RenderCore/Private/ShaderCompileWorkerUtil.cpp
@@ -719,10 +719,7 @@ void FShaderCompileWorkerUtil::DumpDebugCompileInput(FShaderCommonCompileJob& Jo
 
 static void ReadSingleJob(FShaderCompileJob* CurrentJob, FArchive& WorkerOutputFileReader)
 {
-	TRACE_CPUPROFILER_EVENT_SCOPE(ReadSingleJob);
-
-	check(!CurrentJob->bFinalized);
-	CurrentJob->bFinalized = true;
+	TRACE_CPUPROFILER_EVENT_SCOPE(ReadSingleJob)
 
 	// Deserialize the shader compilation output.
 	CurrentJob->SerializeWorkerOutput(WorkerOutputFileReader);
@@ -878,9 +875,6 @@ FSCWErrorCode::ECode FShaderCompileWorkerUtil::ReadTasks(const TArray<FShaderCom
 					OutputFileContext.ModalErrorOrLog(TEXT("Expecting pipeline job \"%s\", got \"%s\" instead!"), CurrentJob->Key.ShaderPipeline->GetName(), *PipelineName);
 				}
 
-				check(!CurrentJob->bFinalized);
-				CurrentJob->bFinalized = true;
-
 				int32 NumStageJobs = -1;
 				OutputFile << NumStageJobs;
 
diff --git a/Engine/Source/Runtime/RenderCore/Private/ShaderCore.cpp b/Engine/Source/Runtime/RenderCore/Private/ShaderCore.cpp
--- a/Engine/Source/Runtime/RenderCore/Private/ShaderCore.cpp
+++ b/Engine/Source/Runtime/RenderCore/Private/ShaderCore.cpp
@@ -4517,12 +4517,23 @@ const TCHAR* LexToString(EShaderCompileJobStatus Status)
 	case EShaderCompileJobStatus::Cancelled: return TEXT("Cancelled");
 	case EShaderCompileJobStatus::PendingDDC: return TEXT("PendingDDC");
 	case EShaderCompileJobStatus::Queued: return TEXT("Queued");
-	case EShaderCompileJobStatus::PendingDistributedExecution: return TEXT("PendingDistributedExecution");
-	case EShaderCompileJobStatus::PendingLocalExecution: return TEXT("PendingLocalExecution");
-	case EShaderCompileJobStatus::CompleteDistributedExecution: return TEXT("CompleteDistributedExecution");
-	case EShaderCompileJobStatus::CompleteFoundInCache: return TEXT("CompleteFoundInCache");
-	case EShaderCompileJobStatus::CompleteFoundInDDC: return TEXT("CompleteFoundInDDC");
-	case EShaderCompileJobStatus::CompleteLocalExecution: return TEXT("CompleteLocalExecution");
+	case EShaderCompileJobStatus::Pending: return TEXT("Pending");
+	case EShaderCompileJobStatus::Completed: return TEXT("Completed");
+	case EShaderCompileJobStatus::Released: return TEXT("Released");
+	default: return TEXT("(unknown)");
+	}
+}
+
+const TCHAR* LexToString(EShaderCompileJobCompletionMethod Method)
+{
+	switch (Method)
+	{
+	case EShaderCompileJobCompletionMethod::Unset: return TEXT("Unset");
+	case EShaderCompileJobCompletionMethod::CacheHit: return TEXT("CacheHit");
+	case EShaderCompileJobCompletionMethod::DdcHit: return TEXT("DdcHit");
+	case EShaderCompileJobCompletionMethod::Duplicate: return TEXT("Duplicate");
+	case EShaderCompileJobCompletionMethod::LocalExecution: return TEXT("LocalExecution");
+	case EShaderCompileJobCompletionMethod::DistributedExecution: return TEXT("DistributedExecution");
 	default: return TEXT("(unknown)");
 	}
 }
@@ -4676,7 +4687,6 @@ void FShaderCompileJob::SerializeOutput(FShaderCacheSerializeContext& Ctx, int32
 
 	if (bIsLoading)
 	{
-		bFinalized = true;
 		bSucceeded = Output.bSucceeded;
 	}
 }
@@ -5002,7 +5012,6 @@ void FShaderPipelineCompileJob::SerializeOutput(FShaderCacheSerializeContext& Ct
 
 	if (Ctx.GetMainArchive().IsLoading())
 	{
-		bFinalized = true;
 		bSucceeded = bAllStagesSucceeded;
 	}
 }
diff --git a/Engine/Source/Runtime/RenderCore/Public/ShaderCompilerJobTypes.h b/Engine/Source/Runtime/RenderCore/Public/ShaderCompilerJobTypes.h
--- a/Engine/Source/Runtime/RenderCore/Public/ShaderCompilerJobTypes.h
+++ b/Engine/Source/Runtime/RenderCore/Public/ShaderCompilerJobTypes.h
@@ -76,20 +76,46 @@ struct FShaderJobCacheRef
 
 enum class EShaderCompileJobStatus : uint8
 {
+	/* Initial status, created but not yet manipulated */
 	Unset,
+	/* Ready for execution, set after preprocessing and input hash computation is completed */
 	Ready,
+	/* Job execution was skipped due to a preprocessing failure */
 	Skipped,
+	/* Job was cancelled by owning system */
 	Cancelled,
+	/* DDC query for job has been submitted and is waiting on results */
 	PendingDDC,
+	/* Job is queued in the shader compilation system, but not yet picked up by a compilation thread (i.e. not a cache or DDC hit) */
 	Queued,
-	PendingDistributedExecution,
-	PendingLocalExecution,
-	CompleteDistributedExecution,
-	CompleteFoundInCache,
-	CompleteFoundInDDC,
-	CompleteLocalExecution,
+	/* Job has been picked up by a compilation thread, has been submitted for compilation is pending execution completion */
+	Pending,
+	/* Job has been completed and results are available for processing */
+	Completed,
+	/* Job results have been read and processed into the owning shadermap, and job has been released by the shader compilation system.
+	*  Note that the job may remain alive if other systems are holding references to it at this point.
+	*/
+	Released
 };
+
+enum class EShaderCompileJobCompletionMethod
+{
+	/* Initial value, execution method not yet determined */
+	Unset,
+	/* Job was found in the in-memory job cache */
+	CacheHit,
+	/* Job was found in DDC */
+	DdcHit,
+	/* Job was a duplicate of another and shared its results */
+	Duplicate,
+	/* Job was executed by a local shader compilation mechanism (local shader compile thread or local fallback) */
+	LocalExecution,
+	/* Job was executed by a distributed shader compilation mechanism (note that this could still have executed on the local machine, but coordinated by a distributed compiler) */
+	DistributedExecution,
+};
+
 RENDERCORE_API const TCHAR* LexToString(EShaderCompileJobStatus Status);
+RENDERCORE_API const TCHAR* LexToString(EShaderCompileJobCompletionMethod Method);
 
 class FShaderCompileJobStatus
 {
@@ -98,38 +124,59 @@ public:
 	{
 		// Unset status before the input hash since GetInputHash() uses it as a substitute for lock guards
 		Status.store(EShaderCompileJobStatus::Unset);
+		CompletionMethod.store(EShaderCompileJobCompletionMethod::Unset);
+
 		InputHash = nullptr;
-		IsDuplicate.store(false);
 	}
 
+	/* Returns true if the job has been completed (i.e. is in "Completed" or "Released" phase) */
+	bool IsCompleted() const
+	{
+		EShaderCompileJobStatus LocalStatus = GetStatus();
+		return
+			(LocalStatus == EShaderCompileJobStatus::Completed) ||
+			(LocalStatus == EShaderCompileJobStatus::Released);
+	}
+
+	/* Returns true if the owning job is completed but the compile step was not executed. */
 	bool WasCompilationSkipped() const
 	{
 		EShaderCompileJobStatus LocalStatus = GetStatus();
+		EShaderCompileJobCompletionMethod LocalCompletionMethod = GetCompletionMethod();
+		bool bIsCompleted = IsCompleted();
+
+		return
+			(LocalStatus == EShaderCompileJobStatus::Skipped) ||
+			(bIsCompleted &&
+			(LocalCompletionMethod == EShaderCompileJobCompletionMethod::Duplicate ||
+			LocalCompletionMethod == EShaderCompileJobCompletionMethod::CacheHit ||
+			LocalCompletionMethod == EShaderCompileJobCompletionMethod::DdcHit));
+	}
 
-		return GetIsDuplicate() ||
-		(LocalStatus == EShaderCompileJobStatus::Skipped) ||
-		(LocalStatus == EShaderCompileJobStatus::CompleteFoundInCache) ||
-		(LocalStatus == EShaderCompileJobStatus::CompleteFoundInDDC);
+	void SetStatus(EShaderCompileJobStatus InStatus)
+	{
+		Status.store(InStatus);
 	}
 
-	void SetStatus(EShaderCompileJobStatus NewStatus)
+	void SetStatus(EShaderCompileJobStatus InStatus, EShaderCompileJobCompletionMethod InCompletionMethod)
 	{
-		Status.store(NewStatus);
+		SetCompletionMethod(InCompletionMethod);
+		SetStatus(InStatus);
 	}
 
-	EShaderCompileJobStatus GetStatus() const
+	void SetCompletionMethod(EShaderCompileJobCompletionMethod InCompletionMethod)
 	{
-		return Status.load();
+		CompletionMethod.store(InCompletionMethod);
 	}
 
-	void SetIsDuplicate(bool bDuplicate = true)
+	EShaderCompileJobCompletionMethod GetCompletionMethod() const
 	{
-		IsDuplicate.store(bDuplicate);
+		return CompletionMethod.load();
 	}
 
-	bool GetIsDuplicate() const
+	EShaderCompileJobStatus GetStatus() const
 	{
-		return IsDuplicate.load();
+		return Status.load();
 	}
 
 	void SetInputHash(const FShaderCompilerInputHash* InInputHash)
@@ -145,7 +192,7 @@ public:
 private:
 	const FShaderCompilerInputHash* InputHash = nullptr; // set when input hash is computed during submission
 	std::atomic<EShaderCompileJobStatus> Status{ EShaderCompileJobStatus::Unset };
-	std::atomic<bool> IsDuplicate{ false };
+	std::atomic<EShaderCompileJobCompletionMethod> CompletionMethod{ EShaderCompileJobCompletionMethod::Unset };
 };
 
 /** Stores all of the common information used to compile a shader or pipeline. */
@@ -180,15 +227,10 @@ public:
 
 	TPimplPtr<UE::DerivedData::FRequestOwner> RequestOwner;
 
-	/** true if the results of the shader compile have been processed. */
-	uint8 bFinalized : 1;
-	/** Output of the shader compile */
+	/** True if the compile job executed successfully, false if errors were encountered */
 	uint8 bSucceeded : 1;
+	/** True if errors encountered by this job are likely to be caused by code problems (inferred from shader type and whether or not shader dev mode is enabled) */
 	uint8 bErrorsAreLikelyToBeCode : 1;
-	/** true if the results of the shader compile have been released from the FShaderCompilerManager.
-		After a job is bFinalized it will be bReleased when ReleaseJob() is invoked, which means that the shader compile thread
-		is no longer processing the job; which is useful for non standard job handling (Niagara as an example). */
-	uint8 bReleased : 1;
 	/** Whether or not we are a default material. */
 	uint8 bIsDefaultMaterial : 1;
 	/** Whether or not we are a global shader. */
@@ -227,6 +269,25 @@ public:
 	{
 		JobStatusPtr->SetStatus(NewStatus);
 	}
+
+	void UpdateStatus(EShaderCompileJobStatus InStatus, EShaderCompileJobCompletionMethod InCompletionMethod)
+	{
+		// InCompletionMethod defaults to Unset if not explicitly passed; intended behaviour is to keep the 
+		// existing value for execution method rather than overwriting with "unset".
+		if (InCompletionMethod != EShaderCompileJobCompletionMethod::Unset)
+		{
+			JobStatusPtr->SetStatus(InStatus, InCompletionMethod);
+		}
+		else
+		{
+			JobStatusPtr->SetStatus(InStatus);
+		}
+	}
+
+	EShaderCompileJobStatus GetStatus() const
+	{
+		return JobStatusPtr->GetStatus();
+	}
 	
 	uint32 AddRef() const
 	{
@@ -291,10 +352,8 @@ protected:
 		Priority(InPriority),
 		PendingPriority(EShaderCompileJobPriority::None),
 		CurrentWorker(EShaderCompilerWorkerType::None),
-		bFinalized(false),
 		bSucceeded(false),
 		bErrorsAreLikelyToBeCode(false),
-		bReleased(false),
 		bIsDefaultMaterial(false),
 		bIsGlobalShader(false),
 		bBypassCache(false),
