Deprecating oculus shaders

#Synced-CL 52534449

diff --git a/Engine/Source/Runtime/RenderCore/Public/OculusShaders.h b/Engine/Source/Runtime/RenderCore/Public/OculusShaders.h
--- a/Engine/Source/Runtime/RenderCore/Public/OculusShaders.h
+++ b/Engine/Source/Runtime/RenderCore/Public/OculusShaders.h
@@ -17,7 +17,7 @@ class FRHISamplerState;
 class FRHITexture;
 class FTexture;
 
-class FOculusVertexShader : public FGlobalShader
+class UE_DEPRECATED(5.8, "Oculus shaders no longer supported") FOculusVertexShader : public FGlobalShader
 {
 	DECLARE_EXPORTED_SHADER_TYPE(FOculusVertexShader, Global, RENDERCORE_API);
 public:
@@ -45,7 +45,7 @@ public:
 	FOculusWhiteShader() {}
 };
 
-class FOculusBlackShader : public FGlobalShader
+class UE_DEPRECATED(5.8, "Oculus shaders no longer supported") FOculusBlackShader : public FGlobalShader
 {
 	DECLARE_EXPORTED_SHADER_TYPE(FOculusBlackShader, Global, RENDERCORE_API);
 public:
@@ -60,7 +60,7 @@ public:
 	FOculusBlackShader() {}
 };
 
-class FOculusAlphaInverseShader : public FGlobalShader
+class UE_DEPRECATED(5.8, "Oculus shaders no longer supported") FOculusAlphaInverseShader : public FGlobalShader
 {
 	DECLARE_EXPORTED_SHADER_TYPE(FOculusAlphaInverseShader, Global, RENDERCORE_API);
 public:
