Fix excessive spam in shader library.

- This message is not bringing value to justify the continuous log spam.

#Synced-CL 52512000

diff --git a/Engine/Source/Runtime/RenderCore/Private/ShaderLibrary/ShaderCodeLibrary.cpp b/Engine/Source/Runtime/RenderCore/Private/ShaderLibrary/ShaderCodeLibrary.cpp
--- a/Engine/Source/Runtime/RenderCore/Private/ShaderLibrary/ShaderCodeLibrary.cpp
+++ b/Engine/Source/Runtime/RenderCore/Private/ShaderLibrary/ShaderCodeLibrary.cpp
@@ -2936,7 +2936,7 @@ void FShaderCodeLibrary::OpenPluginShaderLibrary(IPlugin& Plugin, bool bMonolith
 			// Chunked libraries in plugins that are not built in (i.e. ones that were cooked separately as DLC) are not supported atm. This is because the main game can be cooked without chunks (-fastcook), but still needs
 			// to load the same plugins, so it would not know which ChunkIDs to try.
 			// Plugins that are built-in can be chunked, but their shaders don't go into a separate library (cooker doesn't separate that atm), everything goes into main project's library.
-			UE_LOGF(LogShaderLibrary, Display, "Opening a chunked shader library for plugin '%ls' is ignored. Chunked libraries for plugins are not supported.", *Plugin.GetName());
+			UE_LOGF(LogShaderLibrary, Verbose, "Opening a chunked shader library for plugin '%ls' is ignored. Chunked libraries for plugins are not supported.", *Plugin.GetName());
 		}
 		bMonolithicOnly = true;
 		FShaderCodeLibrary::OpenLibrary(Plugin.GetName(), Plugin.GetContentDir(), bMonolithicOnly);
