[DumGPU] Add CVar r.DumpGPU.RedumpInputs (0, default off) to enable redump of inputs at each pass.

This fix the lost of input resource intermediate update when it was already dumped but the intermediate passes write to it. Enable it will increase the memory footage.

#Synced-CL 52438879

diff --git a/Engine/Extras/GPUDumpViewer/GPUDumpViewer.html b/Engine/Extras/GPUDumpViewer/GPUDumpViewer.html
--- a/Engine/Extras/GPUDumpViewer/GPUDumpViewer.html
+++ b/Engine/Extras/GPUDumpViewer/GPUDumpViewer.html
@@ -1343,6 +1343,16 @@ function display_input_resource(pass_id, subresource_unique_name)
 		subresource_version_info['pass'] = k_null_json_ptr;
 	}
 
+	// Check if a re-dump version exists (created by r.DumpGPU.RedumpInputs).
+	// This version captures the resource at the point the consuming pass reads it,
+	// which includes modifications from passes excluded by r.DumpGPU.Root.
+	var consuming_pass_ptr = g_passes[pass_id]['Pointer'];
+	var redump_version_name = `${subresource_unique_name}.v${consuming_pass_ptr}`;
+	if (does_file_exists(`Resources/${redump_version_name}.bin`)) 
+	{
+		subresource_version_info['pass'] = consuming_pass_ptr;
+	}
+
 	display_pass_internal(pass_id);
 	display_resource_internal(subresource_version_info);
 }
diff --git a/Engine/Source/Runtime/RenderCore/Private/DumpGPU.cpp b/Engine/Source/Runtime/RenderCore/Private/DumpGPU.cpp
--- a/Engine/Source/Runtime/RenderCore/Private/DumpGPU.cpp
+++ b/Engine/Source/Runtime/RenderCore/Private/DumpGPU.cpp
@@ -90,6 +90,12 @@ static TAutoConsoleVariable<int32> GDumpGPUPassParameters(
 	TEXT("Whether to dump the pass parameters."),
 	ECVF_RenderThreadSafe);
 
+static TAutoConsoleVariable<int32> GDumpGPURedumpInputs(
+	TEXT("r.DumpGPU.RedumpInputs"), 0,
+	TEXT("Re-capture input resources at the point they are read, breaking linkage to any prior (stale) version.\n")
+	TEXT("Useful when r.DumpGPU.Root excludes intermediate passes that modify a resource in-place. This would increase dumped memory footage.\n"),
+	ECVF_RenderThreadSafe);
+
 static TAutoConsoleVariable<float> GDumpGPUDelay(
 	TEXT("r.DumpGPU.Delay"), 0.0f,
 	TEXT("Delay in seconds before dumping the frame."),
@@ -1382,7 +1388,8 @@ public:
 
 		// Dump the resource's binary to a .bin file.
 		{
-			FString DumpFilePath = kResourcesDir / FString::Printf(TEXT("%s.v%s.bin"), *UniqueResourceSubResourceName, *PtrToString(bIsOutputResource ? Pass : nullptr));
+			const FRDGPass* VersionPass = bIsOutputResource ? Pass : LastResourceVersion.FindRef(SubresourceDesc.Texture);
+			FString DumpFilePath = kResourcesDir / FString::Printf(TEXT("%s.v%s.bin"), *UniqueResourceSubResourceName, *PtrToString(VersionPass));
 
 			FDumpTexturePass* PassParameters = GraphBuilder.AllocParameters<FDumpTexturePass>();
 			if (SubresourceDumpDesc.bPreprocessForStaging)
@@ -1578,8 +1585,17 @@ public:
 		}
 		else
 		{
-			// The texture has already been dumped, so can early return given it is not being modified by this Pass.
-			bAllowDumpBinary = false;
+			if (GDumpGPURedumpInputs.GetValueOnRenderThread() != 0)
+			{
+				// Re-dump under a unique version key (the consuming pass pointer). So if this texture has been dumped 
+				// previously but modified by intermeddiate passes between then and now, We can still get the correct copy.
+				UpdateResourceLastModifyingPass(SubresourceRangeDesc.Texture, Pass);
+			}
+			else
+			{
+				// The texture has already been dumped, so can early return given it is not being modified by this Pass.
+				bAllowDumpBinary = false;
+			}
 		}
 
 		return AddDumpTextureAllSubResourcesPasses(
