Fix DumpGPU mid-file encoding corruption in ResourceDescs.json

DumpStringToFile used EEncodingOptions::AutoDetect which switches to UTF-16 LE when a resource name contains non-ASCII characters. Since each resource descriptor is appended individually (FILEWRITE_Append), this causes a mid-file encoding switch that corrupts the JSON and breaks the GPU dump viewer.

Change to ForceUTF8WithoutBOM for consistent encoding throughout.

#Synced-CL 52697243

diff --git a/Engine/Source/Runtime/RenderCore/Private/DumpGPU.cpp b/Engine/Source/Runtime/RenderCore/Private/DumpGPU.cpp
--- a/Engine/Source/Runtime/RenderCore/Private/DumpGPU.cpp
+++ b/Engine/Source/Runtime/RenderCore/Private/DumpGPU.cpp
@@ -431,7 +431,7 @@ public:
 		FFileWriteCtx WriteCtx(this, ETimingBucket::MetadataFileWrite, OutputString.Len());
 		return FFileHelper::SaveStringToFile(
 			OutputString, *FullPath,
-			FFileHelper::EEncodingOptions::AutoDetect, &IFileManager::Get(), WriteFlags);
+			FFileHelper::EEncodingOptions::ForceUTF8WithoutBOM, &IFileManager::Get(), WriteFlags);
 	}
 
 	bool DumpStatusToFile(FStringView StatusString)
